Fix crash due to null entries in the Component array

diff --git a/Engine/Source/Editor/Kismet/Private/SSCSEditor.cpp b/Engine/Source/Editor/Kismet/Private/SSCSEditor.cpp
--- a/Engine/Source/Editor/Kismet/Private/SSCSEditor.cpp
+++ b/Engine/Source/Editor/Kismet/Private/SSCSEditor.cpp
@@ -4654,6 +4654,47 @@ void SSCSEditor::UpdateTree(bool bRegenerateTreeNodes)
 				// Get the full set of instanced components
 				TSet<UActorComponent*> ComponentsToAdd(ActorInstance->GetComponents());
 
+				const bool bHideConstructionScriptComponentsInDetailsView = GetDefault<UBlueprintEditorSettings>()->bHideConstructionScriptComponentsInDetailsView;
+				auto ShouldAddInstancedActorComponent = [bHideConstructionScriptComponentsInDetailsView](UActorComponent* ActorComp, USceneComponent* ParentSceneComp)
+				{
+					// Exclude nested DSOs attached to BP-constructed instances, which are not mutable.
+					return (ActorComp != nullptr
+						&& (!ActorComp->IsVisualizationComponent())
+						&& (ActorComp->CreationMethod != EComponentCreationMethod::UserConstructionScript || !bHideConstructionScriptComponentsInDetailsView)
+						&& (ParentSceneComp == nullptr || !ParentSceneComp->IsCreatedByConstructionScript() || !ActorComp->HasAnyFlags(RF_DefaultSubObject)))
+						&& (ActorComp->CreationMethod != EComponentCreationMethod::Native || FComponentEditorUtils::CanEditNativeComponent(ActorComp));
+				};
+
+				for (auto It(ComponentsToAdd.CreateIterator()); It; ++It)
+				{
+					UActorComponent* ActorComp = *It;
+					USceneComponent* SceneComp = Cast<USceneComponent>(ActorComp);
+					USceneComponent* ParentSceneComp = SceneComp != nullptr ? SceneComp->GetAttachParent() : nullptr;
+					if (!ShouldAddInstancedActorComponent(ActorComp, ParentSceneComp))
+					{
+						It.RemoveCurrent();
+					}
+				}
+
+				TFunction<void(USceneComponent*,FSCSEditorTreeNodePtrType)> AddInstancedTreeNodesRecursive = [&](USceneComponent* Component, FSCSEditorTreeNodePtrType TreeNode)
+				{
+					if (Component != nullptr)
+					{
+						TArray<USceneComponent*> Components = Component->GetAttachChildren();
+						for (USceneComponent* ChildComponent : Components)
+						{
+							if (ComponentsToAdd.Contains(ChildComponent)
+								&& ChildComponent->GetOwner() == Component->GetOwner())
+							{
+								ComponentsToAdd.Remove(ChildComponent);
+
+								FSCSEditorTreeNodePtrType NewParentNode = AddTreeNodeFromComponent(ChildComponent, TreeNode);
+								AddInstancedTreeNodesRecursive(ChildComponent, NewParentNode);
+							}
+						}
+					}
+				};
+
 				// Add the root component first (it may not be the first one)
 				USceneComponent* RootComponent = ActorInstance->GetRootComponent();
 				if(RootComponent != nullptr)
@@ -4663,7 +4704,7 @@ void SSCSEditor::UpdateTree(bool bRegenerateTreeNodes)
 					// Recursively add any instanced children that are already attached through the root, and keep track of added
 					// instances. This will be a faster path than the loop below, because we create new parent tree nodes as we go.
 					FSCSEditorTreeNodePtrType NewParentNode = AddTreeNodeFromComponent(RootComponent, FindOrCreateParentForExistingComponent(RootComponent, GetActorNode()));
-					AddInstancedTreeNodesRecursive(RootComponent, NewParentNode, ComponentsToAdd);
+					AddInstancedTreeNodesRecursive(RootComponent, NewParentNode);
 				}
 
 				// Sort components by type (always put scene components first in the tree)
@@ -4676,12 +4717,7 @@ void SSCSEditor::UpdateTree(bool bRegenerateTreeNodes)
 				// unattached scene components followed by any instanced non-scene components owned by the Actor instance.
 				for (UActorComponent* ActorComp : ComponentsToAdd)
 				{
-					USceneComponent* SceneComp = Cast<USceneComponent>(ActorComp);
-					USceneComponent* ParentSceneComp = SceneComp != nullptr ? SceneComp->GetAttachParent() : nullptr;
-					if (ShouldAddInstancedActorComponent(ActorComp, ParentSceneComp))
-					{
-						AddTreeNodeFromComponent(ActorComp, FindOrCreateParentForExistingComponent(ActorComp, GetActorNode()));
-					}
+					AddTreeNodeFromComponent(ActorComp, FindOrCreateParentForExistingComponent(ActorComp, GetActorNode()));
 				}
 			}
 		}
@@ -4742,36 +4778,6 @@ void SSCSEditor::UpdateTree(bool bRegenerateTreeNodes)
 	SCSTreeWidget->RequestTreeRefresh();
 }
 
-void SSCSEditor::AddInstancedTreeNodesRecursive(USceneComponent* Component, FSCSEditorTreeNodePtrType TreeNode, TSet<UActorComponent*>& ComponentsToAdd)
-{
-	if (Component != nullptr)
-	{
-		TArray<USceneComponent*> Components = Component->GetAttachChildren();
-		for (USceneComponent* ChildComponent : Components)
-		{
-			if (ComponentsToAdd.Contains(ChildComponent)
-				&& ShouldAddInstancedActorComponent(ChildComponent, Component)
-				&& ChildComponent->GetOwner() == Component->GetOwner())
-			{
-				ComponentsToAdd.Remove(ChildComponent);
-
-				FSCSEditorTreeNodePtrType NewParentNode = AddTreeNodeFromComponent(ChildComponent, TreeNode);
-				AddInstancedTreeNodesRecursive(ChildComponent, NewParentNode, ComponentsToAdd);
-			}
-		}
-	}
-}
-
-bool SSCSEditor::ShouldAddInstancedActorComponent(UActorComponent* ActorComp, USceneComponent* ParentSceneComp) const
-{
-	// Exclude nested DSOs attached to BP-constructed instances, which are not mutable.
-	return (ActorComp != nullptr
-		&& (!ActorComp->IsVisualizationComponent())
-		&& (ActorComp->CreationMethod != EComponentCreationMethod::UserConstructionScript || !GetDefault<UBlueprintEditorSettings>()->bHideConstructionScriptComponentsInDetailsView)
-		&& (ParentSceneComp == nullptr || !ParentSceneComp->IsCreatedByConstructionScript() || !ActorComp->HasAnyFlags(RF_DefaultSubObject)))
-		&& (ActorComp->CreationMethod != EComponentCreationMethod::Native || FComponentEditorUtils::CanEditNativeComponent(ActorComp) );
-}
-
 void SSCSEditor::DumpTree()
 {
 	/* Example:
diff --git a/Engine/Source/Editor/Kismet/Public/SSCSEditor.h b/Engine/Source/Editor/Kismet/Public/SSCSEditor.h
--- a/Engine/Source/Editor/Kismet/Public/SSCSEditor.h
+++ b/Engine/Source/Editor/Kismet/Public/SSCSEditor.h
@@ -1072,12 +1072,6 @@ protected:
 	/** Helper method to add a tree node for the given actor component */
 	FSCSEditorTreeNodePtrType AddTreeNodeFromComponent(UActorComponent* InSceneComponent, FSCSEditorTreeNodePtrType InParentTreeNode = FSCSEditorTreeNodePtrType());
 	
-	/** Helper method to recursively add tree nodes for an already-attached instanced scene component hierarchy */
-	void AddInstancedTreeNodesRecursive(USceneComponent* Component, FSCSEditorTreeNodePtrType TreeNode, TSet<UActorComponent*>& ComponentsToAdd);
-
-	/** Helper method to determine whether or not an instanced actor component should be added to the tree */
-	bool ShouldAddInstancedActorComponent(UActorComponent* ActorComp, USceneComponent* ParentSceneComp = nullptr) const;
-
 	/** Helper method to recursively find a tree node for the given SCS node starting at the given tree node */
 	FSCSEditorTreeNodePtrType FindTreeNode(const USCS_Node* InSCSNode, FSCSEditorTreeNodePtrType InStartNodePtr = FSCSEditorTreeNodePtrType()) const;
 
