Apply similar fix as CL# 4968306 to fix crash when renaming component via variable name

diff --git a/Engine/Source/Editor/Kismet/Private/BlueprintDetailsCustomization.cpp b/Engine/Source/Editor/Kismet/Private/BlueprintDetailsCustomization.cpp
--- a/Engine/Source/Editor/Kismet/Private/BlueprintDetailsCustomization.cpp
+++ b/Engine/Source/Editor/Kismet/Private/BlueprintDetailsCustomization.cpp
@@ -5935,16 +5935,24 @@ void FBlueprintComponentDetails::OnVariableTextChanged(const FText& InNewText)
 
 	bIsVariableNameInvalid = true;
 
+	const FString& NewTextStr = InNewText.ToString();
+
 	USCS_Node* SCS_Node = CachedNodePtr->GetSCSNode();
-	if(SCS_Node != NULL && !InNewText.IsEmpty() && !FComponentEditorUtils::IsValidVariableNameString(SCS_Node->ComponentTemplate, InNewText.ToString()))
+	if(SCS_Node != NULL && !NewTextStr.IsEmpty() && !FComponentEditorUtils::IsValidVariableNameString(SCS_Node->ComponentTemplate, NewTextStr))
 	{
 		VariableNameEditableTextBox->SetError(LOCTEXT("ComponentVariableRenameFailed_NotValid", "This name is reserved for engine use."));
 		return;
 	}
 
+	if (!FComponentEditorUtils::IsComponentNameAvailable(NewTextStr, SCS_Node->ComponentTemplate->GetOuter(), SCS_Node->ComponentTemplate))
+	{
+		VariableNameEditableTextBox->SetError(FText::Format(LOCTEXT("ComponentVariableRenameFailed_InUse", "{0} is in use by another variable or function!"), InNewText));
+		return;
+	}
+
 	TSharedPtr<INameValidatorInterface> VariableNameValidator = MakeShareable(new FKismetNameValidator(GetBlueprintObj(), CachedNodePtr->GetVariableName()));
 
-	EValidatorResult ValidatorResult = VariableNameValidator->IsValid(InNewText.ToString());
+	EValidatorResult ValidatorResult = VariableNameValidator->IsValid(NewTextStr);
 	if(ValidatorResult == EValidatorResult::AlreadyInUse)
 	{
 		VariableNameEditableTextBox->SetError(FText::Format(LOCTEXT("ComponentVariableRenameFailed_InUse", "{0} is in use by another variable or function!"), InNewText));
