Fix fallout from CL# 4968306 that would display that the name was already in use when setting the name back to itself

diff --git a/Engine/Source/Editor/Kismet/Private/SSCSEditor.cpp b/Engine/Source/Editor/Kismet/Private/SSCSEditor.cpp
--- a/Engine/Source/Editor/Kismet/Private/SSCSEditor.cpp
+++ b/Engine/Source/Editor/Kismet/Private/SSCSEditor.cpp
@@ -3137,8 +3137,15 @@ bool SSCS_RowWidget::OnNameTextVerifyChanged(const FText& InNewText, FText& OutE
 	FSCSEditorTreeNodePtrType NodePtr = GetNode();
 	UBlueprint* Blueprint = GetBlueprint();
 
-	if (!InNewText.IsEmpty())
+	const FString& NewTextStr = InNewText.ToString();
+
+	if (!NewTextStr.IsEmpty())
 	{
+		if (NodePtr->GetVariableName().ToString() == NewTextStr)
+		{
+			return true;
+		}
+
 		const UActorComponent* ComponentInstance = NodePtr->GetComponentTemplate();
 		if (ensure(ComponentInstance))
 		{
@@ -3148,20 +3155,20 @@ bool SSCS_RowWidget::OnNameTextVerifyChanged(const FText& InNewText, FText& OutE
 				ExistingNameSearchScope = Cast<AActor>(Blueprint->GeneratedClass->GetDefaultObject());
 			}
 
-			if (!FComponentEditorUtils::IsValidVariableNameString(ComponentInstance, InNewText.ToString()))
+			if (!FComponentEditorUtils::IsValidVariableNameString(ComponentInstance, NewTextStr))
 			{
 				OutErrorMessage = LOCTEXT("RenameFailed_EngineReservedName", "This name is reserved for engine use.");
 				return false;
 			}
-			else if (InNewText.ToString().Len() > NAME_SIZE)
+			else if (NewTextStr.Len() > NAME_SIZE)
 			{
 				FFormatNamedArguments Arguments;
 				Arguments.Add(TEXT("CharCount"), NAME_SIZE);
 				OutErrorMessage = FText::Format(LOCTEXT("ComponentRenameFailed_TooLong", "Component name must be less than {CharCount} characters long."), Arguments);
 				return false;
 			}
-			else if (!FComponentEditorUtils::IsComponentNameAvailable(InNewText.ToString(), ExistingNameSearchScope, ComponentInstance) 
-					|| !FComponentEditorUtils::IsComponentNameAvailable(InNewText.ToString(), ComponentInstance->GetOuter(), ComponentInstance ))
+			else if (!FComponentEditorUtils::IsComponentNameAvailable(NewTextStr, ExistingNameSearchScope, ComponentInstance) 
+					|| !FComponentEditorUtils::IsComponentNameAvailable(NewTextStr, ComponentInstance->GetOuter(), ComponentInstance ))
 			{
 				OutErrorMessage = LOCTEXT("RenameFailed_ExistingName", "Another component already has the same name.");
 				return false;
@@ -3184,7 +3191,7 @@ bool SSCS_RowWidget::OnNameTextVerifyChanged(const FText& InNewText, FText& OutE
 		NameValidator = MakeShareable(new FStringSetNameValidator(NodePtr->GetComponentTemplate()->GetName()));
 	}
 
-	EValidatorResult ValidatorResult = NameValidator->IsValid(InNewText.ToString());
+	EValidatorResult ValidatorResult = NameValidator->IsValid(NewTextStr);
 	if (ValidatorResult == EValidatorResult::AlreadyInUse)
 	{
 		OutErrorMessage = FText::Format(LOCTEXT("RenameFailed_InUse", "{0} is in use by another variable or function!"), InNewText);
