PR #5715: UE-36821: Open BP editor to view property from Find in Blueprints (Contributed by [REDACTED_PERSON])

diff --git a/Engine/Source/Editor/Kismet/Private/FindInBlueprints.cpp b/Engine/Source/Editor/Kismet/Private/FindInBlueprints.cpp
--- a/Engine/Source/Editor/Kismet/Private/FindInBlueprints.cpp
+++ b/Engine/Source/Editor/Kismet/Private/FindInBlueprints.cpp
@@ -418,7 +418,7 @@ FReply FFindInBlueprintsProperty::OnClick()
 		UBlueprint* Blueprint = GetParentBlueprint();
 		if (Blueprint)
 		{
-			TSharedPtr<IBlueprintEditor> BlueprintEditor = FKismetEditorUtilities::GetIBlueprintEditorForObject(Blueprint, false);
+			TSharedPtr<IBlueprintEditor> BlueprintEditor = FKismetEditorUtilities::GetIBlueprintEditorForObject(Blueprint, true);
 
 			if (BlueprintEditor.IsValid())
 			{
