Remove confusing tooltip that didn't provide much useful information
#rnx

diff --git a/Engine/Source/Editor/Kismet/Private/SBlueprintPalette.cpp b/Engine/Source/Editor/Kismet/Private/SBlueprintPalette.cpp
--- a/Engine/Source/Editor/Kismet/Private/SBlueprintPalette.cpp
+++ b/Engine/Source/Editor/Kismet/Private/SBlueprintPalette.cpp
@@ -933,18 +933,7 @@ private:
 		}
 		else
 		{
-			TSharedPtr<FEdGraphSchemaAction_K2Var> VarAction = StaticCastSharedPtr<FEdGraphSchemaAction_K2Var>(ActionPtr.Pin());
-
-			FString Result;
-			FBlueprintEditorUtils::GetBlueprintVariableMetaData(BlueprintObj, VarAction->GetVariableName(), nullptr, TEXT("tooltip"), Result);
-			if ( !Result.IsEmpty() )
-			{
-				ToolTipText = LOCTEXT("VariablePrivacy_is_public_Tooltip", "Variable is public and is editable on each instance of this Blueprint.");
-			}
-			else
-			{
-				ToolTipText = LOCTEXT("VariablePrivacy_is_public_no_tooltip_Tooltip", "Variable is public but MISSING TOOLTIP.");
-			}
+			ToolTipText = LOCTEXT("VariablePrivacy_is_public_Tooltip", "Variable is public and is editable on each instance of this Blueprint.");
 		}
 		return ToolTipText;
 	}
