Fix static analysis warning
#rnx

diff --git a/Engine/Source/Editor/Kismet/Private/BlueprintDetailsCustomization.cpp b/Engine/Source/Editor/Kismet/Private/BlueprintDetailsCustomization.cpp
--- a/Engine/Source/Editor/Kismet/Private/BlueprintDetailsCustomization.cpp
+++ b/Engine/Source/Editor/Kismet/Private/BlueprintDetailsCustomization.cpp
@@ -5937,17 +5937,19 @@ void FBlueprintComponentDetails::OnVariableTextChanged(const FText& InNewText)
 
 	const FString& NewTextStr = InNewText.ToString();
 
-	USCS_Node* SCS_Node = CachedNodePtr->GetSCSNode();
-	if(SCS_Node != NULL && !NewTextStr.IsEmpty() && !FComponentEditorUtils::IsValidVariableNameString(SCS_Node->ComponentTemplate, NewTextStr))
+	if (USCS_Node* SCS_Node = CachedNodePtr->GetSCSNode())
 	{
-		VariableNameEditableTextBox->SetError(LOCTEXT("ComponentVariableRenameFailed_NotValid", "This name is reserved for engine use."));
-		return;
-	}
+		if (!NewTextStr.IsEmpty() && !FComponentEditorUtils::IsValidVariableNameString(SCS_Node->ComponentTemplate, NewTextStr))
+		{
+			VariableNameEditableTextBox->SetError(LOCTEXT("ComponentVariableRenameFailed_NotValid", "This name is reserved for engine use."));
+			return;
+		}
 
-	if (!FComponentEditorUtils::IsComponentNameAvailable(NewTextStr, SCS_Node->ComponentTemplate->GetOuter(), SCS_Node->ComponentTemplate))
-	{
-		VariableNameEditableTextBox->SetError(FText::Format(LOCTEXT("ComponentVariableRenameFailed_InUse", "{0} is in use by another variable or function!"), InNewText));
-		return;
+		if (!FComponentEditorUtils::IsComponentNameAvailable(NewTextStr, SCS_Node->ComponentTemplate->GetOuter(), SCS_Node->ComponentTemplate))
+		{
+			VariableNameEditableTextBox->SetError(FText::Format(LOCTEXT("ComponentVariableRenameFailed_InUse", "{0} is in use by another variable or function!"), InNewText));
+			return;
+		}
 	}
 
 	TSharedPtr<INameValidatorInterface> VariableNameValidator = MakeShareable(new FKismetNameValidator(GetBlueprintObj(), CachedNodePtr->GetVariableName()));
