Remove unused parameter to bound lambda
#rnx

diff --git a/Engine/Source/Editor/Kismet/Private/DiffUtils.cpp b/Engine/Source/Editor/Kismet/Private/DiffUtils.cpp
--- a/Engine/Source/Editor/Kismet/Private/DiffUtils.cpp
+++ b/Engine/Source/Editor/Kismet/Private/DiffUtils.cpp
@@ -783,7 +783,7 @@ TSharedRef< STreeView<TSharedPtr< FBlueprintDifferenceTreeEntry > > > DiffTreeVi
 			];
 	};
 
-	const auto ChildrenAccessor = [](TSharedPtr<FBlueprintDifferenceTreeEntry> InTreeItem, TArray< TSharedPtr< FBlueprintDifferenceTreeEntry > >& OutChildren, TArray< TSharedPtr<FBlueprintDifferenceTreeEntry> >* MasterList)
+	const auto ChildrenAccessor = [](TSharedPtr<FBlueprintDifferenceTreeEntry> InTreeItem, TArray< TSharedPtr< FBlueprintDifferenceTreeEntry > >& OutChildren)
 	{
 		OutChildren = InTreeItem->Children;
 	};
@@ -798,7 +798,7 @@ TSharedRef< STreeView<TSharedPtr< FBlueprintDifferenceTreeEntry > > > DiffTreeVi
 
 	return SNew(STreeView< TSharedPtr< FBlueprintDifferenceTreeEntry > >)
 		.OnGenerateRow(STreeView< TSharedPtr< FBlueprintDifferenceTreeEntry > >::FOnGenerateRow::CreateStatic(RowGenerator))
-		.OnGetChildren(STreeView< TSharedPtr< FBlueprintDifferenceTreeEntry > >::FOnGetChildren::CreateStatic(ChildrenAccessor, DifferencesList))
+		.OnGetChildren(STreeView< TSharedPtr< FBlueprintDifferenceTreeEntry > >::FOnGetChildren::CreateStatic(ChildrenAccessor))
 		.OnSelectionChanged(STreeView< TSharedPtr< FBlueprintDifferenceTreeEntry > >::FOnSelectionChanged::CreateStatic(Selector))
 		.TreeItemsSource(DifferencesList);
 }
