Fix crash when major blueprint changes have been made to a blueprint that is open in the diff window

diff --git a/Engine/Source/Editor/Kismet/Private/SBlueprintDiff.cpp b/Engine/Source/Editor/Kismet/Private/SBlueprintDiff.cpp
--- a/Engine/Source/Editor/Kismet/Private/SBlueprintDiff.cpp
+++ b/Engine/Source/Editor/Kismet/Private/SBlueprintDiff.cpp
@@ -85,7 +85,7 @@ class FMyBlueprintDiffControl : public TSharedFromThis<FMyBlueprintDiffControl>,
 {
 public:
 	FMyBlueprintDiffControl(const UBlueprint* InOldBlueprint, const UBlueprint* InNewBlueprint, FOnDiffEntryFocused InSelectionCallback)
-		: SelectionCallback(InSelectionCallback), OldBlueprint(InOldBlueprint), NewBlueprint(InNewBlueprint)
+		: SelectionCallback(MoveTemp(InSelectionCallback)), OldBlueprint(InOldBlueprint), NewBlueprint(InNewBlueprint)
 	{
 	}
 
@@ -1527,7 +1527,7 @@ void SBlueprintDiff::GenerateDifferencesList()
 		Graph->GenerateTreeEntries(MasterDifferencesList, RealDifferences);
 	}
 
-	DifferencesTreeView->RequestTreeRefresh();
+	DifferencesTreeView->RebuildList();
 }
 
 SBlueprintDiff::FDiffControl SBlueprintDiff::GenerateBlueprintTypePanel()
