Edit the behavior of the Category dropdown for components to be more consistent in allowing what can and cannot be changed.

diff --git a/Engine/Source/Editor/Kismet/Private/BlueprintDetailsCustomization.cpp b/Engine/Source/Editor/Kismet/Private/BlueprintDetailsCustomization.cpp
--- a/Engine/Source/Editor/Kismet/Private/BlueprintDetailsCustomization.cpp
+++ b/Engine/Source/Editor/Kismet/Private/BlueprintDetailsCustomization.cpp
@@ -6027,8 +6027,8 @@ void FBlueprintComponentDetails::OnTooltipTextCommitted(const FText& NewText, ET
 bool FBlueprintComponentDetails::OnVariableCategoryChangeEnabled() const
 {
 	check(CachedNodePtr.IsValid());
-
-	return !CachedNodePtr->CanRename();
+	
+	return !CachedNodePtr->IsInherited();
 }
 
 FText FBlueprintComponentDetails::OnGetVariableCategoryText() const
