PR #5509: Sort function overrides by Name (Contributed by [REDACTED_PERSON])

diff --git a/Engine/Source/Editor/Kismet/Private/SMyBlueprint.cpp b/Engine/Source/Editor/Kismet/Private/SMyBlueprint.cpp
--- a/Engine/Source/Editor/Kismet/Private/SMyBlueprint.cpp
+++ b/Engine/Source/Editor/Kismet/Private/SMyBlueprint.cpp
@@ -841,6 +841,12 @@ TSharedRef<SWidget> SMyBlueprint::OnGetFunctionListMenu()
 
 void SMyBlueprint::BuildOverridableFunctionsMenu(FMenuBuilder& MenuBuilder)
 {
+	// Sort by function name so that it's easier for users to find the function they're looking for:
+    OverridableFunctionActions.Sort([](const TSharedPtr<FEdGraphSchemaAction_K2Graph> &LHS, const TSharedPtr<FEdGraphSchemaAction_K2Graph> &RHS)
+    {
+        return LHS->GetMenuDescription().CompareToCaseIgnored(RHS->GetMenuDescription()) < 0;
+    });
+	
 	MenuBuilder.BeginSection("OverrideFunction", LOCTEXT("OverrideFunction", "Override Function"));
 	{
 		for ( auto& OverrideAction : OverridableFunctionActions )
