Expose multi line option for arrays, sets, and maps that contain strings or text as the inner/value property in BP editor
Adapted from PR #4038: (Contributed by [REDACTED_PERSON])

diff --git a/Engine/Source/Editor/Kismet/Private/BlueprintDetailsCustomization.cpp b/Engine/Source/Editor/Kismet/Private/BlueprintDetailsCustomization.cpp
--- a/Engine/Source/Editor/Kismet/Private/BlueprintDetailsCustomization.cpp
+++ b/Engine/Source/Editor/Kismet/Private/BlueprintDetailsCustomization.cpp
@@ -2504,17 +2504,32 @@ void FBlueprintVarActionDetails::OnAdvancedDisplayChanged(ECheckBoxState InNewSt
 
 EVisibility FBlueprintVarActionDetails::GetMultilineVisibility() const
 {
-	if (UProperty* VariableProperty = CachedVariableProperty.Get())
+	UProperty* VariableProperty = nullptr;
+	if (UProperty* RawVariableProperty = CachedVariableProperty.Get())
 	{
-		if (IsABlueprintVariable(VariableProperty))
+		if (IsABlueprintVariable(RawVariableProperty))
 		{
-			if (VariableProperty->IsA(UTextProperty::StaticClass()) || VariableProperty->IsA(UStrProperty::StaticClass()))
+			if (const UArrayProperty* ArrayProperty = Cast<UArrayProperty>(RawVariableProperty))
 			{
-				return EVisibility::Visible;
+				VariableProperty = ArrayProperty->Inner;
+			}
+			else if (const USetProperty* SetProperty = Cast<USetProperty>(RawVariableProperty))
+			{
+				VariableProperty = SetProperty->ElementProp;
+			}
+			else if (const UMapProperty* MapProperty = Cast<UMapProperty>(RawVariableProperty))
+			{
+				VariableProperty = MapProperty->ValueProp;
+			}
+			else
+			{
+				VariableProperty = RawVariableProperty;
 			}
 		}
 	}
-	return EVisibility::Collapsed;
+
+	const bool bCanBeMultiline = (VariableProperty != nullptr) && (VariableProperty->IsA(UTextProperty::StaticClass()) || VariableProperty->IsA(UStrProperty::StaticClass()));
+	return bCanBeMultiline ? EVisibility::Visible : EVisibility::Collapsed;
 }
 
 ECheckBoxState FBlueprintVarActionDetails::OnGetMultilineCheckboxState() const
