Removed "Delegate" section that contained duplicate copies of tooltip and category when editing a selected event dispatcher (they're still variables and those can both be set in the variable section of the details panel)

diff --git a/Engine/Source/Editor/Kismet/Private/BlueprintDetailsCustomization.cpp b/Engine/Source/Editor/Kismet/Private/BlueprintDetailsCustomization.cpp
--- a/Engine/Source/Editor/Kismet/Private/BlueprintDetailsCustomization.cpp
+++ b/Engine/Source/Editor/Kismet/Private/BlueprintDetailsCustomization.cpp
@@ -4068,84 +4068,6 @@ UEdGraph* FBlueprintDelegateActionDetails::GetGraph() const
 	return NULL;
 }
 
-FText FBlueprintDelegateActionDetails::OnGetTooltipText() const
-{
-	if (UMulticastDelegateProperty* DelegateProperty = GetDelegateProperty())
-	{
-		FString Result;
-		FBlueprintEditorUtils::GetBlueprintVariableMetaData(GetBlueprintObj(), DelegateProperty->GetFName(), NULL, TEXT("tooltip"), Result);
-		return FText::FromString(Result);
-	}
-	return FText();
-}
-
-void FBlueprintDelegateActionDetails::OnTooltipTextCommitted(const FText& NewText, ETextCommit::Type InTextCommit)
-{
-	if (UMulticastDelegateProperty* DelegateProperty = GetDelegateProperty())
-	{
-		FBlueprintEditorUtils::SetBlueprintVariableMetaData(GetBlueprintObj(), DelegateProperty->GetFName(), NULL, TEXT("tooltip"), NewText.ToString() );
-	}
-}
-
-FText FBlueprintDelegateActionDetails::OnGetCategoryText() const
-{
-	if (UMulticastDelegateProperty* DelegateProperty = GetDelegateProperty())
-	{
-		FName DelegateName = DelegateProperty->GetFName();
-		FText Category = FBlueprintEditorUtils::GetBlueprintVariableCategory(GetBlueprintObj(), DelegateName, NULL);
-
-		// Older blueprints will have their name as the default category
-		if( Category.EqualTo(FText::FromString(GetBlueprintObj()->GetName())) || Category.EqualTo(UEdGraphSchema_K2::VR_DefaultCategory) )
-		{
-			return LOCTEXT("DefaultCategory", "Default");
-		}
-		else
-		{
-			return Category;
-		}
-	}
-	return FText();
-}
-
-void FBlueprintDelegateActionDetails::OnCategoryTextCommitted(const FText& NewText, ETextCommit::Type InTextCommit)
-{
-	if (InTextCommit == ETextCommit::OnEnter || InTextCommit == ETextCommit::OnUserMovedFocus)
-	{
-		if (UMulticastDelegateProperty* DelegateProperty = GetDelegateProperty())
-		{
-			// Remove excess whitespace and prevent categories with just spaces
-			FText CategoryName = FText::TrimPrecedingAndTrailing(NewText);
-			
-			FBlueprintEditorUtils::SetBlueprintVariableCategory(GetBlueprintObj(), DelegateProperty->GetFName(), NULL, CategoryName);
-			check(MyBlueprint.IsValid());
-			FBlueprintVarActionDetails::PopulateCategories(MyBlueprint.Pin().Get(), CategorySource);
-			MyBlueprint.Pin()->ExpandCategory(CategoryName);
-		}
-	}
-}
-
-TSharedRef< ITableRow > FBlueprintDelegateActionDetails::MakeCategoryViewWidget( TSharedPtr<FText> Item, const TSharedRef< STableViewBase >& OwnerTable )
-{
-	return SNew(STableRow<TSharedPtr<FString>>, OwnerTable)
-		[
-			SNew(STextBlock) .Text(*Item.Get())
-		];
-}
-
-void FBlueprintDelegateActionDetails::OnCategorySelectionChanged( TSharedPtr<FText> ProposedSelection, ESelectInfo::Type /*SelectInfo*/ )
-{
-	UMulticastDelegateProperty* DelegateProperty = GetDelegateProperty();
-	if (DelegateProperty && ProposedSelection.IsValid())
-	{
-		FText NewCategory = *ProposedSelection.Get();
-
-		FBlueprintEditorUtils::SetBlueprintVariableCategory(GetBlueprintObj(), DelegateProperty->GetFName(), NULL, NewCategory);
-		CategoryListView.Pin()->ClearSelection();
-		CategoryComboButton.Pin()->SetIsOpen(false);
-		MyBlueprint.Pin()->ExpandCategory(NewCategory);
-	}
-}
-
 void FBlueprintDelegateActionDetails::CustomizeDetails( IDetailLayoutBuilder& DetailLayout )
 {
 	DetailsLayoutPtr = &DetailLayout;
@@ -4156,72 +4078,6 @@ void FBlueprintDelegateActionDetails::CustomizeDetails( IDetailLayoutBuilder& De
 	const UEdGraphSchema_K2* Schema = GetDefault<UEdGraphSchema_K2>();
 	const FSlateFontInfo DetailFontInfo = IDetailLayoutBuilder::GetDetailFont();
 
-	{
-		IDetailCategoryBuilder& Category = DetailLayout.EditCategory("Delegate", LOCTEXT("DelegateDetailsCategory", "Delegate"));
-		Category.AddCustomRow( LOCTEXT("VariableToolTipLabel", "Tooltip") )
-		.NameContent()
-		[
-			SNew(STextBlock)
-				.Text( LOCTEXT("VariableToolTipLabel", "Tooltip") )
-				.Font( DetailFontInfo )
-		]
-		.ValueContent()
-		[
-			SNew(SEditableTextBox)
-				.Text( this, &FBlueprintDelegateActionDetails::OnGetTooltipText )
-				.OnTextCommitted( this, &FBlueprintDelegateActionDetails::OnTooltipTextCommitted)
-				.Font( DetailFontInfo )
-		];
-
-		FBlueprintVarActionDetails::PopulateCategories(MyBlueprint.Pin().Get(), CategorySource);
-		TSharedPtr<SComboButton> NewComboButton;
-		TSharedPtr<SListView<TSharedPtr<FText>>> NewListView;
-
-		Category.AddCustomRow( LOCTEXT("CategoryLabel", "Category") )
-		.NameContent()
-		[
-			SNew(STextBlock)
-				.Text( LOCTEXT("CategoryLabel", "Category") )
-				.Font( DetailFontInfo )
-		]
-		.ValueContent()
-		[
-			SAssignNew(NewComboButton, SComboButton)
-			.ContentPadding(FMargin(0,0,5,0))
-			.IsEnabled(this, &FBlueprintDelegateActionDetails::IsBlueprintProperty)
-			.ButtonContent()
-			[
-				SNew(SBorder)
-				.BorderImage( FEditorStyle::GetBrush("NoBorder") )
-				.Padding(FMargin(0, 0, 5, 0))
-				[
-					SNew(SEditableTextBox)
-						.Text(this, &FBlueprintDelegateActionDetails::OnGetCategoryText)
-						.OnTextCommitted(this, &FBlueprintDelegateActionDetails::OnCategoryTextCommitted)
-						.SelectAllTextWhenFocused(true)
-						.RevertTextOnEscape(true)
-						.Font( DetailFontInfo )
-				]
-			]
-			.MenuContent()
-			[
-				SNew(SVerticalBox)
-				+SVerticalBox::Slot()
-				.AutoHeight()
-				.MaxHeight(400.0f)
-				[
-					SAssignNew(NewListView, SListView<TSharedPtr<FText>>)
-						.ListItemsSource(&CategorySource)
-						.OnGenerateRow(this, &FBlueprintDelegateActionDetails::MakeCategoryViewWidget)
-						.OnSelectionChanged(this, &FBlueprintDelegateActionDetails::OnCategorySelectionChanged)
-				]
-			]
-		];
-
-		CategoryComboButton = NewComboButton;
-		CategoryListView = NewListView;
-	}
-
 	if (UK2Node_EditablePinBase* FunctionEntryNode = FunctionEntryNodePtr.Get())
 	{
 		IDetailCategoryBuilder& InputsCategory = DetailLayout.EditCategory("DelegateInputs", LOCTEXT("DelegateDetailsInputs", "Inputs"));
diff --git a/Engine/Source/Editor/Kismet/Private/BlueprintDetailsCustomization.h b/Engine/Source/Editor/Kismet/Private/BlueprintDetailsCustomization.h
--- a/Engine/Source/Editor/Kismet/Private/BlueprintDetailsCustomization.h
+++ b/Engine/Source/Editor/Kismet/Private/BlueprintDetailsCustomization.h
@@ -395,12 +395,6 @@ private:
 	void SetEntryNode();
 
 	UMulticastDelegateProperty* GetDelegateProperty() const;
-	FText OnGetTooltipText() const;
-	void OnTooltipTextCommitted(const FText& NewText, ETextCommit::Type InTextCommit);
-	FText OnGetCategoryText() const;
-	void OnCategoryTextCommitted(const FText& NewText, ETextCommit::Type InTextCommit);
-	TSharedRef< ITableRow > MakeCategoryViewWidget( TSharedPtr<FText> Item, const TSharedRef< STableViewBase >& OwnerTable );
-	void OnCategorySelectionChanged( TSharedPtr<FText> ProposedSelection, ESelectInfo::Type /*SelectInfo*/ );
 
 	void CollectAvailibleSignatures();
 	void OnFunctionSelected(TSharedPtr<FString> FunctionItemData, ESelectInfo::Type SelectInfo);
@@ -408,14 +402,6 @@ private:
 	EVisibility OnGetSectionTextVisibility(TWeakPtr<SWidget> RowWidget) const;
 
 private:
-
-	/** A list of all category names to choose from */
-	TArray<TSharedPtr<FText>> CategorySource;
-
-	/** Widgets for the categories */
-	TWeakPtr<SComboButton> CategoryComboButton;
-	TWeakPtr<SListView<TSharedPtr<FText>>> CategoryListView;
-
 	TArray<TSharedPtr<FString>> FunctionsToCopySignatureFrom;
 	TSharedPtr<STextComboBox> CopySignatureComboButton;
 };
