Fix crash changing class on a blueprint while the diff window is open due to stale variable pointers

diff --git a/Engine/Source/Editor/Kismet/Private/SBlueprintPalette.cpp b/Engine/Source/Editor/Kismet/Private/SBlueprintPalette.cpp
--- a/Engine/Source/Editor/Kismet/Private/SBlueprintPalette.cpp
+++ b/Engine/Source/Editor/Kismet/Private/SBlueprintPalette.cpp
@@ -710,11 +710,11 @@ public:
 private:
 	FEdGraphPinType OnGetVarType() const
 	{
-		if (VariableProperty)
+		if (UProperty* VarProp = VariableProperty.Get())
 		{
 			const UEdGraphSchema_K2* K2Schema = GetDefault<UEdGraphSchema_K2>();
 			FEdGraphPinType Type;
-			K2Schema->ConvertPropertyToPinType(VariableProperty, Type);
+			K2Schema->ConvertPropertyToPinType(VarProp, Type);
 			return Type;
 		}
 		return FEdGraphPinType();
@@ -724,16 +724,16 @@ private:
 	{
 		if (FBlueprintEditorUtils::IsPinTypeValid(InNewPinType))
 		{
-			if (VariableProperty)
+			if (UProperty* VarProp = VariableProperty.Get())
 			{
-				FName VarName = VariableProperty->GetFName();
+				FName VarName = VarProp->GetFName();
 
 				if (VarName != NAME_None)
 				{
 					// Set the MyBP tab's last pin type used as this, for adding lots of variables of the same type
 					BlueprintEditorPtr.Pin()->GetMyBlueprintWidget()->GetLastPinTypeUsed() = InNewPinType;
 
-					if (UFunction* LocalVariableScope = Cast<UFunction>(VariableProperty->GetOuter()))
+					if (UFunction* LocalVariableScope = Cast<UFunction>(VarProp->GetOuter()))
 					{
 						FBlueprintEditorUtils::ChangeLocalVariableType(BlueprintObj, LocalVariableScope, VarName, InNewPinType);
 					}
@@ -757,7 +757,7 @@ private:
 	TWeakPtr<FBlueprintEditor>     BlueprintEditorPtr;
 
 	/** Variable Property to change the type of */
-	UProperty* VariableProperty;
+	TWeakObjectPtr<UProperty> VariableProperty;
 };
 
 /*******************************************************************************
