PR #5547: Hide Unrelated Nodes Blueprints Editor Feature (Contributed by [REDACTED_PERSON])

diff --git a/Engine/Plugins/Experimental/ControlRig/Source/ControlRigEditor/Private/Graph/SControlRigGraphNode.cpp b/Engine/Plugins/Experimental/ControlRig/Source/ControlRigEditor/Private/Graph/SControlRigGraphNode.cpp
--- a/Engine/Plugins/Experimental/ControlRig/Source/ControlRigEditor/Private/Graph/SControlRigGraphNode.cpp
+++ b/Engine/Plugins/Experimental/ControlRig/Source/ControlRigEditor/Private/Graph/SControlRigGraphNode.cpp
@@ -614,7 +614,7 @@ FSlateColor SControlRigGraphNode::GetPinTextColor(TWeakPtr<SGraphPin> GraphPin)
 		// If there is no schema there is no owning node (or basically this is a deleted node)
 		if (GraphNode)
 		{
-			if(!GraphNode->IsNodeEnabled() || GraphNode->IsDisplayAsDisabledForced() || !GraphPin.Pin()->IsEditingEnabled())
+			if(!GraphNode->IsNodeEnabled() || GraphNode->IsDisplayAsDisabledForced() || !GraphPin.Pin()->IsEditingEnabled() || GraphNode->IsNodeUnrelated())
 			{
 				return FLinearColor(1.0f, 1.0f, 1.0f, 0.5f);
 			}
diff --git a/Engine/Source/Editor/AnimGraph/Private/AnimGraphCommands.cpp b/Engine/Source/Editor/AnimGraph/Private/AnimGraphCommands.cpp
--- a/Engine/Source/Editor/AnimGraph/Private/AnimGraphCommands.cpp
+++ b/Engine/Source/Editor/AnimGraph/Private/AnimGraphCommands.cpp
@@ -7,6 +7,7 @@
 void FAnimGraphCommands::RegisterCommands()
 {
 	UI_COMMAND(TogglePoseWatch, "Toggle Pose Watch", "Toggle Pose Watching on this node", EUserInterfaceActionType::Button, FInputChord());
+	UI_COMMAND( ToggleHideUnrelatedNodes, "Hide Unrelated", "Toggles automatically hiding nodes which are unrelated to the selected nodes.", EUserInterfaceActionType::ToggleButton, FInputChord() );
 }
 
 #undef LOCTEXT_NAMESPACE
diff --git a/Engine/Source/Editor/AnimGraph/Public/AnimGraphCommands.h b/Engine/Source/Editor/AnimGraph/Public/AnimGraphCommands.h
--- a/Engine/Source/Editor/AnimGraph/Public/AnimGraphCommands.h
+++ b/Engine/Source/Editor/AnimGraph/Public/AnimGraphCommands.h
@@ -21,4 +21,7 @@ public:
 
 	// Toggle pose watching for a given node
 	TSharedPtr<FUICommandInfo> TogglePoseWatch;
+	
+	// Toggle hiding nodes which are not related to the selected nodes
+	TSharedPtr< FUICommandInfo > ToggleHideUnrelatedNodes;
 };
diff --git a/Engine/Source/Editor/AnimationBlueprintEditor/Private/AnimationBlueprintEditor.cpp b/Engine/Source/Editor/AnimationBlueprintEditor/Private/AnimationBlueprintEditor.cpp
--- a/Engine/Source/Editor/AnimationBlueprintEditor/Private/AnimationBlueprintEditor.cpp
+++ b/Engine/Source/Editor/AnimationBlueprintEditor/Private/AnimationBlueprintEditor.cpp
@@ -67,6 +67,13 @@
 #include "Widgets/Input/SButton.h"
 #include "EditorFontGlyphs.h"
 
+// Hide related nodes feature
+#include "Preferences/AnimationBlueprintEditorOptions.h"
+#include "Framework/MultiBox/MultiBoxBuilder.h"
+#include "EdGraphNode_Comment.h"
+#include "AnimStateNodeBase.h"
+#include "AnimStateEntryNode.h"
+
 #define LOCTEXT_NAMESPACE "AnimationBlueprintEditor"
 
 const FName AnimationBlueprintEditorAppName(TEXT("AnimationBlueprintEditorApp"));
@@ -162,6 +169,8 @@ FAnimationBlueprintEditor::~FAnimationBlueprintEditor()
 	FReimportManager::Instance()->OnPostReimport().RemoveAll(this);
 
 	// NOTE: Any tabs that we still have hanging out when destroyed will be cleaned up by FBaseToolkit's destructor
+
+	SaveEditorSettings();
 }
 
 UAnimBlueprint* FAnimationBlueprintEditor::GetAnimBlueprint() const
@@ -201,6 +210,8 @@ void FAnimationBlueprintEditor::InitAnimationBlueprintEditor(const EToolkitMode:
 		Toolbar = MakeShareable(new FBlueprintEditorToolbar(SharedThis(this)));
 	}
 
+	LoadEditorSettings();
+
 	GetToolkitCommands()->Append(FPlayWorldCommands::GlobalPlayWorldActions.ToSharedRef());
 
 	FPersonaModule& PersonaModule = FModuleManager::GetModuleChecked<FPersonaModule>("Persona");
@@ -311,6 +322,39 @@ void FAnimationBlueprintEditor::ExtendToolbar()
 		AddToolbarWidget(PersonaModule.CreateAssetFamilyShortcutWidget(SharedThis(this), AssetFamily));
 	}
 	));
+
+	struct Local
+	{
+		static void FillToolbar(FToolBarBuilder& ToolbarBuilder, const TSharedRef< FUICommandList > ToolkitCommands, FAnimationBlueprintEditor* BlueprintEditor)
+		{
+			ToolbarBuilder.BeginSection("Graph");
+			{
+				ToolbarBuilder.AddToolBarButton(
+					FAnimGraphCommands::Get().ToggleHideUnrelatedNodes,
+					NAME_None,
+					TAttribute<FText>(),
+					TAttribute<FText>(),
+					FSlateIcon(FEditorStyle::GetStyleSetName(), "GraphEditor.ToggleHideUnrelatedNodes")
+				);
+				ToolbarBuilder.AddComboButton(
+					FUIAction(),
+					FOnGetContent::CreateSP(BlueprintEditor, &FBlueprintEditor::MakeHideUnrelatedNodesOptionsMenu),
+					LOCTEXT("HideUnrelatedNodesOptions", "Hide Unrelated Nodes Options"),
+					LOCTEXT("HideUnrelatedNodesOptionsMenu", "Hide Unrelated Nodes options menu"),
+					TAttribute<FSlateIcon>(),
+					true
+				);
+			}
+			ToolbarBuilder.EndSection();
+		}
+	};
+
+    ToolbarExtender->AddToolBarExtension(
+		"Asset",
+		EExtensionHook::After,
+		GetToolkitCommands(),
+		FToolBarExtensionDelegate::CreateStatic( &Local::FillToolbar, GetToolkitCommands(), this )
+	);
 }
 
 UBlueprint* FAnimationBlueprintEditor::GetBlueprintObj() const
@@ -350,6 +394,11 @@ void FAnimationBlueprintEditor::OnGraphEditorFocused(const TSharedRef<class SGra
 	{
 		OnPinDefaultValueChangedHandle = AnimationGraph->OnPinDefaultValueChanged.Add(FOnPinDefaultValueChanged::FDelegate::CreateSP(this, &FAnimationBlueprintEditor::HandlePinDefaultValueChanged));
 	}
+
+	if (bHideUnrelatedNodes && GetSelectedNodes().Num() <= 0)
+	{
+		ResetAllNodesUnrelatedStates();
+	}
 }
 
 void FAnimationBlueprintEditor::OnGraphEditorBackgrounded(const TSharedRef<SGraphEditor>& InGraphEditor)
@@ -369,6 +418,12 @@ void FAnimationBlueprintEditor::CreateDefaultCommands()
 	if (GetBlueprintObj())
 	{
 		FBlueprintEditor::CreateDefaultCommands();
+
+		ToolkitCommands->MapAction(
+			FAnimGraphCommands::Get().ToggleHideUnrelatedNodes,
+			FExecuteAction::CreateSP(this, &FBlueprintEditor::ToggleHideUnrelatedNodes),
+			FCanExecuteAction(),
+			FIsActionChecked::CreateSP(this, &FBlueprintEditor::IsToggleHideUnrelatedNodesChecked));
 	}
 	else
 	{
@@ -1349,6 +1404,11 @@ void FAnimationBlueprintEditor::HandleOpenNewAsset(UObject* InNewAsset)
 	FAssetEditorManager::Get().OpenEditorForAsset(InNewAsset);
 }
 
+void FAnimationBlueprintEditor::AddReferencedObjects( FReferenceCollector& Collector )
+{
+    Collector.AddReferencedObject( EditorOptions );
+}
+
 FAnimNode_Base* FAnimationBlueprintEditor::FindAnimNode(UAnimGraphNode_Base* AnimGraphNode) const
 {
 	FAnimNode_Base* AnimNode = nullptr;
@@ -1397,6 +1457,29 @@ void FAnimationBlueprintEditor::OnSelectedNodesChangedImpl(const TSet<class UObj
 			}
 		}
 	}
+
+    bSelectRegularNode = false;
+	for (FGraphPanelSelectionSet::TConstIterator It(NewSelection); It; ++It)
+	{
+		UEdGraphNode_Comment* SeqNode = Cast<UEdGraphNode_Comment>(*It);
+		UAnimStateNodeBase* AnimGraphNodeBase = Cast<UAnimStateNodeBase>(*It);
+		UAnimStateEntryNode* AnimStateEntryNode = Cast<UAnimStateEntryNode>(*It);
+		if (!SeqNode && !AnimGraphNodeBase && !AnimStateEntryNode)
+		{
+			bSelectRegularNode = true;
+			break;
+		}
+	}
+
+    if (bHideUnrelatedNodes && !bLockNodeFadeState)
+	{
+		ResetAllNodesUnrelatedStates();
+
+		if ( bSelectRegularNode )
+		{
+			HideUnrelatedNodes();
+		}
+	}
 }
 
 void FAnimationBlueprintEditor::OnPostCompile()
@@ -1624,5 +1707,24 @@ void FAnimationBlueprintEditor::HandleViewportCreated(const TSharedRef<IPersonaV
 	);
 }
 
+void FAnimationBlueprintEditor::LoadEditorSettings()
+{
+	EditorOptions = NewObject<UAnimationBlueprintEditorOptions>();
+
+	if (EditorOptions->bHideUnrelatedNodes)
+	{
+		ToggleHideUnrelatedNodes();
+	}
+}
+
+void FAnimationBlueprintEditor::SaveEditorSettings()
+{
+	if ( EditorOptions )
+	{
+		EditorOptions->bHideUnrelatedNodes = bHideUnrelatedNodes;
+		EditorOptions->SaveConfig();
+	}
+}
+
 #undef LOCTEXT_NAMESPACE
 
diff --git a/Engine/Source/Editor/AnimationBlueprintEditor/Private/AnimationBlueprintEditor.h b/Engine/Source/Editor/AnimationBlueprintEditor/Private/AnimationBlueprintEditor.h
--- a/Engine/Source/Editor/AnimationBlueprintEditor/Private/AnimationBlueprintEditor.h
+++ b/Engine/Source/Editor/AnimationBlueprintEditor/Private/AnimationBlueprintEditor.h
@@ -11,6 +11,7 @@
 #include "IAnimationBlueprintEditor.h"
 #include "Containers/ArrayView.h"
 
+class UAnimationBlueprintEditorOptions;
 class IPersonaToolkit;
 class IPersonaViewport;
 class ISkeletonTree;
@@ -117,6 +118,10 @@ public:
 
 	/** Get the object to be displayed in the asset properties */
 	UObject* HandleGetObject();
+	
+	//~ Begin FGCObject Interface
+	virtual void AddReferencedObjects(FReferenceCollector& Collector) override;
+	//~ End FGCObject Interface
 
 	/** Handle opening a new asset from the asset browser */
 	void HandleOpenNewAsset(UObject* InNewAsset);
@@ -262,6 +267,16 @@ private:
 	/** Handle the viewport being created */
 	void HandleViewportCreated(const TSharedRef<IPersonaViewport>& InPersonaViewport);
 
+    /**
+	 * Load editor settings from disk (docking state, window pos/size, option state, etc).
+	 */
+	virtual void LoadEditorSettings();
+
+	/**
+	 * Saves editor settings to disk (docking state, window pos/size, option state, etc).
+	 */
+	virtual void SaveEditorSettings();
+
 	/** The extender to pass to the level editor to extend it's window menu */
 	TSharedPtr<FExtender> MenuExtender;
 
@@ -282,4 +297,7 @@ private:
 
 	/** Delegate handle registered for when pin default values change */
 	FDelegateHandle OnPinDefaultValueChangedHandle;
+
+    /** Configuration class used to store editor settings across sessions. */
+	UAnimationBlueprintEditorOptions* EditorOptions;
 };
diff --git a/Engine/Source/Editor/EditorStyle/Private/SlateEditorStyle.cpp b/Engine/Source/Editor/EditorStyle/Private/SlateEditorStyle.cpp
--- a/Engine/Source/Editor/EditorStyle/Private/SlateEditorStyle.cpp
+++ b/Engine/Source/Editor/EditorStyle/Private/SlateEditorStyle.cpp
@@ -4507,6 +4507,9 @@ void FSlateEditorStyle::FStyle::SetupGraphEditorStyles()
 
 		Set( "GraphEditor.DistributeNodesHorizontally", new IMAGE_BRUSH( "Icons/GraphEditor/icon_DistributeNodesHorizontally_20px", Icon20x20 ) );
 		Set( "GraphEditor.DistributeNodesVertically", new IMAGE_BRUSH( "Icons/GraphEditor/icon_DistributeNodesVertically_20px", Icon20x20 ) );
+		
+		Set( "GraphEditor.ToggleHideUnrelatedNodes", new IMAGE_BRUSH( "Icons/icon_HideUnrelatedNodes_40x", Icon40x40 ) );
+		Set( "GraphEditor.ToggleHideUnrelatedNodes.Small", new IMAGE_BRUSH( "Icons/icon_HideUnrelatedNodes_40x", Icon20x20 ) );
 
 		// Graph editor widgets
 		{
diff --git a/Engine/Source/Editor/GraphEditor/Private/BlueprintConnectionDrawingPolicy.cpp b/Engine/Source/Editor/GraphEditor/Private/BlueprintConnectionDrawingPolicy.cpp
--- a/Engine/Source/Editor/GraphEditor/Private/BlueprintConnectionDrawingPolicy.cpp
+++ b/Engine/Source/Editor/GraphEditor/Private/BlueprintConnectionDrawingPolicy.cpp
@@ -563,6 +563,11 @@ void FKismetConnectionDrawingPolicy::DetermineWiringStyle(UEdGraphPin* OutputPin
 			}
 		}
 
+		if ((OutputPin && OutputPin->GetOwningNode()->IsNodeUnrelated()) || (InputPin && InputPin->GetOwningNode()->IsNodeUnrelated()))
+		{
+			bWireIsOnDisabledNodeAndNotPassthru = true;
+		}
+
 		if (bWireIsOnDisabledNodeAndNotPassthru)
 		{
 			Params.WireColor *= 0.5f;
@@ -574,6 +579,7 @@ void FKismetConnectionDrawingPolicy::DetermineWiringStyle(UEdGraphPin* OutputPin
 	{
 		ApplyHoverDeemphasis(OutputPin, InputPin, /*inout*/ Params.WireThickness, /*inout*/ Params.WireColor);
 	}
+
 }
 
 void FKismetConnectionDrawingPolicy::SetIncompatiblePinDrawState(const TSharedPtr<SGraphPin>& StartPin, const TSet< TSharedRef<SWidget> >& VisiblePins)
diff --git a/Engine/Source/Editor/GraphEditor/Private/MaterialGraphConnectionDrawingPolicy.cpp b/Engine/Source/Editor/GraphEditor/Private/MaterialGraphConnectionDrawingPolicy.cpp
--- a/Engine/Source/Editor/GraphEditor/Private/MaterialGraphConnectionDrawingPolicy.cpp
+++ b/Engine/Source/Editor/GraphEditor/Private/MaterialGraphConnectionDrawingPolicy.cpp
@@ -125,7 +125,7 @@ void FMaterialGraphConnectionDrawingPolicy::DetermineWiringStyle(UEdGraphPin* Ou
 				Params.StartDirection = EGPD_Input;
 			}
 		}
-		else if (!OutputNode->IsNodeEnabled() || OutputNode->IsDisplayAsDisabledForced())
+		else if (!OutputNode->IsNodeEnabled() || OutputNode->IsDisplayAsDisabledForced() || OutputNode->IsNodeUnrelated())
 		{
 			Params.WireColor = MaterialGraphSchema->InactivePinColor;
 		}
@@ -145,7 +145,7 @@ void FMaterialGraphConnectionDrawingPolicy::DetermineWiringStyle(UEdGraphPin* Ou
 				Params.EndDirection = EGPD_Output;
 			}
 		}
-		else if (!InputNode->IsNodeEnabled() || InputNode->IsDisplayAsDisabledForced())
+		else if (!InputNode->IsNodeEnabled() || InputNode->IsDisplayAsDisabledForced() || InputNode->IsNodeUnrelated())
 		{
 			Params.WireColor = MaterialGraphSchema->InactivePinColor;
 		}
diff --git a/Engine/Source/Editor/GraphEditor/Private/MaterialNodes/SGraphNodeMaterialBase.cpp b/Engine/Source/Editor/GraphEditor/Private/MaterialNodes/SGraphNodeMaterialBase.cpp
--- a/Engine/Source/Editor/GraphEditor/Private/MaterialNodes/SGraphNodeMaterialBase.cpp
+++ b/Engine/Source/Editor/GraphEditor/Private/MaterialNodes/SGraphNodeMaterialBase.cpp
@@ -453,6 +453,7 @@ TSharedRef<SWidget> SGraphNodeMaterialBase::CreatePreviewWidget()
 
 		TSharedPtr<SViewport> ViewportWidget = 
 			SNew( SViewport )
+			.RenderDirectlyToWindow(true)
 			.EnableGammaCorrection(false);
 
 		PreviewViewport = MakeShareable(new FPreviewViewport(MaterialNode));
diff --git a/Engine/Source/Editor/GraphEditor/Private/SCommentBubble.cpp b/Engine/Source/Editor/GraphEditor/Private/SCommentBubble.cpp
--- a/Engine/Source/Editor/GraphEditor/Private/SCommentBubble.cpp
+++ b/Engine/Source/Editor/GraphEditor/Private/SCommentBubble.cpp
@@ -387,7 +387,7 @@ FSlateColor SCommentBubble::GetBubbleColor() const
 {
 	FLinearColor ReturnColor = ColorAndOpacity.Get().GetSpecifiedColor();
 
-	if(!GraphNode->IsNodeEnabled() || GraphNode->IsDisplayAsDisabledForced())
+	if(!GraphNode->IsNodeEnabled() || GraphNode->IsDisplayAsDisabledForced() || GraphNode->IsNodeUnrelated())
 	{
 		ReturnColor.A *= 0.6f;
 	}
diff --git a/Engine/Source/Editor/GraphEditor/Private/SGraphNode.cpp b/Engine/Source/Editor/GraphEditor/Private/SGraphNode.cpp
--- a/Engine/Source/Editor/GraphEditor/Private/SGraphNode.cpp
+++ b/Engine/Source/Editor/GraphEditor/Private/SGraphNode.cpp
@@ -563,7 +563,7 @@ FSlateColor SGraphNode::GetNodeTitleColor() const
 {
 	FLinearColor ReturnTitleColor = GraphNode->IsDeprecated() ? FLinearColor::Red : GetNodeObj()->GetNodeTitleColor();
 
-	if(!GraphNode->IsNodeEnabled() || GraphNode->IsDisplayAsDisabledForced())
+	if(!GraphNode->IsNodeEnabled() || GraphNode->IsDisplayAsDisabledForced() || GraphNode->IsNodeUnrelated())
 	{
 		ReturnTitleColor *= FLinearColor(0.5f, 0.5f, 0.5f, 0.4f);
 	}
@@ -577,7 +577,7 @@ FSlateColor SGraphNode::GetNodeTitleColor() const
 FSlateColor SGraphNode::GetNodeBodyColor() const
 {
 	FLinearColor ReturnBodyColor = FLinearColor::White;
-	if(!GraphNode->IsNodeEnabled() || GraphNode->IsDisplayAsDisabledForced())
+	if(!GraphNode->IsNodeEnabled() || GraphNode->IsDisplayAsDisabledForced() || GraphNode->IsNodeUnrelated())
 	{
 		ReturnBodyColor *= FLinearColor(1.0f, 1.0f, 1.0f, 0.5f); 
 	}
@@ -587,7 +587,7 @@ FSlateColor SGraphNode::GetNodeBodyColor() const
 FSlateColor SGraphNode::GetNodeTitleIconColor() const
 {
 	FLinearColor ReturnIconColor = IconColor;
-	if(!GraphNode->IsNodeEnabled() || GraphNode->IsDisplayAsDisabledForced())
+	if(!GraphNode->IsNodeEnabled() || GraphNode->IsDisplayAsDisabledForced() || GraphNode->IsNodeUnrelated())
 	{
 		ReturnIconColor *= FLinearColor(1.0f, 1.0f, 1.0f, 0.3f); 
 	}
@@ -597,7 +597,7 @@ FSlateColor SGraphNode::GetNodeTitleIconColor() const
 FLinearColor SGraphNode::GetNodeTitleTextColor() const
 {
 	FLinearColor ReturnTextColor = FLinearColor::White;
-	if(!GraphNode->IsNodeEnabled() || GraphNode->IsDisplayAsDisabledForced())
+	if(!GraphNode->IsNodeEnabled() || GraphNode->IsDisplayAsDisabledForced() || GraphNode->IsNodeUnrelated())
 	{
 		ReturnTextColor *= FLinearColor(1.0f, 1.0f, 1.0f, 0.3f); 
 	}
diff --git a/Engine/Source/Editor/GraphEditor/Private/SGraphNodeDocumentation.cpp b/Engine/Source/Editor/GraphEditor/Private/SGraphNodeDocumentation.cpp
--- a/Engine/Source/Editor/GraphEditor/Private/SGraphNodeDocumentation.cpp
+++ b/Engine/Source/Editor/GraphEditor/Private/SGraphNodeDocumentation.cpp
@@ -215,8 +215,8 @@ TSharedPtr<SWidget> SGraphNodeDocumentation::CreateDocumentationPage()
 							+SOverlay::Slot()
 							[
 								SNew(SSimpleGradient)
-								.StartColor(GraphNodeDocumentationDefs::PageGradientStartColor) 
-								.EndColor(GraphNodeDocumentationDefs::PageGradientEndColor) 
+								.StartColor(this, &SGraphNodeDocumentation::GetPageGradientStartColor)
+								.EndColor(this, &SGraphNodeDocumentation::GetPageGradientEndColor)
 							]
 							+SOverlay::Slot()
 							[
@@ -355,6 +355,18 @@ void SGraphNodeDocumentation::Tick( const FGeometry& AllottedGeometry, const dou
 	}
 }
 
+FLinearColor SGraphNodeDocumentation::GetPageGradientStartColor() const
+{
+	FLinearColor Color = GraphNodeDocumentationDefs::PageGradientStartColor;
+	return Color;
+}
+
+FLinearColor SGraphNodeDocumentation::GetPageGradientEndColor() const
+{
+	FLinearColor Color = GraphNodeDocumentationDefs::PageGradientEndColor;
+	return Color;
+}
+
 /////////////////////////////////////////////////////
 
 #undef LOCTEXT_NAMESPACE
diff --git a/Engine/Source/Editor/GraphEditor/Private/SGraphPanel.cpp b/Engine/Source/Editor/GraphEditor/Private/SGraphPanel.cpp
--- a/Engine/Source/Editor/GraphEditor/Private/SGraphPanel.cpp
+++ b/Engine/Source/Editor/GraphEditor/Private/SGraphPanel.cpp
@@ -171,6 +171,9 @@ int32 SGraphPanel::OnPaint( const FPaintArgs& Args, const FGeometry& AllottedGeo
 			if (bNodeIsVisible)
 			{
 				const bool bSelected = SelectionToVisualize->Contains( StaticCastSharedRef<SNodePanel::SNode>(CurWidget.Widget)->GetObjectBeingDisplayed() );
+				
+				UEdGraphNode* NodeObj = Cast<UEdGraphNode>(ChildNode->GetObjectBeingDisplayed());
+				float Alpha = 1.0f;
 
 				// Handle Node renaming once the node is visible
 				if( bSelected && ChildNode->IsRenamePending() )
@@ -191,7 +194,9 @@ int32 SGraphPanel::OnPaint( const FPaintArgs& Args, const FGeometry& AllottedGeo
 						OutDrawElements,
 						ShadowLayerId,
 						CurWidget.Geometry.ToInflatedPaintGeometry(NodeShadowSize),
-						ShadowBrush
+						ShadowBrush,
+						ESlateDrawEffect::None,
+						FLinearColor(1.0f, 1.0f, 1.0f, Alpha)
 						);
 				}
 
@@ -215,8 +220,6 @@ int32 SGraphPanel::OnPaint( const FPaintArgs& Args, const FGeometry& AllottedGeo
 
 				int32 CurWidgetsMaxLayerId;
 				{
-					UEdGraphNode* NodeObj = Cast<UEdGraphNode>(ChildNode->GetObjectBeingDisplayed());
-
 					/** When diffing nodes, nodes that are different between revisions are opaque, nodes that have not changed are faded */
 					FGraphDiffControl::FNodeMatch NodeMatch = FGraphDiffControl::FindNodeMatch(GraphObjToDiff, NodeObj, NodeMatches);
 					if (NodeMatch.IsValid())
@@ -228,7 +231,10 @@ int32 SGraphPanel::OnPaint( const FPaintArgs& Args, const FGeometry& AllottedGeo
 					/* When dragging off a pin, we want to duck the alpha of some nodes */
 					TSharedPtr< SGraphPin > OnlyStartPin = (1 == PreviewConnectorFromPins.Num()) ? PreviewConnectorFromPins[0].FindInGraphPanel(*this) : TSharedPtr< SGraphPin >();
 					const bool bNodeIsNotUsableInCurrentContext = Schema->FadeNodeWhenDraggingOffPin(NodeObj, OnlyStartPin.IsValid() ? OnlyStartPin.Get()->GetPinObj() : nullptr);
-					const FWidgetStyle& NodeStyleToUse = (bNodeIsDifferent && !bNodeIsNotUsableInCurrentContext)? InWidgetStyle : FadedStyle;
+					
+					const FWidgetStyle& NodeStyle = (bNodeIsDifferent && !bNodeIsNotUsableInCurrentContext)? InWidgetStyle : FadedStyle;
+					FWidgetStyle NodeStyleToUse = NodeStyle;
+					NodeStyleToUse.BlendColorAndOpacityTint(FLinearColor(1.0f, 1.0f, 1.0f, Alpha));
 
 					// Draw the node.O
 					CurWidgetsMaxLayerId = CurWidget.Widget->Paint(NewArgs, CurWidget.Geometry, MyCullingRect, OutDrawElements, ChildLayerId, NodeStyleToUse, !DisplayAsReadOnly.Get() && ShouldBeEnabled( bParentEnabled ) );
@@ -260,7 +266,9 @@ int32 SGraphPanel::OnPaint( const FPaintArgs& Args, const FGeometry& AllottedGeo
 									OutDrawElements,
 									CurWidgetsMaxLayerId,
 									BouncedGeometry,
-									OverlayBrush
+									OverlayBrush,
+									ESlateDrawEffect::None,
+									FLinearColor(1.0f, 1.0f, 1.0f, Alpha)
 									);
 							}
 
diff --git a/Engine/Source/Editor/GraphEditor/Private/SGraphPin.cpp b/Engine/Source/Editor/GraphEditor/Private/SGraphPin.cpp
--- a/Engine/Source/Editor/GraphEditor/Private/SGraphPin.cpp
+++ b/Engine/Source/Editor/GraphEditor/Private/SGraphPin.cpp
@@ -69,7 +69,6 @@ SGraphPin::SGraphPin()
 	: GraphPinObj(nullptr)
 	, PinColorModifier(FLinearColor::White)
 	, CachedNodeOffset(FVector2D::ZeroVector)
-	, bGraphDataInvalid(false)
 	, bShowLabel(true)
 	, bOnlyShowDefaultValue(false)
 	, bIsMovingLinks(false)
@@ -974,7 +973,7 @@ FSlateColor SGraphPin::GetPinColor() const
 		}
 		if (const UEdGraphSchema* Schema = GraphPinObj->GetSchema())
 		{
-			if (!GetPinObj()->GetOwningNode()->IsNodeEnabled() || GetPinObj()->GetOwningNode()->IsDisplayAsDisabledForced() || !IsEditingEnabled())
+			if (!GetPinObj()->GetOwningNode()->IsNodeEnabled() || GetPinObj()->GetOwningNode()->IsDisplayAsDisabledForced() || !IsEditingEnabled() || GetPinObj()->GetOwningNode()->IsNodeUnrelated())
 			{
 				return Schema->GetPinTypeColor(GraphPinObj->PinType) * FLinearColor(1.0f, 1.0f, 1.0f, 0.5f);
 			}
@@ -997,13 +996,13 @@ FSlateColor SGraphPin::GetPinTextColor() const
 	// If there is no schema there is no owning node (or basically this is a deleted node)
 	if (UEdGraphNode* GraphNode = GraphPinObj->GetOwningNodeUnchecked())
 	{
-		const bool bDisabled = (!GraphNode->IsNodeEnabled() || GraphNode->IsDisplayAsDisabledForced() || !IsEditingEnabled());
+		const bool bDisabled = (!GraphNode->IsNodeEnabled() || GraphNode->IsDisplayAsDisabledForced() || !IsEditingEnabled() || GraphNode->IsNodeUnrelated());
 		if (GraphPinObj->bOrphanedPin)
 		{
 			FLinearColor PinColor = FLinearColor::Red;
 			if (bDisabled)
 			{
-				PinColor.A = 0.5f;
+				PinColor.A = .25;
 			}
 			return PinColor;
 		}
diff --git a/Engine/Source/Editor/GraphEditor/Public/SGraphNodeDocumentation.h b/Engine/Source/Editor/GraphEditor/Public/SGraphNodeDocumentation.h
--- a/Engine/Source/Editor/GraphEditor/Public/SGraphNodeDocumentation.h
+++ b/Engine/Source/Editor/GraphEditor/Public/SGraphNodeDocumentation.h
@@ -52,6 +52,10 @@ protected:
 	float GetDocumentationWrapWidth() const;
 	/** Returns the current child widgets visibility for hit testing */
 	EVisibility GetWidgetVisibility() const;
+	/** Color of the page gradient start */
+	FLinearColor GetPageGradientStartColor() const;
+	/** Color of the page gradient end */
+	FLinearColor GetPageGradientEndColor() const;
 
 private:
 
diff --git a/Engine/Source/Editor/Kismet/Private/BlueprintEditor.cpp b/Engine/Source/Editor/Kismet/Private/BlueprintEditor.cpp
--- a/Engine/Source/Editor/Kismet/Private/BlueprintEditor.cpp
+++ b/Engine/Source/Editor/Kismet/Private/BlueprintEditor.cpp
@@ -154,6 +154,11 @@
 #include "Widgets/Notifications/SNotificationList.h"
 #include "NativeCodeGenerationTool.h"
 
+// Focusing related nodes feature
+#include "Preferences/BlueprintEditorOptions.h"
+#include "Framework/MultiBox/MultiBoxBuilder.h"
+#include "Widgets/Input/SNumericEntryBox.h"
+
 #define LOCTEXT_NAMESPACE "BlueprintEditor"
 
 /////////////////////////////////////////////////////
@@ -1467,7 +1472,8 @@ void FBlueprintEditor::OnChangeBreadCrumbGraph(UEdGraph* InGraph)
 }
 
 FBlueprintEditor::FBlueprintEditor()
-	: bSaveIntermediateBuildProducts(false)
+	: EditorOptions(nullptr)
+	, bSaveIntermediateBuildProducts(false)
 	, bPendingDeferredClose(false)
 	, bRequestedSavingOpenDocumentState(false)
 	, bBlueprintModifiedOnOpen (false)
@@ -1475,6 +1481,9 @@ FBlueprintEditor::FBlueprintEditor()
 	, bIsActionMenuContextSensitive(true)
 	, CurrentUISelection(NAME_None)
 	, bEditorMarkedAsClosed(false)
+	, bHideUnrelatedNodes(false)
+	, bLockNodeFadeState(false)
+	, bSelectRegularNode(false)
 	, HasOpenActionMenu(nullptr)
 	, InstructionsFadeCountdown(0.f)
 {
@@ -1489,7 +1498,7 @@ FBlueprintEditor::FBlueprintEditor()
 	AnalyticsStats.NodePasteCreateCount = 0;
 
 	UEditorEngine* Editor = (UEditorEngine*)GEngine;
-	if (Editor != NULL)
+	if (Editor != nullptr)
 	{
 		Editor->RegisterForUndo(this);
 	}
@@ -1742,6 +1751,11 @@ void FBlueprintEditor::InitBlueprintEditor(
 	// TRUE if a single Blueprint is being opened and is marked as newly created
 	bool bNewlyCreated = InBlueprints.Num() == 1 && InBlueprints[0]->bIsNewlyCreated;
 
+	EditorOptions = nullptr;
+
+	// Load editor settings from disk.
+	LoadEditorSettings();
+
 	TArray< UObject* > Objects;
 	for (UBlueprint* Blueprint : InBlueprints)
 	{
@@ -1769,6 +1783,42 @@ void FBlueprintEditor::InitBlueprintEditor(
 
 	InitalizeExtenders();
 
+	struct Local
+	{
+		static void FillToolbar(FToolBarBuilder& ToolbarBuilder, const TSharedRef< FUICommandList > ToolkitCommands, FBlueprintEditor* BlueprintEditor)
+		{
+			ToolbarBuilder.BeginSection("Graph");
+			{
+				ToolbarBuilder.AddToolBarButton(
+					FBlueprintEditorCommands::Get().ToggleHideUnrelatedNodes,
+					NAME_None,
+					TAttribute<FText>(),
+					TAttribute<FText>(),
+					FSlateIcon(FEditorStyle::GetStyleSetName(), "GraphEditor.ToggleHideUnrelatedNodes")
+				);
+				ToolbarBuilder.AddComboButton(
+					FUIAction(),
+					FOnGetContent::CreateSP(BlueprintEditor, &FBlueprintEditor::MakeHideUnrelatedNodesOptionsMenu),
+					LOCTEXT("HideUnrelatedNodesOptions", "Focus Related Nodes Options"),
+					LOCTEXT("HideUnrelatedNodesOptionsMenu", "Focus Related Nodes options menu"),
+					TAttribute<FSlateIcon>(),
+					true
+				);
+			}
+			ToolbarBuilder.EndSection();
+		}
+	};
+
+	TSharedPtr<FExtender> ToolbarExtender = MakeShareable(new FExtender);
+	ToolbarExtender->AddToolBarExtension(
+		"Asset",
+		EExtensionHook::After,
+		GetToolkitCommands(),
+		FToolBarExtensionDelegate::CreateStatic( &Local::FillToolbar, GetToolkitCommands(), this )
+	);
+
+	AddToolbarExtender(ToolbarExtender);
+
 	RegenerateMenusAndToolbars();
 
 	RegisterApplicationModes(InBlueprints, bShouldOpenInDefaultsMode, bNewlyCreated);
@@ -2243,6 +2293,8 @@ FBlueprintEditor::~FBlueprintEditor()
 			FEngineAnalytics::GetProvider().RecordEvent( FString( "Editor.Usage.BPDisallowedPinConnection" ), BPEditorPinConnectAttribs );
 		}
 	}
+
+	SaveEditorSettings();
 }
 
 void FBlueprintEditor::FocusInspectorOnGraphSelection(const FGraphPanelSelectionSet& NewSelection, bool bForceRefresh)
@@ -2653,6 +2705,13 @@ void FBlueprintEditor::CreateDefaultCommands()
 		FGraphEditorCommands::Get().ClearAllQuickJumps,
 		FExecuteAction::CreateSP(this, &FBlueprintEditor::ClearAllGraphEditorQuickJumps)
 	);
+
+	ToolkitCommands->MapAction(
+		FBlueprintEditorCommands::Get().ToggleHideUnrelatedNodes,
+		FExecuteAction::CreateSP(this, &FBlueprintEditor::ToggleHideUnrelatedNodes),
+		FCanExecuteAction(),
+		FIsActionChecked::CreateSP(this, &FBlueprintEditor::IsToggleHideUnrelatedNodesChecked)
+	);
 }
 
 void FBlueprintEditor::OpenNativeCodeGenerationTool()
@@ -3029,6 +3088,11 @@ void FBlueprintEditor::OnGraphEditorFocused(const TSharedRef<SGraphEditor>& InGr
 		}
 	}
 
+	if (bHideUnrelatedNodes && SelectedNodes.Num() <= 0)
+	{
+		ResetAllNodesUnrelatedStates();
+	}
+
 	// If the bookmarks view is active, check whether or not we're restricting the view to the current graph. If we are, update the tree to reflect the focused graph context.
 	if (BookmarksWidget.IsValid()
 		&& GetDefault<UBlueprintEditorSettings>()->bShowBookmarksForCurrentDocumentOnlyInTab)
@@ -3155,6 +3219,28 @@ void FBlueprintEditor::OnSelectedNodesChangedImpl(const FGraphPanelSelectionSet&
 	}
 
 	Inspector->ShowDetailsForObjects(NewSelection.Array());
+
+
+	bSelectRegularNode = false;
+	for (FGraphPanelSelectionSet::TConstIterator It(NewSelection); It; ++It)
+	{
+		UEdGraphNode_Comment* SeqNode = Cast<UEdGraphNode_Comment>(*It);
+		if (!SeqNode)
+		{
+			bSelectRegularNode = true;
+			break;
+		}
+	}
+
+	if (bHideUnrelatedNodes && !bLockNodeFadeState)
+	{
+		ResetAllNodesUnrelatedStates();
+
+		if ( bSelectRegularNode )
+		{
+			HideUnrelatedNodes();
+		}
+	}
 }
 
 void FBlueprintEditor::OnBlueprintChangedImpl(UBlueprint* InBlueprint, bool bIsJustBeingCompiled )
@@ -3580,6 +3666,8 @@ void FBlueprintEditor::AddReferencedObjects( FReferenceCollector& Collector )
 		}
 	}
 
+	Collector.AddReferencedObject(EditorOptions);
+
 	UserDefinedStructures.Remove(TWeakObjectPtr<UUserDefinedStruct>()); // Remove NULLs
 	for (const TWeakObjectPtr<UUserDefinedStruct>& ObjectPtr : UserDefinedStructures)
 	{
@@ -6249,6 +6337,301 @@ ECheckBoxState FBlueprintEditor::CheckEnabledStateForSelectedNodes(ENodeEnabledS
 	return Result;
 }
 
+void FBlueprintEditor::UpdateNodesUnrelatedStatesAfterGraphChange()
+{
+	if (bHideUnrelatedNodes && !bLockNodeFadeState && bSelectRegularNode)
+	{
+		ResetAllNodesUnrelatedStates();
+
+		HideUnrelatedNodes();
+	}
+}
+
+void FBlueprintEditor::ResetAllNodesUnrelatedStates()
+{
+	TSharedPtr<SGraphEditor> FocusedGraphEd = FocusedGraphEdPtr.Pin();
+
+	if (FocusedGraphEd.IsValid())
+	{
+		FocusedGraphEd->ResetAllNodesUnrelatedStates();
+	}
+}
+
+void FBlueprintEditor::CollectExecDownstreamNodes(UEdGraphNode* CurrentNode, TArray<UEdGraphNode*>& CollectedNodes)
+{
+	const UEdGraphSchema_K2* K2Schema = GetDefault<UEdGraphSchema_K2>();
+
+	TArray<UEdGraphPin*> AllPins = CurrentNode->GetAllPins();
+
+	for (auto& Pin : AllPins)
+	{
+		if (Pin->Direction == EGPD_Output && Pin->PinType.PinCategory == K2Schema->PC_Exec)
+		{
+			for (auto& Link : Pin->LinkedTo)
+			{
+				UEdGraphNode* LinkedNode = Cast<UEdGraphNode>(Link->GetOwningNode());
+				if (LinkedNode && !CollectedNodes.Contains(LinkedNode))
+				{
+					CollectedNodes.Add(LinkedNode);
+					CollectExecDownstreamNodes( LinkedNode, CollectedNodes );
+				}
+			}
+		}
+	}
+}
+
+void FBlueprintEditor::CollectExecUpstreamNodes(UEdGraphNode* CurrentNode, TArray<UEdGraphNode*>& CollectedNodes)
+{
+	const UEdGraphSchema_K2* K2Schema = GetDefault<UEdGraphSchema_K2>();
+
+	TArray<UEdGraphPin*> AllPins = CurrentNode->GetAllPins();
+
+	for (auto& Pin : AllPins)
+	{
+		if (Pin->Direction == EGPD_Input && Pin->PinType.PinCategory == K2Schema->PC_Exec)
+		{
+			for (auto& Link : Pin->LinkedTo)
+			{
+				UEdGraphNode* LinkedNode = Cast<UEdGraphNode>(Link->GetOwningNode());
+				if (LinkedNode && !CollectedNodes.Contains(LinkedNode))
+				{
+					CollectedNodes.Add(LinkedNode);
+					CollectExecUpstreamNodes( LinkedNode, CollectedNodes );
+				}
+			}
+		}
+	}
+}
+
+void FBlueprintEditor::CollectPureDownstreamNodes(UEdGraphNode* CurrentNode, TArray<UEdGraphNode*>& CollectedNodes)
+{
+	const UEdGraphSchema_K2* K2Schema = GetDefault<UEdGraphSchema_K2>();
+
+	TArray<UEdGraphPin*> AllPins = CurrentNode->GetAllPins();
+
+	for (auto& Pin : AllPins)
+	{
+		if (Pin->Direction == EGPD_Output && Pin->PinType.PinCategory != K2Schema->PC_Exec)
+		{
+			for (auto& Link : Pin->LinkedTo)
+			{
+				UK2Node* LinkedNode = Cast<UK2Node>(Link->GetOwningNode());
+				if (LinkedNode && !CollectedNodes.Contains(LinkedNode))
+				{
+					CollectedNodes.Add(LinkedNode);
+					if (LinkedNode->IsNodePure())
+					{
+						CollectPureDownstreamNodes( LinkedNode, CollectedNodes );
+					}
+				}
+			}
+		}
+	}
+}
+
+void FBlueprintEditor::CollectPureUpstreamNodes(UEdGraphNode* CurrentNode, TArray<UEdGraphNode*>& CollectedNodes)
+{
+	const UEdGraphSchema_K2* K2Schema = GetDefault<UEdGraphSchema_K2>();
+
+	TArray<UEdGraphPin*> AllPins = CurrentNode->GetAllPins();
+
+	for (auto& Pin : AllPins)
+	{
+		if (Pin->Direction == EGPD_Input && Pin->PinType.PinCategory != K2Schema->PC_Exec)
+		{
+			for (auto& Link : Pin->LinkedTo)
+			{
+				UK2Node* LinkedNode = Cast<UK2Node>(Link->GetOwningNode());
+				if (LinkedNode && !CollectedNodes.Contains(LinkedNode))
+				{
+					CollectedNodes.Add(LinkedNode);
+					if (LinkedNode->IsNodePure())
+					{
+						CollectPureUpstreamNodes( LinkedNode, CollectedNodes );
+					}
+				}
+			}
+		}
+	}
+}
+
+void FBlueprintEditor::HideUnrelatedNodes()
+{
+	TArray<UEdGraphNode*> NodesToShow;
+
+	TSharedPtr<SGraphEditor> FocusedGraphEd = FocusedGraphEdPtr.Pin();
+	if (FocusedGraphEd.IsValid())
+	{
+		FGraphPanelSelectionSet SelectedNodes = GetSelectedNodes();
+
+		TArray<UObject*> ImpureNodes = SelectedNodes.Array().FilterByPredicate([](UObject* Node){
+			UK2Node* K2Node = Cast<UK2Node>(Node);
+			if (K2Node)
+			{
+				return !(K2Node->IsNodePure());
+			}
+			return false;
+		});
+
+		TArray<UObject*> PureNodes = SelectedNodes.Array().FilterByPredicate([](UObject* Node){
+			UK2Node* K2Node = Cast<UK2Node>(Node);
+			if (K2Node)
+			{
+				return K2Node->IsNodePure();
+			}
+			// Treat a node which can't cast to an UK2Node as a pure node (like a document node or a commment node)
+			// Make sure all selected nodes are handled
+			return true;
+		});
+
+		for (auto Node : ImpureNodes)
+		{
+			UEdGraphNode* SelectedNode = Cast<UEdGraphNode>(Node);
+
+			if (SelectedNode)
+			{
+				NodesToShow.Add(SelectedNode);
+				CollectExecDownstreamNodes( SelectedNode, NodesToShow );
+				CollectExecUpstreamNodes( SelectedNode, NodesToShow );
+				CollectPureDownstreamNodes( SelectedNode, NodesToShow );
+				CollectPureUpstreamNodes( SelectedNode, NodesToShow );
+			}
+		}
+
+		for (auto Node : PureNodes)
+		{
+			UEdGraphNode* SelectedNode = Cast<UEdGraphNode>(Node);
+
+			if (SelectedNode)
+			{
+				NodesToShow.Add(SelectedNode);
+				CollectPureDownstreamNodes( SelectedNode, NodesToShow );
+				CollectPureUpstreamNodes( SelectedNode, NodesToShow );
+			}
+		}
+
+		TArray<class UEdGraphNode*> AllNodes = FocusedGraphEd->GetCurrentGraph()->Nodes;
+
+		TArray<UEdGraphNode*> CommentNodes;
+		TArray<UEdGraphNode*> RelatedNodes;
+
+		for (auto& Node : AllNodes)
+		{
+			if (NodesToShow.Contains(Cast<UEdGraphNode>(Node)))
+			{
+				Node->SetNodeUnrelated(false);
+				RelatedNodes.Add(Node);
+			}
+			else
+			{
+				if (UEdGraphNode_Comment* CommentNode = Cast<UEdGraphNode_Comment>(Node))
+				{
+					CommentNodes.Add(Node);
+				}
+				else
+				{
+					Node->SetNodeUnrelated(true);
+				}
+			}
+		}
+
+		if (FocusedGraphEd.IsValid())
+		{
+			FocusedGraphEd->FocusCommentNodes(CommentNodes, RelatedNodes);
+		}
+	}
+}
+
+void FBlueprintEditor::ToggleHideUnrelatedNodes()
+{
+	bHideUnrelatedNodes = !bHideUnrelatedNodes;
+
+	ResetAllNodesUnrelatedStates();
+
+	if (bHideUnrelatedNodes && bSelectRegularNode)
+	{
+		HideUnrelatedNodes();
+	}
+	else
+	{
+		bLockNodeFadeState = false;
+	}
+}
+
+bool FBlueprintEditor::IsToggleHideUnrelatedNodesChecked() const
+{
+	return bHideUnrelatedNodes == true;
+}
+
+TSharedRef<SWidget> FBlueprintEditor::MakeHideUnrelatedNodesOptionsMenu()
+{
+	const bool bShouldCloseWindowAfterMenuSelection = true;
+	FMenuBuilder MenuBuilder( bShouldCloseWindowAfterMenuSelection, GetToolkitCommands() );
+
+	TSharedRef<SWidget> OptionsHeading = SNew(SBox)
+		.Padding(2.0f)
+		[
+			SNew(SHorizontalBox)
+
+			+SHorizontalBox::Slot()
+			[
+				SNew(STextBlock)
+					.Text(LOCTEXT("FocusRelatedOptions", "Focus Related Options"))
+					.TextStyle(FEditorStyle::Get(), "Menu.Heading")
+			]
+		];
+
+	TSharedRef<SWidget> LockNodeStateCheckBox = SNew(SBox)
+		[
+			SNew(SCheckBox)
+				.IsChecked(bLockNodeFadeState ? ECheckBoxState::Checked : ECheckBoxState::Unchecked)
+				.OnCheckStateChanged(this, &FBlueprintEditor::OnLockNodeStateCheckStateChanged)
+				.Style(FEditorStyle::Get(), "Menu.CheckBox")
+				.ToolTipText(LOCTEXT("LockNodeStateCheckBoxToolTip", "Lock the current state of all nodes."))
+				.Content()
+				[
+					SNew(SHorizontalBox)
+
+					+SHorizontalBox::Slot()
+					.Padding(2.0f, 0.0f, 0.0f, 0.0f)
+					[
+						SNew(STextBlock)
+							.Text(LOCTEXT("LockNodeState", "Lock Node State"))
+					]
+				]
+		];
+
+	MenuBuilder.AddWidget(OptionsHeading, FText::GetEmpty(), true);
+
+	MenuBuilder.AddMenuEntry(FUIAction(), LockNodeStateCheckBox);
+
+	return MenuBuilder.MakeWidget();
+}
+
+void FBlueprintEditor::LoadEditorSettings()
+{
+	EditorOptions = NewObject<UBlueprintEditorOptions>();
+
+	if (EditorOptions->bHideUnrelatedNodes)
+	{
+		ToggleHideUnrelatedNodes();
+	}
+}
+
+void FBlueprintEditor::SaveEditorSettings()
+{
+	if ( EditorOptions )
+	{
+		EditorOptions->bHideUnrelatedNodes     = bHideUnrelatedNodes;
+		EditorOptions->SaveConfig();
+	}
+}
+
+void FBlueprintEditor::OnLockNodeStateCheckStateChanged(ECheckBoxState NewCheckedState)
+{
+	bLockNodeFadeState = (NewCheckedState == ECheckBoxState::Checked) ? true : false;
+}
+
 void FBlueprintEditor::ToggleSaveIntermediateBuildProducts()
 {
 	bSaveIntermediateBuildProducts = !bSaveIntermediateBuildProducts;
diff --git a/Engine/Source/Editor/Kismet/Private/BlueprintEditorCommands.cpp b/Engine/Source/Editor/Kismet/Private/BlueprintEditorCommands.cpp
--- a/Engine/Source/Editor/Kismet/Private/BlueprintEditorCommands.cpp
+++ b/Engine/Source/Editor/Kismet/Private/BlueprintEditorCommands.cpp
@@ -74,6 +74,9 @@ void FBlueprintEditorCommands::RegisterCommands()
 	
 	// SCC commands
 	UI_COMMAND( BeginBlueprintMerge, "Merge", "Shows the Blueprint merge panel and toolbar, allowing the user to resolve conflicted blueprints", EUserInterfaceActionType::Button, FInputChord() );
+	
+	// Hide unrelated nodes
+	UI_COMMAND( ToggleHideUnrelatedNodes, "Hide Unrelated", "Toggles  automatically hiding nodes which are unrelated to the selected nodes.", EUserInterfaceActionType::ToggleButton, FInputChord() );
 }
 
 PRAGMA_ENABLE_OPTIMIZATION
diff --git a/Engine/Source/Editor/Kismet/Private/BlueprintEditorCommands.h b/Engine/Source/Editor/Kismet/Private/BlueprintEditorCommands.h
--- a/Engine/Source/Editor/Kismet/Private/BlueprintEditorCommands.h
+++ b/Engine/Source/Editor/Kismet/Private/BlueprintEditorCommands.h
@@ -75,6 +75,9 @@ public:
 
 	// SSC commands
 	TSharedPtr< FUICommandInfo > BeginBlueprintMerge;
+	
+	// Toggle focusing nodes which are related to the selected nodes
+	TSharedPtr< FUICommandInfo > ToggleHideUnrelatedNodes;
 };
 
 //////////////////////////////////////////////////////////////////////////
diff --git a/Engine/Source/Editor/Kismet/Public/BlueprintEditor.h b/Engine/Source/Editor/Kismet/Public/BlueprintEditor.h
--- a/Engine/Source/Editor/Kismet/Public/BlueprintEditor.h
+++ b/Engine/Source/Editor/Kismet/Public/BlueprintEditor.h
@@ -43,6 +43,7 @@ class UEdGraph;
 class UEdGraphNode;
 class UUserDefinedEnum;
 class UUserDefinedStruct;
+class UBlueprintEditorOptions;
 struct Rect;
 
 /////////////////////////////////////////////////////
@@ -1057,6 +1058,19 @@ private:
 	/** Returns the appropriate check box state representing whether or not the selected nodes are enabled */
 	ECheckBoxState GetEnabledCheckBoxStateForSelectedNodes();
 
+	/** Configuration class used to store editor settings across sessions. */
+	UBlueprintEditorOptions* EditorOptions;
+
+	/**
+	 * Load editor settings from disk (docking state, window pos/size, option state, etc).
+	 */
+	virtual void LoadEditorSettings();
+
+	/**
+	 * Saves editor settings to disk (docking state, window pos/size, option state, etc).
+	 */
+	virtual void SaveEditorSettings();
+
 	/** Attempt to match the given enabled state for currently-selected nodes */
 	ECheckBoxState CheckEnabledStateForSelectedNodes(ENodeEnabledState CheckState);
 
@@ -1065,6 +1079,9 @@ private:
 
 public://@TODO
 	TSharedPtr<FDocumentTracker> DocumentManager;
+	
+	/** Update all nodes' unrelated states when the graph has changed */
+	void UpdateNodesUnrelatedStatesAfterGraphChange();
 
 protected:
 
@@ -1166,6 +1183,35 @@ protected:
 	/** The preview actor representing the current preview */
 	mutable TWeakObjectPtr<AActor> PreviewActorPtr;
 
+	/** If true, fade out nodes which are unrelated to the selected nodes automatically. */
+	bool bHideUnrelatedNodes;
+
+	/** Lock the current fade state of each node */
+	bool bLockNodeFadeState;
+
+	/** If a regular node (not a comment node) has been selected */
+	bool bSelectRegularNode;
+
+	/** Focus nodes which are related to the selected nodes */
+	void ResetAllNodesUnrelatedStates();
+	void CollectExecUpstreamNodes(UEdGraphNode* CurrentNode, TArray<UEdGraphNode*>& CollectedNodes);
+	void CollectExecDownstreamNodes(UEdGraphNode* CurrentNode, TArray<UEdGraphNode*>& CollectedNodes);
+	void CollectPureDownstreamNodes(UEdGraphNode* CurrentNode, TArray<UEdGraphNode*>& CollectedNodes);
+	void CollectPureUpstreamNodes(UEdGraphNode* CurrentNode, TArray<UEdGraphNode*>& CollectedNodes);
+	void HideUnrelatedNodes();
+
+public:
+	/** Make nodes which are unrelated to the selected nodes fade out */
+	void ToggleHideUnrelatedNodes();
+	bool IsToggleHideUnrelatedNodesChecked() const;
+
+	/** Make a drop down menu to control the opacity of unrelated nodes */
+	TSharedRef<SWidget> MakeHideUnrelatedNodesOptionsMenu();
+	TOptional<float> HandleUnrelatedNodesOpacityBoxValue() const;
+	void HandleUnrelatedNodesOpacityBoxChanged(float NewOpacity);
+	void OnLockNodeStateCheckStateChanged(ECheckBoxState NewCheckedState);
+
+
 public:
 	//@TODO: To be moved/merged
 	TSharedPtr<SDockTab> OpenDocument(const UObject* DocumentID, FDocumentTracker::EOpenDocumentCause Cause);
diff --git a/Engine/Source/Editor/MaterialEditor/Private/MaterialEditor.cpp b/Engine/Source/Editor/MaterialEditor/Private/MaterialEditor.cpp
--- a/Engine/Source/Editor/MaterialEditor/Private/MaterialEditor.cpp
+++ b/Engine/Source/Editor/MaterialEditor/Private/MaterialEditor.cpp
@@ -11,6 +11,8 @@
 #include "Widgets/Layout/SSeparator.h"
 #include "Widgets/Layout/SBox.h"
 #include "Widgets/Input/SButton.h"
+#include "Widgets/Input/SNumericEntryBox.h"
+#include "Widgets/Input/SCheckBox.h"
 #include "EditorStyleSet.h"
 #include "EdGraph/EdGraph.h"
 #include "MaterialGraph/MaterialGraph.h"
@@ -121,6 +123,8 @@
 #include "Materials/MaterialExpressionMaterialLayerOutput.h"
 
 #include "MaterialStats.h"
+#include "Materials/MaterialExpression.h"
+
 #include "SMaterialParametersOverviewWidget.h"
 #include "IPropertyRowGenerator.h"
 
@@ -719,23 +723,27 @@ void FMaterialEditor::InitMaterialEditor( const EToolkitMode::Type Mode, const T
 FMaterialEditor::FMaterialEditor()
 	: bMaterialDirty(false)
 	, bStatsFromPreviewMaterial(false)
-	, Material(NULL)
-	, OriginalMaterial(NULL)
-	, ExpressionPreviewMaterial(NULL)
-	, EmptyMaterial(NULL)
-	, PreviewExpression(NULL)
-	, MaterialFunction(NULL)
-	, OriginalMaterialObject(NULL)
-	, EditorOptions(NULL)
-	, ScopedTransaction(NULL)
+	, Material(nullptr)
+	, OriginalMaterial(nullptr)
+	, ExpressionPreviewMaterial(nullptr)
+	, EmptyMaterial(nullptr)
+	, PreviewExpression(nullptr)
+	, MaterialFunction(nullptr)
+	, OriginalMaterialObject(nullptr)
+	, EditorOptions(nullptr)
+	, ScopedTransaction(nullptr)
 	, bAlwaysRefreshAllPreviews(false)
 	, bHideUnusedConnectors(false)
 	, bLivePreview(true)
 	, bIsRealtime(false)
 	, bShowBuiltinStats(false)
+	, bHideUnrelatedNodes(false)
+	, bLockNodeFadeState(false)
+	, bFocusWholeChain(false)
+	, bSelectRegularNode(false)
 	, MenuExtensibilityManager(new FExtensibilityManager)
 	, ToolBarExtensibilityManager(new FExtensibilityManager)
-	, MaterialEditorInstance(NULL)
+	, MaterialEditorInstance(nullptr)
 {
 }
 
@@ -1122,6 +1130,21 @@ void FMaterialEditor::FillToolbar(FToolBarBuilder& ToolbarBuilder)
 		ToolbarBuilder.AddToolBarButton(FMaterialEditorCommands::Get().ToggleLivePreview);
 		ToolbarBuilder.AddToolBarButton(FMaterialEditorCommands::Get().ToggleRealtimeExpressions);
 		ToolbarBuilder.AddToolBarButton(FMaterialEditorCommands::Get().AlwaysRefreshAllPreviews);
+		ToolbarBuilder.AddToolBarButton(
+			FMaterialEditorCommands::Get().ToggleHideUnrelatedNodes,
+				NAME_None, 
+				TAttribute<FText>(), 
+				TAttribute<FText>(), 
+				FSlateIcon(FEditorStyle::GetStyleSetName(), "GraphEditor.ToggleHideUnrelatedNodes")
+			);
+		ToolbarBuilder.AddComboButton(
+			FUIAction(),
+			FOnGetContent::CreateSP(this, &FMaterialEditor::MakeHideUnrelatedNodesOptionsMenu),
+			LOCTEXT("HideUnrelatedNodesOptions", "Hide Unrelated Nodes Options"),
+			LOCTEXT("HideUnrelatedNodesOptionsMenu", "Hide Unrelated Nodes options menu"),
+			TAttribute<FSlateIcon>(),
+			true
+		); 
 	}
 	ToolbarBuilder.EndSection();
 
@@ -1506,6 +1529,11 @@ void FMaterialEditor::LoadEditorSettings()
 			PreviewViewport->OnToggleRealtime();
 		}
 	}
+	
+	if (EditorOptions->bHideUnrelatedNodes)
+	{
+		ToggleHideUnrelatedNodes();
+	}
 
 	// Primitive type
 	int32 PrimType;
@@ -1527,7 +1555,8 @@ void FMaterialEditor::SaveEditorSettings()
 		EditorOptions->bHideUnusedConnectors		= !IsOnShowConnectorsChecked();
 		EditorOptions->bAlwaysRefreshAllPreviews	= IsOnAlwaysRefreshAllPreviews();
 		EditorOptions->bRealtimeExpressionViewport	= IsToggleRealTimeExpressionsChecked();
-		EditorOptions->bLivePreviewUpdate		= IsToggleLivePreviewChecked();
+		EditorOptions->bLivePreviewUpdate           = IsToggleLivePreviewChecked();
+		EditorOptions->bHideUnrelatedNodes          = bHideUnrelatedNodes;
 		EditorOptions->SaveConfig();
 	}
 
@@ -2232,6 +2261,12 @@ void FMaterialEditor::BindCommands()
 		FCanExecuteAction(),
 		FIsActionChecked::CreateSP(this, &FMaterialEditor::IsToggleLivePreviewChecked));
 
+	ToolkitCommands->MapAction(
+		Commands.ToggleHideUnrelatedNodes,
+		FExecuteAction::CreateSP(this, &FMaterialEditor::ToggleHideUnrelatedNodes),
+		FCanExecuteAction(),
+		FIsActionChecked::CreateSP(this, &FMaterialEditor::IsToggleHideUnrelatedNodesChecked));
+
 	ToolkitCommands->MapAction(
 		Commands.ToggleRealtimeExpressions,
 		FExecuteAction::CreateSP(this, &FMaterialEditor::ToggleRealTimeExpressions),
@@ -2421,6 +2456,203 @@ bool FMaterialEditor::IsOnAlwaysRefreshAllPreviews() const
 	return bAlwaysRefreshAllPreviews == true;
 }
 
+void FMaterialEditor::ToggleHideUnrelatedNodes()
+{
+	bHideUnrelatedNodes = !bHideUnrelatedNodes;
+
+	GraphEditor->ResetAllNodesUnrelatedStates();
+
+	if (bHideUnrelatedNodes && bSelectRegularNode)
+	{
+		HideUnrelatedNodes();
+	}
+	else
+	{
+		bLockNodeFadeState = false;
+		bFocusWholeChain = false;
+	}
+}
+
+bool FMaterialEditor::IsToggleHideUnrelatedNodesChecked() const
+{
+	return bHideUnrelatedNodes == true;
+}
+
+void FMaterialEditor::CollectDownstreamNodes(UMaterialGraphNode* CurrentNode, TArray<UMaterialGraphNode*>& CollectedNodes)
+{
+	TArray<UEdGraphPin*> OutputPins;
+	CurrentNode->GetOutputPins(OutputPins);
+
+	for (auto& OutputPin : OutputPins)
+	{
+		for (auto& Link : OutputPin->LinkedTo)
+		{
+			UMaterialGraphNode* LinkedNode = Cast<UMaterialGraphNode>(Link->GetOwningNode());
+			if (LinkedNode && !CollectedNodes.Contains(LinkedNode))
+			{
+				CollectedNodes.Add(LinkedNode);
+				CollectDownstreamNodes( LinkedNode, CollectedNodes );
+
+				if (bFocusWholeChain)
+				{
+					CollectUpstreamNodes( LinkedNode, CollectedNodes );
+				}
+			}
+		}
+	}
+}
+
+void FMaterialEditor::CollectUpstreamNodes(UMaterialGraphNode* CurrentNode, TArray<UMaterialGraphNode*>& CollectedNodes)
+{
+	TArray<UEdGraphPin*> InputPins;
+	CurrentNode->GetInputPins(InputPins);
+
+	for (auto& InputPin : InputPins)
+	{
+		for (auto& Link : InputPin->LinkedTo)
+		{
+			UMaterialGraphNode* LinkedNode = Cast<UMaterialGraphNode>(Link->GetOwningNode());
+			if (LinkedNode && !CollectedNodes.Contains(LinkedNode))
+			{
+				CollectedNodes.Add(LinkedNode);
+				CollectUpstreamNodes( LinkedNode, CollectedNodes );
+			}
+		}
+	}
+}
+
+void FMaterialEditor::HideUnrelatedNodes()
+{
+	TArray<UMaterialGraphNode*> NodesToShow;
+
+	const FGraphPanelSelectionSet SelectedNodes = GraphEditor->GetSelectedNodes();
+
+	for (FGraphPanelSelectionSet::TConstIterator NodeIt(SelectedNodes); NodeIt; ++NodeIt)
+	{
+		UMaterialGraphNode* SelectedNode = Cast<UMaterialGraphNode>(*NodeIt);
+
+		if (SelectedNode)
+		{
+			NodesToShow.Add(SelectedNode);
+			CollectDownstreamNodes( SelectedNode, NodesToShow );
+			CollectUpstreamNodes( SelectedNode, NodesToShow );
+		}
+	}
+
+	TArray<class UEdGraphNode*> AllNodes = GraphEditor->GetCurrentGraph()->Nodes;
+
+	TArray<UEdGraphNode*> CommentNodes;
+	TArray<UEdGraphNode*> RelatedNodes;
+
+	for (auto& Node : AllNodes)
+	{
+		// Always draw the root graph node which can't cast to UMaterialGraphNode
+		if (UMaterialGraphNode* GraphNode = Cast<UMaterialGraphNode>(Node))
+		{
+			if (NodesToShow.Contains(GraphNode))
+			{
+				Node->SetNodeUnrelated(false);
+				RelatedNodes.Add(Node);
+			}
+			else
+			{
+				Node->SetNodeUnrelated(true);
+			}
+		}
+		else if (UMaterialGraphNode_Comment* CommentNode = Cast<UMaterialGraphNode_Comment>(Node))
+		{
+			CommentNodes.Add(Node);
+		}
+	}
+
+	GraphEditor->FocusCommentNodes(CommentNodes, RelatedNodes);
+}
+
+TSharedRef<SWidget> FMaterialEditor::MakeHideUnrelatedNodesOptionsMenu()
+{
+	const bool bShouldCloseWindowAfterMenuSelection = true;
+	FMenuBuilder MenuBuilder( bShouldCloseWindowAfterMenuSelection, GetToolkitCommands() );
+
+	TSharedRef<SWidget> OptionsHeading = SNew(SBox)
+		.Padding(2.0f)
+		[
+			SNew(SHorizontalBox)
+
+			+SHorizontalBox::Slot()
+			[
+				SNew(STextBlock)
+					.Text(LOCTEXT("HideUnrelatedOptions", "Hide Unrelated Options"))
+					.TextStyle(FEditorStyle::Get(), "Menu.Heading")
+			]
+		];
+
+	TSharedRef<SWidget> LockNodeStateCheckBox = SNew(SBox)
+		[
+			SNew(SCheckBox)
+				.IsChecked(bLockNodeFadeState ? ECheckBoxState::Checked : ECheckBoxState::Unchecked)
+				.OnCheckStateChanged(this, &FMaterialEditor::OnLockNodeStateCheckStateChanged)
+				.Style(FEditorStyle::Get(), "Menu.CheckBox")
+				.ToolTipText(LOCTEXT("LockNodeStateCheckBoxToolTip", "Lock the current state of all nodes."))
+				.Content()
+				[
+					SNew(SHorizontalBox)
+
+					+SHorizontalBox::Slot()
+					.Padding(2.0f, 0.0f, 0.0f, 0.0f)
+					[
+						SNew(STextBlock)
+							.Text(LOCTEXT("LockNodeState", "Lock Node State"))
+					]
+				]
+		];
+
+	TSharedRef<SWidget> FocusWholeChainCheckBox = SNew(SBox)
+		[
+			SNew(SCheckBox)
+				.IsChecked(bFocusWholeChain ? ECheckBoxState::Checked : ECheckBoxState::Unchecked)
+				.OnCheckStateChanged(this, &FMaterialEditor::OnFocusWholeChainCheckStateChanged)
+				.Style(FEditorStyle::Get(), "Menu.CheckBox")
+				.ToolTipText(LOCTEXT("FocusWholeChainCheckBoxToolTip", "Focus all nodes in the chain."))
+				.Content()
+				[
+					SNew(SHorizontalBox)
+
+					+SHorizontalBox::Slot()
+					.Padding(2.0f, 0.0f, 0.0f, 0.0f)
+					[
+						SNew(STextBlock)
+							.Text(LOCTEXT("FocusWholeChain", "Focus Whole Chain"))
+					]
+				]
+		];
+
+	MenuBuilder.AddWidget(OptionsHeading, FText::GetEmpty(), true);
+
+	MenuBuilder.AddMenuEntry(FUIAction(), LockNodeStateCheckBox);
+
+	MenuBuilder.AddMenuEntry(FUIAction(), FocusWholeChainCheckBox);
+
+	return MenuBuilder.MakeWidget();
+}
+
+void FMaterialEditor::OnLockNodeStateCheckStateChanged(ECheckBoxState NewCheckedState)
+{
+	bLockNodeFadeState = (NewCheckedState == ECheckBoxState::Checked) ? true : false;
+}
+
+void FMaterialEditor::OnFocusWholeChainCheckStateChanged(ECheckBoxState NewCheckedState)
+{
+	bFocusWholeChain = (NewCheckedState == ECheckBoxState::Checked) ? true : false;
+
+	if (bHideUnrelatedNodes && !bLockNodeFadeState && bSelectRegularNode)
+	{
+		GraphEditor->ResetAllNodesUnrelatedStates();
+
+		HideUnrelatedNodes();
+	}
+}
+
+
 void FMaterialEditor::OnUseCurrentTexture()
 {
 	// Set the currently selected texture in the generic browser
@@ -3944,6 +4176,13 @@ void FMaterialEditor::UpdateMaterialAfterGraphChange()
 	RegenerateCodeView();
 	RefreshExpressionPreviews();
 	SetMaterialDirty();
+	
+	if (bHideUnrelatedNodes && !bLockNodeFadeState && bSelectRegularNode)
+	{
+		GraphEditor->ResetAllNodesUnrelatedStates();
+
+		HideUnrelatedNodes();
+	}
 }
 
 int32 FMaterialEditor::GetNumberOfSelectedNodes() const
@@ -3951,7 +4190,7 @@ int32 FMaterialEditor::GetNumberOfSelectedNodes() const
 	return GraphEditor->GetSelectedNodes().Num();
 }
 
-FMaterialRenderProxy* FMaterialEditor::GetExpressionPreview(UMaterialExpression* InExpression)
+FMatExpressionPreview* FMaterialEditor::GetExpressionPreview(UMaterialExpression* InExpression)
 {
 	bool bNewlyCreated;
 	return GetExpressionPreview(InExpression, bNewlyCreated);
@@ -4684,6 +4923,7 @@ void FMaterialEditor::OnSelectedNodesChanged(const TSet<class UObject*>& NewSele
 		EditObject = MaterialFunction;
 	}
 
+	bSelectRegularNode = false;
 	if( NewSelection.Num() == 0 )
 	{
 		SelectedObjects.Add(EditObject);
@@ -4694,6 +4934,7 @@ void FMaterialEditor::OnSelectedNodesChanged(const TSet<class UObject*>& NewSele
 		{
 			if (UMaterialGraphNode* GraphNode = Cast<UMaterialGraphNode>(*SetIt))
 			{
+				bSelectRegularNode = true;
 				SelectedObjects.Add(GraphNode->MaterialExpression);
 			}
 			else if (UMaterialGraphNode_Comment* CommentNode = Cast<UMaterialGraphNode_Comment>(*SetIt))
@@ -4709,6 +4950,16 @@ void FMaterialEditor::OnSelectedNodesChanged(const TSet<class UObject*>& NewSele
 
 	GetDetailView()->SetObjects( SelectedObjects, true );
 	FocusDetailsPanel();
+	
+	if (bHideUnrelatedNodes && !bLockNodeFadeState)
+	{
+		GraphEditor->ResetAllNodesUnrelatedStates();
+
+		if (bSelectRegularNode)
+		{
+			HideUnrelatedNodes();
+		}
+	}
 }
 
 void FMaterialEditor::OnNodeDoubleClicked(class UEdGraphNode* Node)
diff --git a/Engine/Source/Editor/MaterialEditor/Private/MaterialEditor.h b/Engine/Source/Editor/MaterialEditor/Private/MaterialEditor.h
--- a/Engine/Source/Editor/MaterialEditor/Private/MaterialEditor.h
+++ b/Engine/Source/Editor/MaterialEditor/Private/MaterialEditor.h
@@ -34,6 +34,7 @@ class UFactory;
 class UMaterialEditorOptions;
 class UMaterialExpressionComment;
 class UMaterialInstance;
+class UMaterialGraphNode;
 struct FGraphAppearanceInfo;
 
 /**
@@ -44,6 +45,7 @@ class FMatExpressionPreview : public FMaterial, public FMaterialRenderProxy
 public:
 	FMatExpressionPreview()
 	: FMaterial()
+	, UnrelatedNodesOpacity(1.0f)
 	{
 		// Register this FMaterial derivative with AddEditorLoadedMaterialResource since it does not have a corresponding UMaterialInterface
 		FMaterial::AddEditorLoadedMaterialResource(this);
@@ -52,6 +54,7 @@ public:
 
 	FMatExpressionPreview(UMaterialExpression* InExpression)
 	: FMaterial()
+	, UnrelatedNodesOpacity(1.0f)
 	, Expression(InExpression)
 	{
 		FMaterial::AddEditorLoadedMaterialResource(this);
@@ -150,7 +153,7 @@ public:
 	virtual bool IsSpecialEngineMaterial() const override { return false; }
 	virtual bool IsWireframe() const override { return false; }
 	virtual bool IsMasked() const override { return false; }
-	virtual enum EBlendMode GetBlendMode() const override { return BLEND_Opaque; }
+	virtual enum EBlendMode GetBlendMode() const override { return BLEND_Translucent; }
 	virtual enum EMaterialShadingModel GetShadingModel() const override { return MSM_Unlit; }
 	virtual float GetOpacityMaskClipValue() const override { return 0.5f; }
 	virtual bool GetCastDynamicShadowAsMasked() const override { return false; }
@@ -172,6 +175,8 @@ public:
 		return Ar << V.Expression;
 	}
 
+	float UnrelatedNodesOpacity;
+
 private:
 	TWeakObjectPtr<UMaterialExpression> Expression;
 	TArray<UTexture*> ReferencedTextures;
@@ -363,7 +368,7 @@ public:
 	virtual bool CanPasteNodes() const override;
 	virtual void PasteNodesHere(const FVector2D& Location) override;
 	virtual int32 GetNumberOfSelectedNodes() const override;
-	virtual FMaterialRenderProxy* GetExpressionPreview(UMaterialExpression* InExpression) override;
+	virtual FMatExpressionPreview* GetExpressionPreview(UMaterialExpression* InExpression) override;
 	virtual void DeleteNodes(const TArray<class UEdGraphNode*>& NodesToDelete) override;
 
 
@@ -585,6 +590,20 @@ private:
 	bool IsOnAlwaysRefreshAllPreviews() const;
 	/** Command for the stats button */
 
+	/** Make nodes which are unrelated to the selected nodes fade out */
+	void ToggleHideUnrelatedNodes();
+	bool IsToggleHideUnrelatedNodesChecked() const;
+	void CollectDownstreamNodes(UMaterialGraphNode* CurrentNode, TArray<UMaterialGraphNode*>& CollectedNodes);
+	void CollectUpstreamNodes(UMaterialGraphNode* CurrentNode, TArray<UMaterialGraphNode*>& CollectedNodes);
+	void HideUnrelatedNodes();
+
+	/** Make a drop down menu to control the opacity of unrelated nodes */
+	TSharedRef<SWidget> MakeHideUnrelatedNodesOptionsMenu();
+	TOptional<float> HandleUnrelatedNodesOpacityBoxValue() const;
+	void HandleUnrelatedNodesOpacityBoxChanged(float NewOpacity);
+	void OnLockNodeStateCheckStateChanged(ECheckBoxState NewCheckedState);
+	void OnFocusWholeChainCheckStateChanged(ECheckBoxState NewCheckedState);
+
 	void ToggleReleaseStats();
 	bool IsToggleReleaseStatsChecked() const;
 	void ToggleBuiltinStats();
@@ -787,6 +806,18 @@ private:
 	/** If true, show stats for an empty material. Helps artists to judge the cost of their changes to the graph. */
 	bool bShowBuiltinStats;
 
+	/** If true, fade out nodes which are unrelated to the selected nodes automatically. */
+	bool bHideUnrelatedNodes;
+
+	/** Lock the current fade state of each node */
+	bool bLockNodeFadeState;
+
+	/** Focus all nodes in the same output chain  */
+	bool bFocusWholeChain;
+
+	/** If a regular node (not a comment node, not the output node) has been selected */
+	bool bSelectRegularNode;
+
 	/** Command list for this editor */
 	TSharedPtr<FUICommandList> GraphEditorCommands;
 
diff --git a/Engine/Source/Editor/MaterialEditor/Private/MaterialEditorActions.cpp b/Engine/Source/Editor/MaterialEditor/Private/MaterialEditorActions.cpp
--- a/Engine/Source/Editor/MaterialEditor/Private/MaterialEditorActions.cpp
+++ b/Engine/Source/Editor/MaterialEditor/Private/MaterialEditorActions.cpp
@@ -30,7 +30,8 @@ void FMaterialEditorCommands::RegisterCommands()
 	UI_COMMAND( ToggleRealtimeExpressions, "Live Nodes", "Toggles real time update of the graph canvas.", EUserInterfaceActionType::ToggleButton, FInputChord() );
 	UI_COMMAND( AlwaysRefreshAllPreviews, "Live Update", "All nodes are previewed live.", EUserInterfaceActionType::ToggleButton, FInputChord() );
 	UI_COMMAND( ToggleMaterialStats, "Stats", "Toggles displaying of the material's stats.", EUserInterfaceActionType::ToggleButton, FInputChord());
-	UI_COMMAND(TogglePlatformStats, "Platform Stats", "Toggles the window that shows material stats and compilation errors for multiple platforms.", EUserInterfaceActionType::ToggleButton, FInputChord() );
+	UI_COMMAND( TogglePlatformStats, "Platform Stats", "Toggles the window that shows material stats and compilation errors for multiple platforms.", EUserInterfaceActionType::ToggleButton, FInputChord() );
+	UI_COMMAND( ToggleHideUnrelatedNodes, "Hide Unrelated", "Toggles hiding nodes which are unrelated to the selected nodes automatically.", EUserInterfaceActionType::ToggleButton, FInputChord() );
 	UI_COMMAND( NewComment, "New Comment", "Creates a new comment node.", EUserInterfaceActionType::Button, FInputChord() );
 	UI_COMMAND( MatertialPasteHere, "Paste Here", "Pastes copied items at this location.", EUserInterfaceActionType::Button, FInputChord() );
 	UI_COMMAND( UseCurrentTexture, "Use Current Texture", "Uses the current texture selected in the content browser.", EUserInterfaceActionType::Button, FInputChord() );
diff --git a/Engine/Source/Editor/MaterialEditor/Public/MaterialEditorActions.h b/Engine/Source/Editor/MaterialEditor/Public/MaterialEditorActions.h
--- a/Engine/Source/Editor/MaterialEditor/Public/MaterialEditorActions.h
+++ b/Engine/Source/Editor/MaterialEditor/Public/MaterialEditorActions.h
@@ -86,6 +86,9 @@ public:
 	
 	/** Toggles live updating of the preview material. */
 	TSharedPtr< FUICommandInfo > ToggleLivePreview;
+	
+	/** Fade nodes which are not connected to the selected nodes. */
+	TSharedPtr< FUICommandInfo > ToggleHideUnrelatedNodes;
 
 	/** Toggles real time expression nodes */
 	TSharedPtr< FUICommandInfo > ToggleRealtimeExpressions;
diff --git a/Engine/Source/Editor/UnrealEd/Classes/Preferences/AnimationBlueprintEditorOptions.h b/Engine/Source/Editor/UnrealEd/Classes/Preferences/AnimationBlueprintEditorOptions.h
new file mode 100644
--- /dev/null
+++ b/Engine/Source/Editor/UnrealEd/Classes/Preferences/AnimationBlueprintEditorOptions.h
@@ -0,0 +1,22 @@
+/**
+ * A configuration class used by the UAnimationBlueprint Editor to save editor
+ * settings across sessions.
+ */
+
+#pragma once
+
+#include "CoreMinimal.h"
+#include "UObject/ObjectMacros.h"
+#include "UObject/Object.h"
+#include "AnimationBlueprintEditorOptions.generated.h"
+
+UCLASS(hidecategories=Object, config=EditorPerProjectUserSettings)
+class UNREALED_API UAnimationBlueprintEditorOptions : public UObject
+{
+	GENERATED_UCLASS_BODY()
+
+	/** If true, fade nodes which are not connected to the selected nodes */
+	UPROPERTY(EditAnywhere, config, Category=Options)
+	uint32 bHideUnrelatedNodes:1;
+
+};
diff --git a/Engine/Source/Editor/UnrealEd/Classes/Preferences/BlueprintEditorOptions.h b/Engine/Source/Editor/UnrealEd/Classes/Preferences/BlueprintEditorOptions.h
new file mode 100644
--- /dev/null
+++ b/Engine/Source/Editor/UnrealEd/Classes/Preferences/BlueprintEditorOptions.h
@@ -0,0 +1,22 @@
+/**
+ * A configuration class used by the UBlueprint Editor to save editor
+ * settings across sessions.
+ */
+
+#pragma once
+
+#include "CoreMinimal.h"
+#include "UObject/ObjectMacros.h"
+#include "UObject/Object.h"
+#include "BlueprintEditorOptions.generated.h"
+
+UCLASS(hidecategories=Object, config=EditorPerProjectUserSettings)
+class UNREALED_API UBlueprintEditorOptions : public UObject
+{
+	GENERATED_UCLASS_BODY()
+
+	/** If true, fade nodes which are not connected to the selected nodes */
+	UPROPERTY(EditAnywhere, config, Category=Options)
+	uint32 bHideUnrelatedNodes:1;
+
+};
diff --git a/Engine/Source/Editor/UnrealEd/Classes/Preferences/MaterialEditorOptions.h b/Engine/Source/Editor/UnrealEd/Classes/Preferences/MaterialEditorOptions.h
--- a/Engine/Source/Editor/UnrealEd/Classes/Preferences/MaterialEditorOptions.h
+++ b/Engine/Source/Editor/UnrealEd/Classes/Preferences/MaterialEditorOptions.h
@@ -37,6 +37,10 @@ class UNREALED_API UMaterialEditorOptions : public UObject
 	/** If true, always refresh the material preview. */
 	UPROPERTY(EditAnywhere, config, Category = Options)
 	uint32 bLivePreviewUpdate : 1;
+	
+	/** If true, fade nodes which are not connected to the selected nodes */
+	UPROPERTY(EditAnywhere, config, Category=Options)
+	uint32 bHideUnrelatedNodes:1;
 
 	/** If true, always refresh all expression previews. */
 	UPROPERTY(EditAnywhere, config, Category=Options)
diff --git a/Engine/Source/Editor/UnrealEd/Private/GraphEditor.cpp b/Engine/Source/Editor/UnrealEd/Private/GraphEditor.cpp
--- a/Engine/Source/Editor/UnrealEd/Private/GraphEditor.cpp
+++ b/Engine/Source/Editor/UnrealEd/Private/GraphEditor.cpp
@@ -117,3 +117,41 @@ void SGraphEditor::NotifyPostPropertyChange( const FPropertyChangedEvent& Proper
 	}
 }
 
+void SGraphEditor::ResetAllNodesUnrelatedStates()
+{
+	TArray<class UEdGraphNode*> AllNodes = GetCurrentGraph()->Nodes;
+
+	for (auto& GraphNode : AllNodes)
+	{
+		if (GraphNode->IsNodeUnrelated())
+		{
+			GraphNode->SetNodeUnrelated(false);
+		}
+	}
+}
+
+void SGraphEditor::FocusCommentNodes(TArray<UEdGraphNode*> &CommentNodes, TArray<UEdGraphNode*> &RelatedNodes)
+{
+	for (auto& CommentNode : CommentNodes)
+	{
+		CommentNode->SetNodeUnrelated(true);
+
+		const FSlateRect CommentRect(
+			CommentNode->NodePosX, 
+			CommentNode->NodePosY, 
+			CommentNode->NodePosX + CommentNode->NodeWidth, 
+			CommentNode->NodePosY + CommentNode->NodeHeight
+		);
+
+		for (auto& RelatedNode : RelatedNodes)
+		{
+			const FVector2D NodePos(RelatedNode->NodePosX, RelatedNode->NodePosY);
+
+			if (CommentRect.ContainsPoint(NodePos))
+			{
+				CommentNode->SetNodeUnrelated(false);
+				break;
+			}
+		}
+	}
+}
diff --git a/Engine/Source/Editor/UnrealEd/Private/Kismet2/BlueprintEditorUtils.cpp b/Engine/Source/Editor/UnrealEd/Private/Kismet2/BlueprintEditorUtils.cpp
--- a/Engine/Source/Editor/UnrealEd/Private/Kismet2/BlueprintEditorUtils.cpp
+++ b/Engine/Source/Editor/UnrealEd/Private/Kismet2/BlueprintEditorUtils.cpp
@@ -2140,6 +2140,13 @@ void FBlueprintEditorUtils::MarkBlueprintAsModified(UBlueprint* Blueprint, FProp
 		// Clear out the cache as the user may have added or removed a latent action to a macro graph
 		FBlueprintEditorUtils::ClearMacroCosmeticInfoCache(Blueprint);
 	}
+	
+	IAssetEditorInstance* AssetEditor = FAssetEditorManager::Get().FindEditorForAsset(Blueprint, true);
+	if (AssetEditor)
+	{
+		FBlueprintEditor* BlueprintEditor = static_cast<FBlueprintEditor*>(AssetEditor);
+		BlueprintEditor->UpdateNodesUnrelatedStatesAfterGraphChange();
+	}
 }
 
 bool FBlueprintEditorUtils::ShouldRegenerateBlueprint(UBlueprint* Blueprint)
diff --git a/Engine/Source/Editor/UnrealEd/Private/PreferenceStubs.cpp b/Engine/Source/Editor/UnrealEd/Private/PreferenceStubs.cpp
--- a/Engine/Source/Editor/UnrealEd/Private/PreferenceStubs.cpp
+++ b/Engine/Source/Editor/UnrealEd/Private/PreferenceStubs.cpp
@@ -5,7 +5,9 @@
 #include "Preferences/CascadeOptions.h"
 #include "Preferences/CurveEdOptions.h"
 #include "Preferences/MaterialEditorOptions.h"
+#include "Preferences/BlueprintEditorOptions.h"
 #include "Preferences/PersonaOptions.h"
+#include "Preferences/AnimationBlueprintEditorOptions.h"
 #include "Preferences/PhysicsAssetEditorOptions.h"
 #include "Preferences/MaterialStatsOptions.h"
 
@@ -70,6 +72,15 @@ UMaterialStatsOptions::UMaterialStatsOptions(const FObjectInitializer& ObjectIni
 	bMaterialQualityUsed[EMaterialQualityLevel::High] = 1;
 }
 
+UBlueprintEditorOptions::UBlueprintEditorOptions(const FObjectInitializer& ObjectInitializer)
+	: Super(ObjectInitializer)
+{
+}
+
+UAnimationBlueprintEditorOptions::UAnimationBlueprintEditorOptions(const FObjectInitializer& ObjectInitializer)
+	: Super(ObjectInitializer)
+{
+}
 
 UCurveEdOptions::UCurveEdOptions(const FObjectInitializer& ObjectInitializer)
 	: Super(ObjectInitializer)
diff --git a/Engine/Source/Editor/UnrealEd/Public/GraphEditor.h b/Engine/Source/Editor/UnrealEd/Public/GraphEditor.h
--- a/Engine/Source/Editor/UnrealEd/Public/GraphEditor.h
+++ b/Engine/Source/Editor/UnrealEd/Public/GraphEditor.h
@@ -457,6 +457,11 @@ public:
 			Implementation->SetNodeFactory(NewNodeFactory);
 		}
 	}
+	
+	/** Common methods for MaterialEditor and BlueprintEditor's focusing related nodes feature */
+	UNREALED_API void ResetAllNodesUnrelatedStates();
+
+	UNREALED_API void FocusCommentNodes(TArray<UEdGraphNode*> &CommentNodes, TArray<UEdGraphNode*> &RelatedNodes);
 
 	virtual void OnAlignTop()
 	{
diff --git a/Engine/Source/Runtime/Engine/Classes/EdGraph/EdGraphNode.h b/Engine/Source/Runtime/Engine/Classes/EdGraph/EdGraphNode.h
--- a/Engine/Source/Runtime/Engine/Classes/EdGraph/EdGraphNode.h
+++ b/Engine/Source/Runtime/Engine/Classes/EdGraph/EdGraphNode.h
@@ -245,24 +245,33 @@ private:
 	/** Whether the node was created as part of an expansion step */
 	uint8 bIsIntermediateNode : 1;
 
+#if WITH_EDITORONLY_DATA
+	/** Whether this node is unrelated to the selected nodes or not */
+	uint8 bUnrelated : 1;
+
+#endif
+
 public:
+
 	/** Flag to check for compile error/warning */
 	UPROPERTY()
 	uint8 bHasCompilerMessage:1;
 
 	/** Comment bubble pinned state */
+
+
+#if WITH_EDITORONLY_DATA
 	UPROPERTY()
-	uint8 bCommentBubblePinned:1;
+	uint8 bCommentBubblePinned : 1;
 
 	/** Comment bubble visibility */
 	UPROPERTY()
-	uint8 bCommentBubbleVisible:1;
+	uint8 bCommentBubbleVisible : 1;
 
 	/** Make comment bubble visible */
 	UPROPERTY(Transient)
-	uint8 bCommentBubbleMakeVisible:1;
+	uint8 bCommentBubbleMakeVisible : 1;
 
-#if WITH_EDITORONLY_DATA
 	/** If true, this node can be renamed in the editor */
 	UPROPERTY()
 	uint8 bCanRenameNode:1;
@@ -325,6 +334,20 @@ public:
 		return bUserSetEnabledState;
 	}
 
+#if WITH_EDITOR
+	/** Set this node unrelated or not. */
+	FORCEINLINE void SetNodeUnrelated(bool bNodeUnrelated)
+	{
+		bUnrelated = bNodeUnrelated;
+	}
+
+	/** Determines whether this node is unrelated to the selected nodes or not. */
+	FORCEINLINE bool IsNodeUnrelated() const
+	{
+		return bUnrelated;
+	}
+#endif
+
 	/** Determines whether or not the node will compile in development mode. */
 	virtual bool IsInDevelopmentMode() const;
 
diff --git a/Engine/Source/Runtime/Engine/Private/EdGraph/EdGraphNode.cpp b/Engine/Source/Runtime/Engine/Private/EdGraph/EdGraphNode.cpp
--- a/Engine/Source/Runtime/Engine/Private/EdGraph/EdGraphNode.cpp
+++ b/Engine/Source/Runtime/Engine/Private/EdGraph/EdGraphNode.cpp
@@ -158,10 +158,12 @@ UEdGraphNode::UEdGraphNode(const FObjectInitializer& ObjectInitializer)
 	, bIsNodeEnabled_DEPRECATED(true)
 #if WITH_EDITORONLY_DATA
 	, bCanResizeNode(false)
-#endif // WITH_EDITORONLY_DATA
+	, bUnrelated(false)
 	, bCommentBubblePinned(false)
 	, bCommentBubbleVisible(false)
 	, bCommentBubbleMakeVisible(false)
+#endif // WITH_EDITORONLY_DATA
+
 {
 #if WITH_EDITORONLY_DATA
 	{ static const FAutoRegisterLocalizationDataGatheringCallback AutomaticRegistrationOfLocalizationGatherer(UEdGraphNode::StaticClass(), &GatherGraphNodeForLocalization); }
