Prevent new reparenting logic from modifying objects that are being replaced by asset reload
#rnx

diff --git a/Engine/Source/Editor/Kismet/Private/BlueprintCompilationManager.cpp b/Engine/Source/Editor/Kismet/Private/BlueprintCompilationManager.cpp
--- a/Engine/Source/Editor/Kismet/Private/BlueprintCompilationManager.cpp
+++ b/Engine/Source/Editor/Kismet/Private/BlueprintCompilationManager.cpp
@@ -1443,7 +1443,10 @@ void FBlueprintCompilationManagerImpl::FlushReinstancingQueueImpl()
 		FScopedDurationTimer ReinstTimer(GTimeReinstancing);
 		
 		TGuardValue<bool> ReinstancingGuard(GIsReinstancing, true);
-		FBlueprintCompileReinstancer::BatchReplaceInstancesOfClass(ClassesToReinstance, true);
+		
+		FBatchReplaceInstancesOfClassParameters Options;
+		Options.bArchetypesAreUpToDate = true;
+		FBlueprintCompileReinstancer::BatchReplaceInstancesOfClass(ClassesToReinstance, Options);
 
 		if (IsAsyncLoading())
 		{
@@ -1619,7 +1622,28 @@ void FBlueprintCompilationManagerImpl::ReparentHierarchies(const TMap<UClass*, U
 		OldClassToNewClassDerivedTypes.Add(ReinstancingJob.OldToNew);
 	}
 	TGuardValue<bool> ReinstancingGuard(GIsReinstancing, true);
-	FBlueprintCompileReinstancer::BatchReplaceInstancesOfClass( OldClassToNewClassDerivedTypes, /* bArchetypesAreUpToDate */ true );
+	FBatchReplaceInstancesOfClassParameters Options;
+	Options.bArchetypesAreUpToDate = true;
+
+	// Make sure we don't replace old instances that are in the *callers* old to new TMap!
+	TSet<UObject*> OldObjects;
+	for(TPair<UClass*, UClass*> OldToNew : OldClassToNewClassDerivedTypes)
+	{
+		TArray< UObject* > OldObjectsOfType;
+		GetObjectsOfClass(OldToNew.Key, OldObjectsOfType);
+
+		for(UObject* Obj : OldObjectsOfType)
+		{
+			if(Obj->HasAnyFlags(RF_NewerVersionExists))
+			{
+				OldObjects.Add(Obj);
+			}
+		}
+	}
+	Options.ObjectsThatShouldUseOldStuff = &OldObjects;
+	Options.InstancesThatShouldUseOldClass = &OldObjects;
+
+	FBlueprintCompileReinstancer::BatchReplaceInstancesOfClass( OldClassToNewClassDerivedTypes, Options );
 }
 
 
diff --git a/Engine/Source/Editor/UnrealEd/Private/EditorServer.cpp b/Engine/Source/Editor/UnrealEd/Private/EditorServer.cpp
--- a/Engine/Source/Editor/UnrealEd/Private/EditorServer.cpp
+++ b/Engine/Source/Editor/UnrealEd/Private/EditorServer.cpp
@@ -1396,7 +1396,7 @@ void UEditorEngine::PostUndo(bool)
 	}
 
 	// Re-instance any actors that need it
-	FBlueprintCompileReinstancer::BatchReplaceInstancesOfClass(OldToNewClassMapToReinstance, false);
+	FBlueprintCompileReinstancer::BatchReplaceInstancesOfClass(OldToNewClassMapToReinstance);
 
 	RedrawLevelEditingViewports();
 }
diff --git a/Engine/Source/Editor/UnrealEd/Private/Kismet2/KismetReinstanceUtilities.cpp b/Engine/Source/Editor/UnrealEd/Private/Kismet2/KismetReinstanceUtilities.cpp
--- a/Engine/Source/Editor/UnrealEd/Private/Kismet2/KismetReinstanceUtilities.cpp
+++ b/Engine/Source/Editor/UnrealEd/Private/Kismet2/KismetReinstanceUtilities.cpp
@@ -1480,14 +1480,14 @@ void FBlueprintCompileReinstancer::ReplaceInstancesOfClassEx(const FReplaceInsta
 	ReplaceInstancesOfClass_Inner(OldToNewClassMap, Parameters.OriginalCDO, Parameters.ObjectsThatShouldUseOldStuff, Parameters.bClassObjectReplaced, Parameters.bPreserveRootComponent, /*bArchetypesAreUpToDate=*/false, Parameters.InstancesThatShouldUseOldClass);
 }
 
-void FBlueprintCompileReinstancer::BatchReplaceInstancesOfClass(TMap<UClass*, UClass*>& InOldToNewClassMap, bool bArchetypesAreUpToDate)
+void FBlueprintCompileReinstancer::BatchReplaceInstancesOfClass(TMap<UClass*, UClass*>& InOldToNewClassMap, const FBatchReplaceInstancesOfClassParameters& Options )
 {
 	if (InOldToNewClassMap.Num() == 0)
 	{
 		return;
 	}
 
-	ReplaceInstancesOfClass_Inner(InOldToNewClassMap, nullptr, nullptr, false /*bClassObjectReplaced*/, true /*bPreserveRootComponent*/, bArchetypesAreUpToDate);
+	ReplaceInstancesOfClass_Inner(InOldToNewClassMap, nullptr, Options.ObjectsThatShouldUseOldStuff, false /*bClassObjectReplaced*/, true /*bPreserveRootComponent*/, Options.bArchetypesAreUpToDate, Options.InstancesThatShouldUseOldClass);
 }
 
 UClass* FBlueprintCompileReinstancer::MoveCDOToNewClass(UClass* OwnerClass, const TMap<UClass*, UClass*>& OldToNewMap, bool bAvoidCDODuplication)
diff --git a/Engine/Source/Editor/UnrealEd/Public/Kismet2/KismetReinstanceUtilities.h b/Engine/Source/Editor/UnrealEd/Public/Kismet2/KismetReinstanceUtilities.h
--- a/Engine/Source/Editor/UnrealEd/Public/Kismet2/KismetReinstanceUtilities.h
+++ b/Engine/Source/Editor/UnrealEd/Public/Kismet2/KismetReinstanceUtilities.h
@@ -62,6 +62,15 @@ struct UNREALED_API FReplaceInstancesOfClassParameters
 	bool bPreserveRootComponent; 
 };
 
+struct UNREALED_API FBatchReplaceInstancesOfClassParameters
+{
+	TSet<UObject*>* ObjectsThatShouldUseOldStuff = nullptr;
+	
+	const TSet<UObject*>* InstancesThatShouldUseOldClass = nullptr;
+
+	bool bArchetypesAreUpToDate = false; 
+};
+
 class UNREALED_API FBlueprintCompileReinstancer : public TSharedFromThis<FBlueprintCompileReinstancer>, public FGCObject
 {
 protected:
@@ -154,7 +163,7 @@ public:
 	static void ReplaceInstancesOfClassEx(const FReplaceInstancesOfClassParameters& Parameters );
 
 	/** Batch replaces a mapping of one or more classes to their new class by leveraging ReplaceInstancesOfClass */
-	static void BatchReplaceInstancesOfClass(TMap<UClass*, UClass*>& InOldToNewClassMap, bool bArchetypesAreUpToDate);
+	static void BatchReplaceInstancesOfClass(TMap<UClass*, UClass*>& InOldToNewClassMap, const FBatchReplaceInstancesOfClassParameters& Options = FBatchReplaceInstancesOfClassParameters());
 	
 	/** Function used to safely discard a CDO, so that the class can have its layout changed, callers must move parent CDOs aside before moving child CDOs aside: */
 	static UClass* MoveCDOToNewClass(UClass* OwnerClass, const TMap<UClass*, UClass*>& OldToNewMap, bool bAvoidCDODuplication);
