Fix for a backwards-compatibility issue that prevents any global Blueprint search from functioning properly (regression).

Change summary:
- Added FSearchDataVersionInfo to consolidate cached version data passed between internal APIs.
- Modified FFindInBlueprintManager::QuerySingleBlueprint() to return a read-only ptr to the matching entry.
- Modified a few internal APIs to include FSearchDataVersionInfo as input.
- Modified FFindInBlueprintSearchManager::ExtractUnloadedFiBData() to query unloaded asset package summaries for last-saved editor object custom version info. This value is then cached into the search DB.
- Modified FFindInBlueprintSearchManager::ConvertJsonStringToObject() (and helpers) to configure the FArchive stream with the cached editor object custom version from the search DB before deserializing the LUT.

diff --git a/Engine/Source/Editor/Kismet/Private/FindInBlueprintManager.cpp b/Engine/Source/Editor/Kismet/Private/FindInBlueprintManager.cpp
--- a/Engine/Source/Editor/Kismet/Private/FindInBlueprintManager.cpp
+++ b/Engine/Source/Editor/Kismet/Private/FindInBlueprintManager.cpp
@@ -49,6 +49,7 @@
 #include "BlueprintAssetHandler.h"
 
 #include "JsonObjectConverter.h"
+#include "UObject/EditorObjectVersion.h"
 #include "Framework/Notifications/NotificationManager.h"
 #include "Widgets/Notifications/SNotificationList.h"
 
@@ -110,6 +111,14 @@ inline UObject* GetAssetObject(UObject* InObject)
 	return AssetObject;
 }
 
+////////////////////////////////////
+// FSearchDataVersionInfo
+FSearchDataVersionInfo FSearchDataVersionInfo::Current =
+{
+	EFiBVersion::FIB_VER_LATEST,
+	FEditorObjectVersion::LatestVersion
+};
+
 ////////////////////////////////////
 // FStreamSearch
 FStreamSearch::FStreamSearch(const FString& InSearchValue)
@@ -157,7 +166,7 @@ uint32 FStreamSearch::Run()
 		if (QueryResult.ImaginaryBlueprint.IsValid())
 		{
 			// If the Blueprint is below the version, add it to a list. The search will still proceed on this Blueprint
-			if (QueryResult.Version < MinimiumVersionRequirement)
+			if (QueryResult.VersionInfo.FiBDataVersion < MinimiumVersionRequirement)
 			{
 				++BlueprintCountBelowVersion;
 			}
@@ -334,7 +343,9 @@ namespace FiBSerializationHelpers
 	{
 		TArray<uint8> DerivedData;
 		DecodeFromStream(InStream, sizeof(Type), DerivedData);
+
 		FMemoryReader SizeOfDataAr(DerivedData);
+		SizeOfDataAr.SetCustomVersions(InStream.GetCustomVersions());
 
 		Type ReturnValue;
 		SizeOfDataAr << ReturnValue;
@@ -347,7 +358,9 @@ namespace FiBSerializationHelpers
 	{
 		TArray<uint8> DerivedData;
 		DecodeFromStream(InStream, InBytes, DerivedData);
+
 		FMemoryReader SizeOfDataAr(DerivedData);
+		SizeOfDataAr.SetCustomVersions(InStream.GetCustomVersions());
 
 		Type ReturnValue;
 		SizeOfDataAr << ReturnValue;
@@ -1634,7 +1647,30 @@ void FFindInBlueprintSearchManager::ExtractUnloadedFiBData(const FAssetData& InA
 	{
 		checkf(NewSearchData.Value.Len(), TEXT("Versioned search data was zero length!"));
 		FBufferReader ReaderStream((void*)*NewSearchData.Value, NewSearchData.Value.Len() * sizeof(TCHAR), false);
-		NewSearchData.Version = FiBSerializationHelpers::Deserialize<int32>(ReaderStream);
+		NewSearchData.VersionInfo.FiBDataVersion = FiBSerializationHelpers::Deserialize<int32>(ReaderStream);
+	}
+
+	// Determine the editor object version that the asset package was last serialized with
+	FString PackageFilename;
+	if (ensureMsgf(FPackageName::DoesPackageExist(InAssetData.PackageName.ToString(), nullptr, &PackageFilename), TEXT("FiB: Failed to map package to filename.")))
+	{
+		// Open a new file archive for reading
+		FArchive* PackageFile = IFileManager::Get().CreateFileReader(*PackageFilename);
+		if (ensureMsgf(PackageFile != nullptr, TEXT("FiB: Unable to open package to read file summary.")))
+		{
+			// Read the package file summary
+			FPackageFileSummary PackageFileSummary;
+			*PackageFile << PackageFileSummary;
+
+			// Close the file
+			delete PackageFile;
+
+			// If an editor object version exists in the package file summary, record it
+			if (const FCustomVersion* const EditorObjectVersion = PackageFileSummary.GetCustomVersionContainer().GetVersion(FEditorObjectVersion::GUID))
+			{
+				NewSearchData.VersionInfo.EditorObjectVersion = EditorObjectVersion->Version;
+			}
+		}
 	}
 
 	// Since the asset was not loaded, pull out the searchable data stored in the asset
@@ -1799,8 +1835,10 @@ FString FFindInBlueprintSearchManager::GatherBlueprintSearchMetadata(const UBlue
 	Writer->WriteObjectEnd();
 	Writer->Close();
 
-	int32 Version = EFiBVersion::FIB_VER_LATEST;
-	SearchMetaData = FiBSerializationHelpers::Serialize(Version, false) + Writer->GetSerializedLookupTable() + SearchMetaData;
+	// Build the search metadata string for the asset tag (version + LUT + JSON)
+	SearchMetaData = FiBSerializationHelpers::Serialize(FSearchDataVersionInfo::Current.FiBDataVersion, false)
+		+ Writer->GetSerializedLookupTable()
+		+ SearchMetaData;
 
 	return SearchMetaData;
 }
@@ -1856,8 +1894,11 @@ void FFindInBlueprintSearchManager::AddOrUpdateBlueprintSearchMetadata(UBlueprin
 		// Cannot successfully gather most searchable data if there is no SkeletonGeneratedClass, so don't try, leave it as whatever it was last set to
 		if (InBlueprint->SkeletonGeneratedClass != nullptr)
 		{
+			// Update search metadata string content
 			SearchArray[Index].Value = GatherBlueprintSearchMetadata(InBlueprint);
-			SearchArray[Index].Version = EFiBVersion::FIB_VER_LATEST;
+
+			// Update version info stored in database to latest
+			SearchArray[Index].VersionInfo = FSearchDataVersionInfo::Current;
 		}
 
 		// Remove it from the list of pending assets (if it exists)
@@ -1915,7 +1956,7 @@ bool FFindInBlueprintSearchManager::ContinueSearchQuery(const FStreamSearch* InS
 				// If there is FiB data, parse it into an ImaginaryBlueprint
 				if (SearchArray[SearchIdx].Value.Len() > 0)
 				{
-					SearchArray[SearchIdx].ImaginaryBlueprint = MakeShareable(new FImaginaryBlueprint(FPaths::GetBaseFilename(SearchArray[SearchIdx].AssetPath.ToString()), SearchArray[SearchIdx].AssetPath.ToString(), SearchArray[SearchIdx].ParentClass, SearchArray[SearchIdx].Interfaces, SearchArray[SearchIdx].Value, SearchArray[SearchIdx].Version != 0));
+					SearchArray[SearchIdx].ImaginaryBlueprint = MakeShareable(new FImaginaryBlueprint(FPaths::GetBaseFilename(SearchArray[SearchIdx].AssetPath.ToString()), SearchArray[SearchIdx].AssetPath.ToString(), SearchArray[SearchIdx].ParentClass, SearchArray[SearchIdx].Interfaces, SearchArray[SearchIdx].Value, SearchArray[SearchIdx].VersionInfo));
 					SearchArray[SearchIdx].Value.Empty();
 				}
 
@@ -1962,7 +2003,7 @@ float FFindInBlueprintSearchManager::GetPercentComplete(const FStreamSearch* InS
 	return ReturnPercent;
 }
 
-FString FFindInBlueprintSearchManager::QuerySingleBlueprint(UBlueprint* InBlueprint, bool bInRebuildSearchData/* = true*/)
+const FSearchData* FFindInBlueprintSearchManager::QuerySingleBlueprint(UBlueprint* InBlueprint, bool bInRebuildSearchData)
 {
 	// AddOrUpdateBlueprintSearchMetadata would fail to cache any data for a Blueprint loaded specifically for diffing, but the bigger question
 	// here in this function is how you are doing a search specifically for data within this Blueprint. This function is limited to be called
@@ -1988,7 +2029,7 @@ FString FFindInBlueprintSearchManager::QuerySingleBlueprint(UBlueprint* InBluepr
 				*ArrayIdx,
 				SearchArray.Num());
 
-			return SearchArray[*ArrayIdx].Value;
+			return &SearchArray[*ArrayIdx];
 		}
 		else if(bInRebuildSearchData)
 		{
@@ -2002,7 +2043,7 @@ FString FFindInBlueprintSearchManager::QuerySingleBlueprint(UBlueprint* InBluepr
 		// Also warn here as we do not index diff-only packages.
 		UE_LOG(LogBlueprint, Warning, TEXT("Attempted to query an old Blueprint package opened for diffing!"));
 	}
-	return FString();
+	return nullptr;
 }
 
 void FFindInBlueprintSearchManager::PauseFindInBlueprintSearch()
@@ -2233,7 +2274,7 @@ void FFindInBlueprintSearchManager::CacheAllAssets(TWeakPtr< SFindInBlueprints >
 			// Add any out-of-date Blueprints to the list
 			for (FSearchData SearchData : SearchArray)
 			{
-				if ((SearchData.Value.Len() != 0 || SearchData.ImaginaryBlueprint.IsValid()) && SearchData.Version < InOptions.MinimiumVersionRequirement)
+				if ((SearchData.Value.Len() != 0 || SearchData.ImaginaryBlueprint.IsValid()) && SearchData.VersionInfo.FiBDataVersion < InOptions.MinimiumVersionRequirement)
 				{
 					BlueprintsToUpdate.Add(SearchData.AssetPath);
 				}
@@ -2424,7 +2465,7 @@ bool FFindInBlueprintSearchManager::IsUnindexedCacheInProgress() const
 	return IsCacheInProgress() && CurrentCacheOpType == EFiBCacheOpType::CacheUnindexedAssets;
 }
 
-TSharedPtr< FJsonObject > FFindInBlueprintSearchManager::ConvertJsonStringToObject(bool bInIsVersioned, FString InJsonString, TMap<int32, FText>& OutFTextLookupTable)
+TSharedPtr< FJsonObject > FFindInBlueprintSearchManager::ConvertJsonStringToObject(FSearchDataVersionInfo InVersionInfo, FString InJsonString, TMap<int32, FText>& OutFTextLookupTable)
 {
 	/** The searchable data is more complicated than a Json string, the Json being the main searchable body that is parsed. Below is a diagram of the full data:
 	 *  | int32 "Version" | int32 "Size" | TMap "Lookup Table" | Json String |
@@ -2441,12 +2482,19 @@ TSharedPtr< FJsonObject > FFindInBlueprintSearchManager::ConvertJsonStringToObje
 	int32 SizeOfData;
 	FBufferReader ReaderStream((void*)*InJsonString, InJsonString.Len() * sizeof(TCHAR), false);
 
-	int32 Version = 0;
-	if (bInIsVersioned)
+	// If the stream is versioned, read past the version info
+	if (InVersionInfo.FiBDataVersion > EFiBVersion::FIB_VER_BASE)
 	{
-		FiBSerializationHelpers::Deserialize<int32>(ReaderStream);
+		// Read the FiB search data version
+		const int32 Version = FiBSerializationHelpers::Deserialize<int32>(ReaderStream);
+
+		// Check that the deserialized version matches up with what's recorded in the search database
+		ensureMsgf(Version == InVersionInfo.FiBDataVersion, TEXT("FiB: JSON stream data does not match search data version from database. This is unexpected."));
 	}
 
+	// Configure the JSON stream with the proper object version for FText serialization when reading the LUT
+	ReaderStream.SetCustomVersion(FEditorObjectVersion::GUID, InVersionInfo.EditorObjectVersion, TEXT("Dev-Editor"));
+
  	// Read, as a byte string, the number of characters composing the Lookup Table for the Json.
 	SizeOfData = FiBSerializationHelpers::Deserialize<int32>(ReaderStream);
 
diff --git a/Engine/Source/Editor/Kismet/Private/FindInBlueprints.cpp b/Engine/Source/Editor/Kismet/Private/FindInBlueprints.cpp
--- a/Engine/Source/Editor/Kismet/Private/FindInBlueprints.cpp
+++ b/Engine/Source/Editor/Kismet/Private/FindInBlueprints.cpp
@@ -1075,12 +1075,13 @@ void SFindInBlueprints::MakeSearchQuery(FString InSearchString, bool bInIsFindWi
 				Interfaces.Add(InterfaceDesc.Interface->GetPathName());
 			}
 
-			const FString UnparsedStringData = FFindInBlueprintSearchManager::Get().QuerySingleBlueprint(Blueprint);
-			const bool bIsBlueprintIndexed = !UnparsedStringData.IsEmpty();
+			const bool bRebuildSearchData = true;
+			const FSearchData* SearchData = FFindInBlueprintSearchManager::Get().QuerySingleBlueprint(Blueprint, bRebuildSearchData);
+			const bool bIsBlueprintIndexed = SearchData && !SearchData->Value.IsEmpty();
 
 			if (bIsBlueprintIndexed)
 			{
-				FImaginaryFiBDataSharedPtr ImaginaryBlueprint(new FImaginaryBlueprint(Blueprint->GetName(), Blueprint->GetPathName(), ParentClass, Interfaces, UnparsedStringData));
+				FImaginaryFiBDataSharedPtr ImaginaryBlueprint(new FImaginaryBlueprint(Blueprint->GetName(), Blueprint->GetPathName(), ParentClass, Interfaces, SearchData->Value, SearchData->VersionInfo));
 				TSharedPtr< FFiBSearchInstance > SearchInstance(new FFiBSearchInstance);
 				FSearchResult SearchResult = RootSearchResult = SearchInstance->StartSearchQuery(SearchValue, ImaginaryBlueprint);
 
diff --git a/Engine/Source/Editor/Kismet/Private/ImaginaryBlueprintData.cpp b/Engine/Source/Editor/Kismet/Private/ImaginaryBlueprintData.cpp
--- a/Engine/Source/Editor/Kismet/Private/ImaginaryBlueprintData.cpp
+++ b/Engine/Source/Editor/Kismet/Private/ImaginaryBlueprintData.cpp
@@ -427,12 +427,12 @@ void FCategorySectionHelper::ParseAllChildData_Internal(ESearchableValueStatus I
 //////////////////////////////////////////
 // FImaginaryBlueprint
 
-FImaginaryBlueprint::FImaginaryBlueprint(FString InBlueprintName, FString InBlueprintPath, FString InBlueprintParentClass, TArray<FString>& InInterfaces, FString InUnparsedStringData, bool bInIsVersioned/*=true*/)
+FImaginaryBlueprint::FImaginaryBlueprint(FString InBlueprintName, FString InBlueprintPath, FString InBlueprintParentClass, TArray<FString>& InInterfaces, FString InUnparsedStringData, FSearchDataVersionInfo InVersionInfo)
 	: FImaginaryFiBData(nullptr)
 	, BlueprintPath(InBlueprintPath)
 	, UnparsedStringData(InUnparsedStringData)
 {
-	ParseToJson(bInIsVersioned);
+	ParseToJson(InVersionInfo);
 	LookupTablePtr = &LookupTable;
 	ParsedTagsAndValues.Add(FindInBlueprintsHelpers::FSimpleFTextKeyStorage(FFindInBlueprintSearchTags::FiB_Name), FSearchableValueInfo(FFindInBlueprintSearchTags::FiB_Name, FText::FromString(InBlueprintName), ESearchableValueStatus::ExplicitySearchable));
 	ParsedTagsAndValues.Add(FindInBlueprintsHelpers::FSimpleFTextKeyStorage(FFindInBlueprintSearchTags::FiB_Path), FSearchableValueInfo(FFindInBlueprintSearchTags::FiB_Path, FText::FromString(InBlueprintPath), ESearchableValueStatus::ExplicitySearchable));
@@ -478,9 +478,9 @@ bool FImaginaryBlueprint::CanCallFilter(ESearchQueryFilter InSearchQueryFilter)
 		FImaginaryFiBData::CanCallFilter(InSearchQueryFilter);
 }
 
-void FImaginaryBlueprint::ParseToJson(bool bInIsVersioned)
+void FImaginaryBlueprint::ParseToJson(FSearchDataVersionInfo InVersionInfo)
 {
-	UnparsedJsonObject = FFindInBlueprintSearchManager::ConvertJsonStringToObject(bInIsVersioned, UnparsedStringData, LookupTable);
+	UnparsedJsonObject = FFindInBlueprintSearchManager::ConvertJsonStringToObject(InVersionInfo, UnparsedStringData, LookupTable);
 }
 
 bool FImaginaryBlueprint::TrySpecialHandleJsonValue(FText InKey, TSharedPtr< FJsonValue > InJsonValue)
diff --git a/Engine/Source/Editor/Kismet/Public/FindInBlueprintManager.h b/Engine/Source/Editor/Kismet/Public/FindInBlueprintManager.h
--- a/Engine/Source/Editor/Kismet/Public/FindInBlueprintManager.h
+++ b/Engine/Source/Editor/Kismet/Public/FindInBlueprintManager.h
@@ -102,6 +102,31 @@ struct KISMET_API FFindInBlueprintSearchTags
 	/** End const values for Find-in-Blueprint */
 };
 
+/** FiB data versioning */
+enum EFiBVersion
+{
+	FIB_VER_BASE = 0, // All Blueprints prior to versioning will automatically be assumed to be at 0 if they have FiB data collected
+	FIB_VER_VARIABLE_REFERENCE, // Variable references (FMemberReference) is collected in FiB
+	FIB_VER_INTERFACE_GRAPHS, // Implemented Interface Graphs is collected in FiB
+
+	// -----<new versions can be added before this line>-------------------------------------------------
+	FIB_VER_PLUS_ONE,
+	FIB_VER_LATEST = FIB_VER_PLUS_ONE - 1 // Always the last version, we want Blueprints to be at latest
+};
+
+/** Consolidated version info for a Blueprint search data entry */
+struct FSearchDataVersionInfo
+{
+	/** FiB asset registry tag data version */
+	int32 FiBDataVersion = EFiBVersion::FIB_VER_BASE;
+
+	/** Editor object version used to serialize values in the JSON string lookup table */
+	int32 EditorObjectVersion = -1;
+
+	/** Current version info */
+	static FSearchDataVersionInfo Current;
+};
+
 /** Tracks data relevant to a Blueprint for searches */
 struct FSearchData
 {
@@ -126,13 +151,12 @@ struct FSearchData
 	/** Cached ImaginaryBlueprint data for the searchable content, prevents having to re-parse every search */
 	FImaginaryFiBDataSharedPtr ImaginaryBlueprint;
 
-	/** Version of the data */
-	int32 Version;
+	/** Data versioning */
+	FSearchDataVersionInfo VersionInfo;
 
 	FSearchData()
 		: Blueprint(nullptr)
 		, bMarkedForDeletion(false)
-		, Version(0)
 	{
 	}
 };
@@ -165,17 +189,6 @@ struct FSearchTagDataPair
 	FText Value;
 };
 
-enum EFiBVersion
-{
-	FIB_VER_BASE = 0, // All Blueprints prior to versioning will automatically be assumed to be at 0 if they have FiB data collected
-	FIB_VER_VARIABLE_REFERENCE, // Variable references (FMemberReference) is collected in FiB
-	FIB_VER_INTERFACE_GRAPHS, // Implemented Interface Graphs is collected in FiB
-
-	// -----<new versions can be added before this line>-------------------------------------------------
-	FIB_VER_PLUS_ONE,
-	FIB_VER_LATEST = FIB_VER_PLUS_ONE - 1 // Always the last version, we want Blueprints to be at latest
-};
-
 struct KISMET_API FFiBMD
 {
 	static const FString FiBSearchableMD;
@@ -434,13 +447,13 @@ public:
 	float GetPercentComplete(const class FStreamSearch* InSearchOriginator) const;
 
 	/**
-	 * Query for a single, specific Blueprint's search block
+	 * Query for a single, specific Blueprint's search block.
 	 *
 	 * @param InBlueprint				The Blueprint to search for
 	 * @param bInRebuildSearchData		When TRUE the search data will be freshly collected
 	 * @return							The search block
 	 */
-	FString QuerySingleBlueprint(UBlueprint* InBlueprint, bool bInRebuildSearchData = true);
+	const FSearchData* QuerySingleBlueprint(UBlueprint* InBlueprint, bool bInRebuildSearchData);
 
 	/** Converts a string of hex characters, previously converted by ConvertFTextToHexString, to an FText. */
 	static FText ConvertHexStringToFText(FString InHexString);
@@ -514,7 +527,7 @@ public:
 	TWeakPtr<SFindInBlueprints> GetSourceCachingWidget() const { return SourceCachingWidget; }
 
 	/** Given a fully constructed Find-in-Blueprint FString of searchable data, will parse and construct a JsonObject */
-	static TSharedPtr< class FJsonObject > ConvertJsonStringToObject(bool bInIsVersioned, FString InJsonString, TMap<int32, FText>& OutFTextLookupTable);
+	static TSharedPtr< class FJsonObject > ConvertJsonStringToObject(FSearchDataVersionInfo InVersionInfo, FString InJsonString, TMap<int32, FText>& OutFTextLookupTable);
 
 	void EnableGatheringData(bool bInEnableGatheringData) { bEnableGatheringData = bInEnableGatheringData; }
 
diff --git a/Engine/Source/Editor/Kismet/Public/ImaginaryBlueprintData.h b/Engine/Source/Editor/Kismet/Public/ImaginaryBlueprintData.h
--- a/Engine/Source/Editor/Kismet/Public/ImaginaryBlueprintData.h
+++ b/Engine/Source/Editor/Kismet/Public/ImaginaryBlueprintData.h
@@ -292,7 +292,7 @@ protected:
 class FImaginaryBlueprint : public FImaginaryFiBData
 {
 public:
-	FImaginaryBlueprint(FString InBlueprintName, FString InBlueprintPath, FString InBlueprintParentClass, TArray<FString>& InInterfaces, FString InUnparsedStringData, bool bInIsVersioned = true);
+	FImaginaryBlueprint(FString InBlueprintName, FString InBlueprintPath, FString InBlueprintParentClass, TArray<FString>& InInterfaces, FString InUnparsedStringData, FSearchDataVersionInfo InVersionInfo);
 
 	/** FImaginaryFiBData Interface */
 	virtual bool IsCompatibleWithFilter(ESearchQueryFilter InSearchQueryFilter) const override;
@@ -313,7 +313,7 @@ protected:
 	void ParseComponents(TSharedPtr< FJsonObject > InJsonObject, TArray<FImaginaryFiBDataSharedPtr>& OutParsedChildData);
 
 	/** Parses a raw string of Json to a Json object hierarchy */
-	void ParseToJson(bool bInIsVersioned);
+	void ParseToJson(FSearchDataVersionInfo InVersionInfo);
 
 protected:
 	/** The path for this Blueprint */
diff --git a/Engine/Source/Runtime/Engine/Private/Blueprint.cpp b/Engine/Source/Runtime/Engine/Private/Blueprint.cpp
--- a/Engine/Source/Runtime/Engine/Private/Blueprint.cpp
+++ b/Engine/Source/Runtime/Engine/Private/Blueprint.cpp
@@ -973,7 +973,14 @@ void UBlueprint::GetAssetRegistryTags(TArray<FAssetRegistryTag>& OutTags) const
 	// Only add the FiB tags in the editor, this now gets run for standalone uncooked games
 	if ( ParentClass && GIsEditor)
 	{
-		OutTags.Add(FAssetRegistryTag(FBlueprintTags::FindInBlueprintsData, FFindInBlueprintSearchManager::Get().QuerySingleBlueprint((UBlueprint*)this, false), FAssetRegistryTag::TT_Hidden));
+		FString Value;
+		const bool bRebuildSearchData = false;
+		if (const FSearchData* SearchData = FFindInBlueprintSearchManager::Get().QuerySingleBlueprint((UBlueprint*)this, bRebuildSearchData))
+		{
+			Value = SearchData->Value;
+		}
+		
+		OutTags.Add( FAssetRegistryTag(FBlueprintTags::FindInBlueprintsData, Value, FAssetRegistryTag::TT_Hidden) );
 	}
 
 	// Only show for strict blueprints (not animation or widget blueprints)
