Fix ensure when copying a UCS created component in the SCS tree

diff --git a/Engine/Source/Editor/Kismet/Private/SSCSEditor.cpp b/Engine/Source/Editor/Kismet/Private/SSCSEditor.cpp
--- a/Engine/Source/Editor/Kismet/Private/SSCSEditor.cpp
+++ b/Engine/Source/Editor/Kismet/Private/SSCSEditor.cpp
@@ -5578,7 +5578,7 @@ void SSCSEditor::CopySelectedNodes()
 		{
 			ComponentsToCopy.Add(ComponentTemplate);
 
-			if (EditorMode == EComponentEditorMode::BlueprintSCS)
+			if (EditorMode == EComponentEditorMode::BlueprintSCS && ComponentTemplate->CreationMethod != EComponentCreationMethod::UserConstructionScript)
 			{
 				// CopyComponents uses component attachment to maintain hierarchy, but the SCS templates are not
 				// setup with a relationship to each other. Briefly setup the attachment between the templates being
@@ -5605,10 +5605,13 @@ void SSCSEditor::CopySelectedNodes()
 	{
 		for (UActorComponent* ComponentTemplate : ComponentsToCopy)
 		{
-			if (USceneComponent* SceneTemplate = Cast<USceneComponent>(ComponentTemplate))
+			if (ComponentTemplate->CreationMethod != EComponentCreationMethod::UserConstructionScript)
 			{
-				// clear back out any temporary attachments we set up for the copy
-				SceneTemplate->SetupAttachment(nullptr);
+				if (USceneComponent* SceneTemplate = Cast<USceneComponent>(ComponentTemplate))
+				{
+					// clear back out any temporary attachments we set up for the copy
+					SceneTemplate->SetupAttachment(nullptr);
+				}
 			}
 		}
 	}
