Allow deleting the user construction script if a Blueprint is no longer derived from AActor

diff --git a/Engine/Source/Editor/Kismet/Private/BlueprintEditor.cpp b/Engine/Source/Editor/Kismet/Private/BlueprintEditor.cpp
--- a/Engine/Source/Editor/Kismet/Private/BlueprintEditor.cpp
+++ b/Engine/Source/Editor/Kismet/Private/BlueprintEditor.cpp
@@ -1552,7 +1552,11 @@ void FBlueprintEditor::EnsureBlueprintIsUpToDate(UBlueprint* BlueprintObj)
 		}
 
 		// If we should have a UCS but don't yet, make it
-		if(!FBlueprintEditorUtils::FindUserConstructionScript(BlueprintObj))
+		if (UEdGraph* ExistingUCS = FBlueprintEditorUtils::FindUserConstructionScript(BlueprintObj))
+		{
+			ExistingUCS->bAllowDeletion = false;
+		}
+		else
 		{
 			UEdGraph* UCSGraph = FBlueprintEditorUtils::CreateNewGraph(BlueprintObj, UEdGraphSchema_K2::FN_UserConstructionScript, UEdGraph::StaticClass(), UEdGraphSchema_K2::StaticClass());
 			FBlueprintEditorUtils::AddFunctionGraph(BlueprintObj, UCSGraph, /*bIsUserCreated=*/ false, AActor::StaticClass());
@@ -1569,12 +1573,12 @@ void FBlueprintEditor::EnsureBlueprintIsUpToDate(UBlueprint* BlueprintObj)
 	else
 	{
 		// If we have an SCS but don't support it, then we remove it
-		if(BlueprintObj->SimpleConstructionScript)
+		if (BlueprintObj->SimpleConstructionScript)
 		{
 			// Remove any SCS variable nodes
 			for (USCS_Node* SCS_Node : BlueprintObj->SimpleConstructionScript->GetAllNodes())
 			{
-				if(SCS_Node)
+				if (SCS_Node)
 				{
 					FBlueprintEditorUtils::RemoveVariableNodes(BlueprintObj, SCS_Node->GetVariableName());
 				}
@@ -1586,6 +1590,12 @@ void FBlueprintEditor::EnsureBlueprintIsUpToDate(UBlueprint* BlueprintObj)
 			// Mark the Blueprint as having been structurally modified
 			FBlueprintEditorUtils::MarkBlueprintAsStructurallyModified(BlueprintObj);
 		}
+
+		// Allow deleting the UCS if we've somehow changed away from being an actor (e.g., because of C++ reparenting the parent of our native parent)
+		if (UEdGraph* ExistingUCS = FBlueprintEditorUtils::FindUserConstructionScript(BlueprintObj))
+		{
+			ExistingUCS->bAllowDeletion = true;
+		}
 	}
 
 	// Make sure that this blueprint is up-to-date with regards to its parent functions
