Make sure transaction buffer references to archetypes get updated when reinstancine archetypes

diff --git a/Engine/Source/Editor/Kismet/Private/BlueprintCompilationManager.cpp b/Engine/Source/Editor/Kismet/Private/BlueprintCompilationManager.cpp
--- a/Engine/Source/Editor/Kismet/Private/BlueprintCompilationManager.cpp
+++ b/Engine/Source/Editor/Kismet/Private/BlueprintCompilationManager.cpp
@@ -9,6 +9,7 @@
 #include "Kismet2/CompilerResultsLog.h"
 #include "Components/TimelineComponent.h"
 #include "Editor.h"
+#include "Editor/UnrealEdEngine.h"
 #include "Engine/Engine.h"
 #include "Engine/LevelScriptBlueprint.h"
 #include "Engine/SCS_Node.h"
@@ -35,6 +36,8 @@
 #include "BlueprintEditorModule.h"
 #include "Stats/StatsHierarchical.h"
 
+extern UNREALED_API UUnrealEdEngine* GUnrealEd;
+
 #define LOCTEXT_NAMESPACE "BlueprintCompilationManager"
 
 /*
@@ -1864,6 +1867,14 @@ void FBlueprintCompilationManagerImpl::ReinstanceBatch(TArray<FReinstancingJob>&
 	// 3. Update any remaining instances that are tagged as RF_ArchetypeObject or RF_InheritableComponentTemplate - 
 	// we may need to do further sorting to ensure that interdependent archetypes are initialized correctly:
 	TSet<UObject*> ArchetypeReferencers;
+
+	// The transaction buffer could reference archetypes, and tag serialization
+	// will be simpler if we update the instance:
+	if(GUnrealEd && GUnrealEd->Trans)
+	{
+		ArchetypeReferencers.Add(GUnrealEd->Trans);
+	}
+
 	for (const FReinstancingJob& ReinstancingJob : Reinstancers)
 	{
 		UClass* OldClass = ReinstancingJob.OldToNew.Key;
