Avoid clearing transaction buffer references to private objects when compiling a blueprint
#rnx

diff --git a/Engine/Source/Editor/Kismet/Private/BlueprintCompilationManager.cpp b/Engine/Source/Editor/Kismet/Private/BlueprintCompilationManager.cpp
--- a/Engine/Source/Editor/Kismet/Private/BlueprintCompilationManager.cpp
+++ b/Engine/Source/Editor/Kismet/Private/BlueprintCompilationManager.cpp
@@ -26,7 +26,7 @@
 #include "KismetCompiler.h"
 #include "ProfilingDebugging/ScopedTimers.h"
 #include "Serialization/ArchiveHasReferences.h"
-#include "Serialization/ArchiveReplaceOrClearExternalReferences.h"
+#include "Serialization/ArchiveReplaceObjectRef.h"
 #include "Settings/EditorProjectSettings.h"
 #include "TickableEditorObject.h"
 #include "UObject/MetaData.h"
@@ -2059,8 +2059,7 @@ void FBlueprintCompilationManagerImpl::ReinstanceBatch(TArray<FReinstancingJob>&
 
 	for(UObject* ArchetypeReferencer : ArchetypeReferencers)
 	{
-		UPackage* NewPackage = ArchetypeReferencer->GetOutermost();
-		FArchiveReplaceOrClearExternalReferences<UObject> ReplaceInCDOAr(ArchetypeReferencer, OldArchetypeToNewArchetype, NewPackage);
+		FArchiveReplaceObjectRef<UObject> ReplaceInCDOAr(ArchetypeReferencer, OldArchetypeToNewArchetype, false, false, false);
 	}
 }
 
