Make it so that non-self layers are shown using functions icons in the AnimBP editor (tab icons)

diff --git a/Engine/Source/Editor/Kismet/Private/BlueprintEditor.cpp b/Engine/Source/Editor/Kismet/Private/BlueprintEditor.cpp
--- a/Engine/Source/Editor/Kismet/Private/BlueprintEditor.cpp
+++ b/Engine/Source/Editor/Kismet/Private/BlueprintEditor.cpp
@@ -551,7 +551,15 @@ const FSlateBrush* FBlueprintEditor::GetGlyphForGraph(const UEdGraph* Graph, boo
 				}
 				else
 				{
-					ReturnValue = FEditorStyle::GetBrush( bInLargeIcon ? TEXT("GraphEditor.Animation_24x") : TEXT("GraphEditor.Animation_16x") );
+					// If it has overriden an interface, show it as a function
+					if ( Graph->InterfaceGuid.IsValid() )
+					{ 
+						ReturnValue = FEditorStyle::GetBrush(bInLargeIcon ? TEXT("GraphEditor.Function_24x") : TEXT("GraphEditor.Function_16x"));
+					}
+					else
+					{
+						ReturnValue = FEditorStyle::GetBrush(bInLargeIcon ? TEXT("GraphEditor.Animation_24x") : TEXT("GraphEditor.Animation_16x"));	
+					}
 				}
 			}
 		}
