Disable the variable type dropdown on UMulticastDelegateProperty

diff --git a/Engine/Source/Editor/Kismet/Private/BlueprintDetailsCustomization.cpp b/Engine/Source/Editor/Kismet/Private/BlueprintDetailsCustomization.cpp
--- a/Engine/Source/Editor/Kismet/Private/BlueprintDetailsCustomization.cpp
+++ b/Engine/Source/Editor/Kismet/Private/BlueprintDetailsCustomization.cpp
@@ -1267,7 +1267,7 @@ void FBlueprintVarActionDetails::OnVarNameCommitted(const FText& InNewText, ETex
 bool FBlueprintVarActionDetails::GetVariableTypeChangeEnabled() const
 {
 	UProperty* VariableProperty = CachedVariableProperty.Get();
-	if(VariableProperty && IsVariableInBlueprint())
+	if(VariableProperty && !VariableProperty->IsA<UMulticastDelegateProperty>() && IsVariableInBlueprint())
 	{
 		if (!IsALocalVariable(VariableProperty))
 		{
