Added a best-guess modifier to overriding functions. If the original function was created a BP function, that will be prefered over an event implementation.

diff --git a/Engine/Source/Editor/Kismet/Private/SMyBlueprint.cpp b/Engine/Source/Editor/Kismet/Private/SMyBlueprint.cpp
--- a/Engine/Source/Editor/Kismet/Private/SMyBlueprint.cpp
+++ b/Engine/Source/Editor/Kismet/Private/SMyBlueprint.cpp
@@ -2106,7 +2106,7 @@ void SMyBlueprint::ImplementFunction(FEdGraphSchemaAction_K2Graph* GraphAction)
 	// Some types of blueprints don't have an event graph (IE gameplay ability blueprints), in that case just make a new graph, even
 	// for events:
 	UEdGraph* EventGraph = FBlueprintEditorUtils::FindEventGraph(GetBlueprintObj());
-	if (UEdGraphSchema_K2::FunctionCanBePlacedAsEvent(OverrideFunc) && EventGraph)
+	if (UEdGraphSchema_K2::FunctionCanBePlacedAsEvent(OverrideFunc) && !IsImplementationDesiredAsFunction(OverrideFunc) && EventGraph)
 	{
 		// Add to event graph
 		FName EventName = OverrideFunc->GetFName();
@@ -2143,6 +2143,27 @@ void SMyBlueprint::ImplementFunction(FEdGraphSchemaAction_K2Graph* GraphAction)
 	}
 }
 
+bool SMyBlueprint::IsImplementationDesiredAsFunction(const UFunction* OverrideFunc) const
+{	
+	// If the original function was created in a parent blueprint, then prefer a BP function
+	if (OverrideFunc)
+	{
+		FName OverrideName = *OverrideFunc->GetName();
+		TSet<FName> GraphNames;
+		FBlueprintEditorUtils::GetAllGraphNames(GetBlueprintObj(), GraphNames);
+		for (const FName & Name : GraphNames)
+		{
+			if (Name == OverrideName)
+			{
+				return true;
+			}
+		}
+	}
+	
+	// Otherwise, we would prefer an event
+	return false;
+}
+
 void SMyBlueprint::OnFindReference()
 {
 	bool bUseQuotes = true;
diff --git a/Engine/Source/Editor/Kismet/Private/SMyBlueprint.h b/Engine/Source/Editor/Kismet/Private/SMyBlueprint.h
--- a/Engine/Source/Editor/Kismet/Private/SMyBlueprint.h
+++ b/Engine/Source/Editor/Kismet/Private/SMyBlueprint.h
@@ -186,6 +186,16 @@ private:
 	/** Helper function to open the selected graph */
 	void OpenGraph(FDocumentTracker::EOpenDocumentCause InCause);
 
+	/**
+	* Check if the override of a given function is most likely desired as a blueprint function 
+	* or as an event. 
+	* 
+	* @param OverrideFunc	Desired function to override
+	* 
+	* @return	True if the function is desired as a function, false if desired as an event
+	*/
+	bool IsImplementationDesiredAsFunction(const UFunction* OverrideFunc) const;
+
 	/** Callbacks for commands */
 	void OnOpenGraph();
 	void OnOpenGraphInNewTab();
