Update references to bCanRenameNode to GetCanRenameNode()

diff --git a/Engine/Source/Editor/AnimGraph/Private/AnimationStateMachineSchema.cpp b/Engine/Source/Editor/AnimGraph/Private/AnimationStateMachineSchema.cpp
--- a/Engine/Source/Editor/AnimGraph/Private/AnimationStateMachineSchema.cpp
+++ b/Engine/Source/Editor/AnimGraph/Private/AnimationStateMachineSchema.cpp
@@ -337,7 +337,7 @@ void UAnimationStateMachineSchema::GetContextMenuActions(const UEdGraph* Current
 				MenuBuilder->AddMenuEntry( FGenericCommands::Get().Duplicate );
 				MenuBuilder->AddMenuEntry( FGraphEditorCommands::Get().ReconstructNodes );
 				MenuBuilder->AddMenuEntry( FGraphEditorCommands::Get().BreakNodeLinks );
-				if(InGraphNode->bCanRenameNode)
+				if(InGraphNode->GetCanRenameNode())
 				{
 					MenuBuilder->AddMenuEntry( FGenericCommands::Get().Rename );
 				}
diff --git a/Engine/Source/Editor/Kismet/Private/BlueprintDetailsCustomization.cpp b/Engine/Source/Editor/Kismet/Private/BlueprintDetailsCustomization.cpp
--- a/Engine/Source/Editor/Kismet/Private/BlueprintDetailsCustomization.cpp
+++ b/Engine/Source/Editor/Kismet/Private/BlueprintDetailsCustomization.cpp
@@ -6094,7 +6094,7 @@ void FBlueprintGraphNodeDetails::CustomizeDetails( IDetailLayoutBuilder& DetailL
 		}
 	}
 
-	if(!GraphNodePtr.IsValid() || !GraphNodePtr.Get()->bCanRenameNode)
+	if(!GraphNodePtr.IsValid() || !GraphNodePtr.Get()->GetCanRenameNode())
 	{
 		return;
 	}
diff --git a/Engine/Source/Editor/Kismet/Private/SBlueprintPalette.cpp b/Engine/Source/Editor/Kismet/Private/SBlueprintPalette.cpp
--- a/Engine/Source/Editor/Kismet/Private/SBlueprintPalette.cpp
+++ b/Engine/Source/Editor/Kismet/Private/SBlueprintPalette.cpp
@@ -578,7 +578,7 @@ public:
 		FEdGraphSchemaAction_K2Event* EventAction = (FEdGraphSchemaAction_K2Event*)ActionPtr.Pin().Get();
 
 		UK2Node* AssociatedNode = EventAction->NodeTemplate;
-		if (AssociatedNode && AssociatedNode->bCanRenameNode)
+		if (AssociatedNode && AssociatedNode->GetCanRenameNode())
 		{
 			TSharedPtr<INameValidatorInterface> NodeNameValidator = FNameValidatorFactory::MakeValidator(AssociatedNode);
 			bIsNameValid = (NodeNameValidator->IsValid(InNewText.ToString(), true) == EValidatorResult::Ok);
@@ -643,7 +643,7 @@ public:
 		FEdGraphSchemaAction_K2TargetNode* TargetNodeAction = (FEdGraphSchemaAction_K2TargetNode*)ActionPtr.Pin().Get();
 
 		UK2Node* AssociatedNode = TargetNodeAction->NodeTemplate;
-		if (AssociatedNode && AssociatedNode->bCanRenameNode)
+		if (AssociatedNode && AssociatedNode->GetCanRenameNode())
 		{
 			TSharedPtr<INameValidatorInterface> NodeNameValidator = FNameValidatorFactory::MakeValidator(AssociatedNode);
 			bIsNameValid = (NodeNameValidator->IsValid(InNewText.ToString(), true) == EValidatorResult::Ok);
diff --git a/Engine/Source/Editor/UnrealEd/Private/Kismet2/BlueprintEditorUtils.cpp b/Engine/Source/Editor/UnrealEd/Private/Kismet2/BlueprintEditorUtils.cpp
--- a/Engine/Source/Editor/UnrealEd/Private/Kismet2/BlueprintEditorUtils.cpp
+++ b/Engine/Source/Editor/UnrealEd/Private/Kismet2/BlueprintEditorUtils.cpp
@@ -8708,7 +8708,7 @@ bool FBlueprintEditorUtils::IsPaletteActionReadOnly(TSharedPtr<FEdGraphSchemaAct
 			FEdGraphSchemaAction_K2Event* EventAction = (FEdGraphSchemaAction_K2Event*)ActionIn.Get();
 			UK2Node* AssociatedNode = EventAction->NodeTemplate;
 
-			bIsReadOnly = (AssociatedNode == nullptr) || (!AssociatedNode->bCanRenameNode);	
+			bIsReadOnly = (AssociatedNode == nullptr) || (!AssociatedNode->GetCanRenameNode());
 		}
 		else if (ActionIn->GetTypeId() == FEdGraphSchemaAction_K2InputAction::StaticGetTypeId())
 		{
