BlueprintNativeEvents that are functions (ie have a return value) and defined in Interfaces do not show up in the function list as available to be overridden

Add the ability for BlueprintNativeEvents from Interfaces to display in child blueprints.
Add the parent class to the right of the function name in the `Override Function` menu
Add a scoped transaction in SMyBlueprint::ImplementFunction so that you can actually undo the creation of an interface function

diff --git a/Engine/Source/Editor/BlueprintGraph/Private/BlueprintEditorSettings.cpp b/Engine/Source/Editor/BlueprintGraph/Private/BlueprintEditorSettings.cpp
--- a/Engine/Source/Editor/BlueprintGraph/Private/BlueprintEditorSettings.cpp
+++ b/Engine/Source/Editor/BlueprintGraph/Private/BlueprintEditorSettings.cpp
@@ -23,6 +23,8 @@ UBlueprintEditorSettings::UBlueprintEditorSettings(const FObjectInitializer& Obj
 	, bAutoCastObjectConnections(false)
 	, bShowViewportOnSimulate(false)
 	, bShowInheritedVariables(false)
+	, bAlwaysShowInterfacesInOverrides(true)
+	, bShowParentClassInOverrides(true)
 	, bShowEmptySections(true)
 	, bSpawnDefaultBlueprintNodes(true)
 	, bHideConstructionScriptComponentsInDetailsView(true)
diff --git a/Engine/Source/Editor/BlueprintGraph/Public/BlueprintEditorSettings.h b/Engine/Source/Editor/BlueprintGraph/Public/BlueprintEditorSettings.h
--- a/Engine/Source/Editor/BlueprintGraph/Public/BlueprintEditorSettings.h
+++ b/Engine/Source/Editor/BlueprintGraph/Public/BlueprintEditorSettings.h
@@ -74,6 +74,14 @@ public:
 	UPROPERTY(config)
 	bool bShowInheritedVariables;
 
+	/** If set interface functions will always show in the overrides menu, even if they are already shown in the interfaces menu */
+	UPROPERTY(config)
+	bool bAlwaysShowInterfacesInOverrides;
+
+	/** If set then the parent class will be listed next to the override function name in the overrides function menu */
+	UPROPERTY(config)
+	bool bShowParentClassInOverrides;
+
 	/** If set we'll show empty sections in the My Blueprint view. */
 	UPROPERTY(config)
 	bool bShowEmptySections;
diff --git a/Engine/Source/Editor/Kismet/Private/SMyBlueprint.cpp b/Engine/Source/Editor/Kismet/Private/SMyBlueprint.cpp
--- a/Engine/Source/Editor/Kismet/Private/SMyBlueprint.cpp
+++ b/Engine/Source/Editor/Kismet/Private/SMyBlueprint.cpp
@@ -396,6 +396,20 @@ void SMyBlueprint::Construct(const FArguments& InArgs, TWeakPtr<FBlueprintEditor
 		TEXT("MyBlueprint_ShowReplicatedVariablesOnly")
 	);
 
+	ViewOptions.AddMenuEntry(
+		LOCTEXT("AlwaysShowInterfacesInOverrides", "Show interfaces in the function override menu"),
+		LOCTEXT("AlwaysShowInterfacesInOverridesTooltip", "Should we always display interface functions/events in the override menu?"),
+		FSlateIcon(),
+		FUIAction(
+			FExecuteAction::CreateSP(this, &SMyBlueprint::OnToggleAlwaysShowInterfacesInOverrides),
+			FCanExecuteAction(),
+			FIsActionChecked::CreateSP(this, &SMyBlueprint::GetAlwaysShowInterfacesInOverrides)
+		),
+		NAME_None,
+		EUserInterfaceActionType::ToggleButton,
+		TEXT("MyBlueprint_AlwaysShowInterfacesInOverrides")
+	);
+
 	SAssignNew(FilterBox, SSearchBox)
 		.OnTextChanged( this, &SMyBlueprint::OnFilterTextChanged );
 
@@ -849,16 +863,41 @@ void SMyBlueprint::BuildOverridableFunctionsMenu(FMenuBuilder& MenuBuilder)
 	
 	MenuBuilder.BeginSection("OverrideFunction", LOCTEXT("OverrideFunction", "Override Function"));
 	{
-		for ( auto& OverrideAction : OverridableFunctionActions )
+		for (TSharedPtr<FEdGraphSchemaAction_K2Graph>& OverrideAction : OverridableFunctionActions)
 		{
+			UClass* const OverrideFuncClass = FBlueprintEditorUtils::GetOverrideFunctionClass(GetBlueprintObj(), OverrideAction->FuncName);
+			
+			// Add the function name and tooltip 
+			TSharedRef<SHorizontalBox> FunctionBox = SNew(SHorizontalBox)
+				+ SHorizontalBox::Slot()
+				.AutoWidth()
+				.VAlign(VAlign_Center)
+				.HAlign(HAlign_Left)
+				.Padding(FMargin(2.0f, 0.0f, 20.0f, 0.0f))
+				[
+					SNew(STextBlock)
+					.Text(OverrideAction->GetMenuDescription())
+					.ToolTipText(OverrideAction->GetTooltipDescription())					
+				]
+				// Where the function came from function came from
+				+ SHorizontalBox::Slot()
+				.VAlign(VAlign_Center)
+				.HAlign(HAlign_Right)
+				.Padding(FMargin(1.0f, 0.0f, 0.0f, 0.0f))
+				[
+					SNew(STextBlock)
+					.Text(OverrideFuncClass ? OverrideFuncClass->GetDisplayNameText() : FText::GetEmpty())
+					.ToolTipText(OverrideAction->GetTooltipDescription())
+					.ColorAndOpacity(FSlateColor::UseSubduedForeground())
+				];
+
 			MenuBuilder.AddMenuEntry(
-				OverrideAction->GetMenuDescription(),
-				OverrideAction->GetTooltipDescription(),
-				FSlateIcon(),
 				FUIAction(
 					FExecuteAction::CreateSP(this, &SMyBlueprint::ImplementFunction, OverrideAction),
 					FCanExecuteAction::CreateSP(this, &SMyBlueprint::IsEditingMode)),
+				FunctionBox,
 				NAME_None,
+				OverrideAction->GetTooltipDescription(),
 				EUserInterfaceActionType::Button
 				);
 		}
@@ -1064,6 +1103,9 @@ void SMyBlueprint::CollectAllActions(FGraphActionListBuilderBase& OutAllActions)
 	// List of names of functions we implement
 	ImplementedFunctionCache.Empty();
 
+	// Fill with functions names we've already collected for rename, to ensure we do not add the same function multiple times.
+	TArray<FName> OverridableFunctionNames;
+
 	// Grab Variables
 	for (TFieldIterator<UProperty> PropertyIt(BlueprintObj->SkeletonGeneratedClass, FieldIteratorSuperFlag); PropertyIt; ++PropertyIt)
 	{
@@ -1142,6 +1184,17 @@ void SMyBlueprint::CollectAllActions(FGraphActionListBuilderBase& OutAllActions)
 		}
 	}
 
+	// Grab what events are implemented in the event graphs so they don't show up in the menu if they are already implemented
+	for (UEdGraph* const Graph : BlueprintObj->EventGraphs)
+	{
+		if (Graph && !Graph->IsUnreachable())
+		{
+			FName GraphName = Graph->GetFName();
+			ImplementedFunctionCache.Add(GraphName);
+			OverridableFunctionNames.Add(GraphName);
+		}
+	}
+
 	// Grab functions implemented by the blueprint
 	for (UEdGraph* Graph : BlueprintObj->FunctionGraphs)
 	{
@@ -1225,9 +1278,6 @@ void SMyBlueprint::CollectAllActions(FGraphActionListBuilderBase& OutAllActions)
 
 	OverridableFunctionActions.Reset();
 
-	// Fill with functions names we've already collected for rename, to ensure we do not add the same function multiple times.
-	TArray<FName> OverridableFunctionNames;
-
 	// Cache potentially overridable functions
 	UClass* ParentClass = BlueprintObj->SkeletonGeneratedClass ? BlueprintObj->SkeletonGeneratedClass->GetSuperClass() : *BlueprintObj->ParentClass;
 	for ( TFieldIterator<UFunction> FunctionIt(ParentClass, EFieldIteratorFlags::IncludeSuper); FunctionIt; ++FunctionIt )
@@ -1267,48 +1317,67 @@ void SMyBlueprint::CollectAllActions(FGraphActionListBuilderBase& OutAllActions)
 	for (int32 i=0; i < BlueprintObj->ImplementedInterfaces.Num(); i++)
 	{
 		FBPInterfaceDescription& InterfaceDesc = BlueprintObj->ImplementedInterfaces[i];
-		for (int32 FuncIdx = 0; FuncIdx < InterfaceDesc.Graphs.Num(); FuncIdx++)
+		UClass* InterfaceClass = InterfaceDesc.Interface.Get();
+		if (InterfaceClass)
 		{
-			UEdGraph* Graph = InterfaceDesc.Graphs[FuncIdx];
-			check(Graph);
-		
-			const FName FunctionName = Graph->GetFName();
-			FString FunctionTooltip = FunctionName.ToString();
-			FString FunctionDesc = FunctionName.ToString();
-			
-			FText FunctionCategory;
-			bool bIsAnimFunction = false;
-
-			if (BlueprintObj->SkeletonGeneratedClass != nullptr)
+			for (TFieldIterator<UFunction> FunctionIt(InterfaceClass, EFieldIteratorFlags::IncludeSuper); FunctionIt; ++FunctionIt)
 			{
-				if (UFunction* Function = BlueprintObj->SkeletonGeneratedClass->FindFunctionByName(Graph->GetFName()))
+				const FName FunctionName = (*FunctionIt)->GetFName();
+
+				if (FunctionName != UEdGraphSchema_K2::FN_ExecuteUbergraphBase)
 				{
-					FunctionCategory = Function->GetMetaDataText(FBlueprintMetadata::MD_FunctionCategory, TEXT("UObjectCategory"), Function->GetFullGroupName(false));
+					FString FunctionTooltip = FunctionName.ToString();
+					FString FunctionDesc = FunctionName.ToString();
+
+					FText FunctionCategory;
+					bool bIsAnimFunction = false;
 
-					if(IAnimClassInterface* AnimClassInterface = IAnimClassInterface::GetFromClass(BlueprintObj->SkeletonGeneratedClass))
+					if (BlueprintObj->SkeletonGeneratedClass != nullptr)
 					{
-						if(IAnimClassInterface::IsAnimBlueprintFunction(AnimClassInterface, Function))
+						if (UFunction * Function = BlueprintObj->SkeletonGeneratedClass->FindFunctionByName(FunctionName))
 						{
-							bIsAnimFunction = true;
+							FunctionCategory = Function->GetMetaDataText(FBlueprintMetadata::MD_FunctionCategory, TEXT("UObjectCategory"), Function->GetFullGroupName(false));
+
+							if (IAnimClassInterface * AnimClassInterface = IAnimClassInterface::GetFromClass(BlueprintObj->SkeletonGeneratedClass))
+							{
+								if (IAnimClassInterface::IsAnimBlueprintFunction(AnimClassInterface, Function))
+								{
+									bIsAnimFunction = true;
+								}
+							}
 						}
 					}
-				}
-			}
 
-			TSharedPtr<FEdGraphSchemaAction_K2Graph> NewFuncAction = MakeShareable(new FEdGraphSchemaAction_K2Graph(EEdGraphSchemaAction_K2Graph::Interface, FunctionCategory, FText::FromString(FunctionDesc), FText::FromString(FunctionTooltip), 1, bIsAnimFunction ? NodeSectionID::ANIMLAYER : NodeSectionID::INTERFACE));
-			NewFuncAction->FuncName = FunctionName;
-			NewFuncAction->EdGraph = Graph;
-			OutAllActions.AddAction(NewFuncAction);
+					TSharedPtr<FEdGraphSchemaAction_K2Graph> NewFuncAction = MakeShareable(new FEdGraphSchemaAction_K2Graph(EEdGraphSchemaAction_K2Graph::Interface, FunctionCategory, FText::FromString(FunctionDesc), FText::FromString(FunctionTooltip), 1, bIsAnimFunction ? NodeSectionID::ANIMLAYER : NodeSectionID::INTERFACE));
+
+					NewFuncAction->FuncName = FunctionName;
+					OutAllActions.AddAction(NewFuncAction);
 
-			GetChildGraphs(Graph, NewFuncAction->GetSectionID(), SortList, FunctionCategory);
-			GetChildEvents(Graph, NewFuncAction->GetSectionID(), SortList, FunctionCategory);
+					// Find the graph that this function is on so the user can double click and open it from the interfaces menu
+					for (UEdGraph* const Graph : InterfaceDesc.Graphs)
+					{
+						if (Graph && Graph->GetFName() == FunctionName)
+						{
+							NewFuncAction->EdGraph = Graph;
+							break;
+						}
+					}
+
+					// if this function is not in the interfaces menu, then allow it to be put in the override function menu
+					if (GetAlwaysShowInterfacesInOverrides())
+					{
+						OverridableFunctionActions.Add(NewFuncAction);
+						OverridableFunctionNames.Add(FunctionName);
+					}
+				}
+			}
 		}
 	}
 
 	// also walk up the class chain to look for overridable functions in natively implemented interfaces
-	for ( UClass* TempClass=BlueprintObj->ParentClass; TempClass; TempClass=TempClass->GetSuperClass() )
+	for (UClass* TempClass = BlueprintObj->ParentClass; TempClass; TempClass = TempClass->GetSuperClass())
 	{
-		for (int32 Idx=0; Idx<TempClass->Interfaces.Num(); ++Idx)
+		for (int32 Idx = 0; Idx < TempClass->Interfaces.Num(); ++Idx)
 		{
 			FImplementedInterface const& I = TempClass->Interfaces[Idx];
 			if (!I.bImplementedByK2)
@@ -1319,7 +1388,7 @@ void SMyBlueprint::CollectAllActions(FGraphActionListBuilderBase& OutAllActions)
 					const UFunction* Function = *FunctionIt;
 					const FName FunctionName = Function->GetFName();
 
-					if ( UEdGraphSchema_K2::CanKismetOverrideFunction(Function) && !ImplementedFunctionCache.Contains(FunctionName) && !UEdGraphSchema_K2::FunctionCanBePlacedAsEvent(Function) )
+					if (UEdGraphSchema_K2::CanKismetOverrideFunction(Function) && !ImplementedFunctionCache.Contains(FunctionName) && !UEdGraphSchema_K2::FunctionCanBePlacedAsEvent(Function))
 					{
 						FText FunctionTooltip = Function->GetToolTipText();
 						FText FunctionDesc = K2Schema->GetFriendlySignatureName(Function);
@@ -1328,6 +1397,8 @@ void SMyBlueprint::CollectAllActions(FGraphActionListBuilderBase& OutAllActions)
 
 						TSharedPtr<FEdGraphSchemaAction_K2Graph> NewFuncAction = MakeShareable(new FEdGraphSchemaAction_K2Graph(EEdGraphSchemaAction_K2Graph::Function, FunctionCategory, FunctionDesc, FunctionTooltip, 1, NodeSectionID::INTERFACE));
 						NewFuncAction->FuncName = FunctionName;
+						OverridableFunctionActions.Add(NewFuncAction);
+						OverridableFunctionNames.Add(FunctionName);
 						OutAllActions.AddAction(NewFuncAction);
 					}
 				}
@@ -1456,6 +1527,34 @@ void SMyBlueprint::OnToggleShowReplicatedVariablesOnly()
 	Refresh();
 }
 
+void SMyBlueprint::OnToggleAlwaysShowInterfacesInOverrides()
+{
+	UBlueprintEditorSettings* Settings = GetMutableDefault<UBlueprintEditorSettings>();
+	Settings->bAlwaysShowInterfacesInOverrides = !Settings->bAlwaysShowInterfacesInOverrides;
+	Settings->PostEditChange();
+	Settings->SaveConfig();
+	Refresh();
+}
+
+bool SMyBlueprint::GetAlwaysShowInterfacesInOverrides() const
+{
+	return GetMutableDefault<UBlueprintEditorSettings>()->bAlwaysShowInterfacesInOverrides;
+}
+
+void SMyBlueprint::OnToggleShowParentClassInOverrides()
+{
+	UBlueprintEditorSettings* Settings = GetMutableDefault<UBlueprintEditorSettings>();
+	Settings->bShowParentClassInOverrides = !Settings->bShowParentClassInOverrides;
+	Settings->PostEditChange();
+	Settings->SaveConfig();
+	Refresh();
+}
+
+bool SMyBlueprint::GetShowParentClassInOverrides() const
+{
+	return GetMutableDefault<UBlueprintEditorSettings>()->bShowParentClassInOverrides;
+}
+
 bool SMyBlueprint::IsShowingReplicatedVariablesOnly() const
 {
 	return bShowReplicatedVariablesOnly;
@@ -2065,57 +2164,20 @@ void SMyBlueprint::ImplementFunction(TSharedPtr<FEdGraphSchemaAction_K2Graph> Gr
 
 void SMyBlueprint::ImplementFunction(FEdGraphSchemaAction_K2Graph* GraphAction)
 {
-	check(GetBlueprintObj()->SkeletonGeneratedClass);
-	const UFunction* OverrideFunc = FindField<UFunction>(GetBlueprintObj()->SkeletonGeneratedClass, GraphAction->FuncName);
-
-	// search up the class hiearchy, we want to find the original declaration of the function to match FBlueprintEventNodeSpawner.
-	// Doing so ensures that we can find the existing node if there is one:
-	const UClass* Iter = GetBlueprintObj()->SkeletonGeneratedClass->GetSuperClass();
-	while (Iter != nullptr)
-	{
-		if (const UFunction* F = Iter->FindFunctionByName(GraphAction->FuncName))
-		{
-			OverrideFunc = F;
-		}
-		else
-		{
-			break;
-		}
-		Iter = Iter->GetSuperClass();
-	}
-
+	UBlueprint* BlueprintObj = GetBlueprintObj();
+	check(BlueprintObj && BlueprintObj->SkeletonGeneratedClass);
 
-	if (OverrideFunc == nullptr)
-	{
-		// maybe it's from a native interface, check those too
-		for ( UClass* TempClass=GetBlueprintObj()->ParentClass; (nullptr != TempClass) && (nullptr == OverrideFunc); TempClass=TempClass->GetSuperClass() )
-		{
-			for (int32 Idx=0; Idx<TempClass->Interfaces.Num(); ++Idx)
-			{
-				FImplementedInterface const& I = TempClass->Interfaces[Idx];
-				if (!I.bImplementedByK2)
-				{
-					OverrideFunc = FindField<UFunction>(I.Class, GraphAction->FuncName);
-					if (OverrideFunc)
-					{
-						// found it, done
-						break;
-					}
-				}
-			}
-		}
-	}
+	UFunction* OverrideFunc = nullptr;
+	UClass* const OverrideFuncClass = FBlueprintEditorUtils::GetOverrideFunctionClass(BlueprintObj, GraphAction->FuncName, &OverrideFunc);
 	check(OverrideFunc);
-	UClass* const OverrideFuncClass = CastChecked<UClass>(OverrideFunc->GetOuter())->GetAuthoritativeClass();
-
 	// Some types of blueprints don't have an event graph (IE gameplay ability blueprints), in that case just make a new graph, even
 	// for events:
-	UEdGraph* EventGraph = FBlueprintEditorUtils::FindEventGraph(GetBlueprintObj());
+	UEdGraph* EventGraph = FBlueprintEditorUtils::FindEventGraph(BlueprintObj);
 	if (UEdGraphSchema_K2::FunctionCanBePlacedAsEvent(OverrideFunc) && !IsImplementationDesiredAsFunction(OverrideFunc) && EventGraph)
 	{
 		// Add to event graph
 		FName EventName = OverrideFunc->GetFName();
-		UK2Node_Event* ExistingNode = FBlueprintEditorUtils::FindOverrideForFunction(GetBlueprintObj(), OverrideFuncClass, EventName);
+		UK2Node_Event* ExistingNode = FBlueprintEditorUtils::FindOverrideForFunction(BlueprintObj, OverrideFuncClass, EventName);
 
 		if (ExistingNode)
 		{
@@ -2141,10 +2203,23 @@ void SMyBlueprint::ImplementFunction(FEdGraphSchemaAction_K2Graph* GraphAction)
 	}
 	else
 	{
-		// Implement the function graph
-		UEdGraph* const NewGraph = FBlueprintEditorUtils::CreateNewGraph(GetBlueprintObj(), GraphAction->FuncName, UEdGraph::StaticClass(), UEdGraphSchema_K2::StaticClass());
-		FBlueprintEditorUtils::AddFunctionGraph(GetBlueprintObj(), NewGraph, /*bIsUserCreated=*/ false, OverrideFuncClass);
-		BlueprintEditorPtr.Pin()->OpenDocument(NewGraph, FDocumentTracker::OpenNewDocument);
+		// If there is an already existing graph of this function than just open that
+		// Needed for implementing interface functions on the base class through the override menu
+		UEdGraph* const ExistingGraph = FindObject<UEdGraph>(BlueprintObj, *GraphAction->FuncName.ToString());
+		if (ExistingGraph)
+		{
+			FKismetEditorUtilities::BringKismetToFocusAttentionOnObject(ExistingGraph);
+		}
+		else
+		{
+			const FScopedTransaction Transaction(LOCTEXT("CreateOverrideFunctionGraph", "Create Override Function Graph"));
+			BlueprintObj->Modify();
+			// Implement the function graph
+			UEdGraph* const NewGraph = FBlueprintEditorUtils::CreateNewGraph(BlueprintObj, GraphAction->FuncName, UEdGraph::StaticClass(), UEdGraphSchema_K2::StaticClass());
+			FBlueprintEditorUtils::AddFunctionGraph(BlueprintObj, NewGraph, /*bIsUserCreated=*/ false, OverrideFuncClass);
+			NewGraph->Modify();
+			BlueprintEditorPtr.Pin()->OpenDocument(NewGraph, FDocumentTracker::OpenNewDocument);
+		}
 	}
 }
 
diff --git a/Engine/Source/Editor/Kismet/Private/SMyBlueprint.h b/Engine/Source/Editor/Kismet/Private/SMyBlueprint.h
--- a/Engine/Source/Editor/Kismet/Private/SMyBlueprint.h
+++ b/Engine/Source/Editor/Kismet/Private/SMyBlueprint.h
@@ -183,6 +183,14 @@ private:
 	void OnToggleShowReplicatedVariablesOnly();
 	bool IsShowingReplicatedVariablesOnly() const;
 
+	/** Support functions for view options for bAlwaysShowInterfacesInOverrides blueprint editor setting */
+	void OnToggleAlwaysShowInterfacesInOverrides();
+	bool GetAlwaysShowInterfacesInOverrides() const;
+
+	/** Support functions for view options for bShowParentClassInOverrides blueprint editor setting*/
+	void OnToggleShowParentClassInOverrides();
+	bool GetShowParentClassInOverrides() const;
+
 	/** Helper function to open the selected graph */
 	void OpenGraph(FDocumentTracker::EOpenDocumentCause InCause);
 
