Remove the variable type field from event dispatcher entirely.

diff --git a/Engine/Source/Editor/Kismet/Private/BlueprintDetailsCustomization.cpp b/Engine/Source/Editor/Kismet/Private/BlueprintDetailsCustomization.cpp
--- a/Engine/Source/Editor/Kismet/Private/BlueprintDetailsCustomization.cpp
+++ b/Engine/Source/Editor/Kismet/Private/BlueprintDetailsCustomization.cpp
@@ -284,28 +284,31 @@ void FBlueprintVarActionDetails::CustomizeDetails( IDetailLayoutBuilder& DetailL
 		.Font(IDetailLayoutBuilder::GetDetailFont())
 	];
 
-	TSharedPtr<SToolTip> VarTypeTooltip = IDocumentation::Get()->CreateToolTip(LOCTEXT("VarTypeTooltip", "The type of the variable."), NULL, DocLink, TEXT("VariableType"));
+	if (GetVariableTypeChangeEnabled())
+	{
+		TSharedPtr<SToolTip> VarTypeTooltip = IDocumentation::Get()->CreateToolTip(LOCTEXT("VarTypeTooltip", "The type of the variable."), NULL, DocLink, TEXT("VariableType"));
 
-	Category.AddCustomRow( LOCTEXT("VariableTypeLabel", "Variable Type") )
-	.NameContent()
-	[
-		SNew(STextBlock)
-		.Text( LOCTEXT("VariableTypeLabel", "Variable Type") )
-		.ToolTip(VarTypeTooltip)
-		.Font(DetailFontInfo)
-	]
-	.ValueContent()
-	.MaxDesiredWidth(980.f)
-	[
-		SNew(SPinTypeSelector, FGetPinTypeTree::CreateUObject(Schema, &UEdGraphSchema_K2::GetVariableTypeTree))
-		.TargetPinType(this, &FBlueprintVarActionDetails::OnGetVarType)
-		.OnPinTypeChanged(this, &FBlueprintVarActionDetails::OnVarTypeChanged)
-		.IsEnabled(this, &FBlueprintVarActionDetails::GetVariableTypeChangeEnabled)
-		.Schema(Schema)
-		.TypeTreeFilter(ETypeTreeFilter::None)
-		.Font( DetailFontInfo )
-		.ToolTip(VarTypeTooltip)
-	];
+		Category.AddCustomRow(LOCTEXT("VariableTypeLabel", "Variable Type"))
+			.NameContent()
+			[
+				SNew(STextBlock)
+				.Text(LOCTEXT("VariableTypeLabel", "Variable Type"))
+			.ToolTip(VarTypeTooltip)
+			.Font(DetailFontInfo)
+			]
+		.ValueContent()
+			.MaxDesiredWidth(980.f)
+			[
+				SNew(SPinTypeSelector, FGetPinTypeTree::CreateUObject(Schema, &UEdGraphSchema_K2::GetVariableTypeTree))
+				.TargetPinType(this, &FBlueprintVarActionDetails::OnGetVarType)
+			.OnPinTypeChanged(this, &FBlueprintVarActionDetails::OnVarTypeChanged)
+			.IsEnabled(this, &FBlueprintVarActionDetails::GetVariableTypeChangeEnabled)
+			.Schema(Schema)
+			.TypeTreeFilter(ETypeTreeFilter::None)
+			.Font(DetailFontInfo)
+			.ToolTip(VarTypeTooltip)
+			];
+	}
 
 	TSharedPtr<SToolTip> EditableTooltip = IDocumentation::Get()->CreateToolTip(LOCTEXT("VarEditableTooltip", "Whether this variable is publicly editable on instances of this Blueprint."), NULL, DocLink, TEXT("Editable"));
 
