Add the option to re-connect a node's exec/then pins when you press "Shift+DEL"
Add an editor action for "shift+del"
Only works for nodes with simple output/exec nodes

diff --git a/Engine/Source/Editor/GraphEditor/Private/GraphEditorActions.cpp b/Engine/Source/Editor/GraphEditor/Private/GraphEditorActions.cpp
--- a/Engine/Source/Editor/GraphEditor/Private/GraphEditorActions.cpp
+++ b/Engine/Source/Editor/GraphEditor/Private/GraphEditorActions.cpp
@@ -24,6 +24,8 @@ void FGraphEditorCommandsImpl::RegisterCommands()
 
 	UI_COMMAND( ChangePinType, "Change Pin Type", "Changes the type of this pin (boolean, int, etc.)", EUserInterfaceActionType::Button, FInputChord() )
 
+	UI_COMMAND( DeleteAndReconnectNodes, "Delete and reconnect exec pins", "Deletes the currently selected nodes and makes connections between all input pins to their output pins.", EUserInterfaceActionType::Button, FInputChord(EModifierKey::Shift, EKeys::Delete) )
+
 	UI_COMMAND( ShowAllPins, "Show All Pins", "Shows all pins", EUserInterfaceActionType::RadioButton, FInputChord() )
 	UI_COMMAND( HideNoConnectionPins, "Hide Unconnected Pins", "Hides all pins with no connections", EUserInterfaceActionType::RadioButton, FInputChord() )
 	UI_COMMAND( HideNoConnectionNoDefaultPins, "Hide Unused Pins", "Hides all pins with no connections and no default value", EUserInterfaceActionType::RadioButton, FInputChord() )
diff --git a/Engine/Source/Editor/GraphEditor/Public/GraphEditorActions.h b/Engine/Source/Editor/GraphEditor/Public/GraphEditorActions.h
--- a/Engine/Source/Editor/GraphEditor/Public/GraphEditorActions.h
+++ b/Engine/Source/Editor/GraphEditor/Public/GraphEditorActions.h
@@ -39,6 +39,7 @@ public:
 	TSharedPtr< FUICommandInfo > AddOptionPin;
 	TSharedPtr< FUICommandInfo > RemoveOptionPin;
 	TSharedPtr< FUICommandInfo > ChangePinType;
+	TSharedPtr< FUICommandInfo > DeleteAndReconnectNodes;
 
 	// Pin visibility modes
 	TSharedPtr< FUICommandInfo > ShowAllPins;
diff --git a/Engine/Source/Editor/Kismet/Private/BlueprintEditor.cpp b/Engine/Source/Editor/Kismet/Private/BlueprintEditor.cpp
--- a/Engine/Source/Editor/Kismet/Private/BlueprintEditor.cpp
+++ b/Engine/Source/Editor/Kismet/Private/BlueprintEditor.cpp
@@ -1221,6 +1221,11 @@ TSharedRef<SGraphEditor> FBlueprintEditor::CreateGraphEditorWidget(TSharedRef<FT
 				FCanExecuteAction::CreateSP( this, &FBlueprintEditor::CanDeleteNodes )
 				);
 
+			GraphEditorCommands->MapAction(FGraphEditorCommands::Get().DeleteAndReconnectNodes,
+				FExecuteAction::CreateSP(this, &FBlueprintEditor::DeleteSelectedNodes),
+				FCanExecuteAction::CreateSP(this, &FBlueprintEditor::CanDeleteNodes)
+			);
+
 			GraphEditorCommands->MapAction( FGenericCommands::Get().Copy,
 				FExecuteAction::CreateSP( this, &FBlueprintEditor::CopySelectedNodes ),
 				FCanExecuteAction::CreateSP( this, &FBlueprintEditor::CanCopyNodes )
@@ -5482,6 +5487,13 @@ void FBlueprintEditor::DeleteSelectedNodes()
 					DocumentManager->CleanInvalidTabs();
 				}
 				AnalyticsTrackNodeEvent( GetBlueprintObj(), Node, true );
+				
+				// If the user is pressing shift then try and reconnect the pins
+				if (FSlateApplication::Get().GetModifierKeys().IsShiftDown())
+				{
+					ReconnectExecPins(K2Node);
+				}
+				
 				FBlueprintEditorUtils::RemoveNode(GetBlueprintObj(), Node, true);
 			}
 		}
@@ -5499,6 +5511,32 @@ void FBlueprintEditor::DeleteSelectedNodes()
 	//@TODO: Reselect items that were not deleted
 }
 
+void FBlueprintEditor::ReconnectExecPins(UK2Node* Node)
+{
+	// Only do this on impure nodes
+	if (Node && !Node->IsNodePure())
+	{
+		// Are there exec/then connections?
+		UEdGraphPin* const ExecPin = Node->GetExecPin();
+		UEdGraphPin* const ThenPin = Node->FindPin(UEdGraphSchema_K2::PN_Then);
+		// Nodes with multiple "then" pins (branch, sequence, foreach, etc) will actually have the FindPin return nullptr
+		if (ExecPin && ThenPin)
+		{
+			// Make a connection from every incoming exec pin to every outgoing then pin
+			for (UEdGraphPin* const IncomingConnectionPin : ExecPin->LinkedTo)
+			{
+				if (IncomingConnectionPin)
+				{
+					for (UEdGraphPin* const ConnectedCompletePin : ThenPin->LinkedTo)
+					{
+						IncomingConnectionPin->MakeLinkTo(ConnectedCompletePin);
+					}
+				}
+			}
+		}
+	}
+}
+
 bool FBlueprintEditor::CanDeleteNodes() const
 {
 	const FGraphPanelSelectionSet SelectedNodes = GetSelectedNodes();
diff --git a/Engine/Source/Editor/Kismet/Public/BlueprintEditor.h b/Engine/Source/Editor/Kismet/Public/BlueprintEditor.h
--- a/Engine/Source/Editor/Kismet/Public/BlueprintEditor.h
+++ b/Engine/Source/Editor/Kismet/Public/BlueprintEditor.h
@@ -833,6 +833,14 @@ protected:
 	virtual void DeleteSelectedNodes();
 	bool CanDeleteNodes() const;
 
+	/**
+	* Given a node, make connections from anything connected to it's input pin to
+	* anything connected to it's "then" pins. Only works on impure nodes with single exec/then pins.
+	* 
+	* @param Node	The node to reconnect
+	*/
+	void ReconnectExecPins(class UK2Node* Node);
+
 	void DeleteSelectedDuplicatableNodes();
 
 	void CutSelectedNodes();
