PR #6090: Fixed a copy-paste typo: bIsFliteredOut -> bIsFilteredOut (Contributed by [REDACTED_PERSON])
#rnx

diff --git a/Engine/Source/Editor/Kismet/Private/BlueprintActionMenuUtils.cpp b/Engine/Source/Editor/Kismet/Private/BlueprintActionMenuUtils.cpp
--- a/Engine/Source/Editor/Kismet/Private/BlueprintActionMenuUtils.cpp
+++ b/Engine/Source/Editor/Kismet/Private/BlueprintActionMenuUtils.cpp
@@ -158,21 +158,21 @@ static bool BlueprintActionMenuUtilsImpl::IsNonFavoritedAction(FBlueprintActionF
 //------------------------------------------------------------------------------
 static bool BlueprintActionMenuUtilsImpl::IsPureNonConstAction(FBlueprintActionFilter const& Filter, FBlueprintActionInfo& BlueprintAction)
 {
-	bool bIsFliteredOut = false;
+	bool bIsFilteredOut = false;
 
 	if (UFunction const* Function = BlueprintAction.GetAssociatedFunction())
 	{
 		bool const bIsImperative = !Function->HasAnyFunctionFlags(FUNC_BlueprintPure);
 		bool const bIsConstFunc  =  Function->HasAnyFunctionFlags(FUNC_Const);
-		bIsFliteredOut = !bIsImperative && !bIsConstFunc;
+		bIsFilteredOut = !bIsImperative && !bIsConstFunc;
 	}
-	return bIsFliteredOut;
+	return bIsFilteredOut;
 }
 
 //------------------------------------------------------------------------------
 static bool BlueprintActionMenuUtilsImpl::IsUnexposedMemberAction(FBlueprintActionFilter const& Filter, FBlueprintActionInfo& BlueprintAction)
 {
-	bool bIsFliteredOut = false;
+	bool bIsFilteredOut = false;
 
 	if (UFunction const* Function = BlueprintAction.GetAssociatedFunction())
 	{
@@ -194,15 +194,15 @@ static bool BlueprintActionMenuUtilsImpl::IsUnexposedMemberAction(FBlueprintActi
 		}
 
 		const FString& FunctionCategory = Function->GetMetaData(FBlueprintMetadata::MD_FunctionCategory);
-		bIsFliteredOut = !AllExposedCategories.Contains(FunctionCategory);
+		bIsFilteredOut = !AllExposedCategories.Contains(FunctionCategory);
 	}
-	return bIsFliteredOut;
+	return bIsFilteredOut;
 }
 
 //------------------------------------------------------------------------------
 static bool BlueprintActionMenuUtilsImpl::IsUnexposedNonComponentAction(FBlueprintActionFilter const& Filter, FBlueprintActionInfo& BlueprintAction)
 {
-	bool bIsFliteredOut = false;
+	bool bIsFilteredOut = false;
 
 	for (TWeakObjectPtr<UObject> Binding : BlueprintAction.GetBindings())
 	{
@@ -217,11 +217,11 @@ static bool BlueprintActionMenuUtilsImpl::IsUnexposedNonComponentAction(FBluepri
 		}
 		// else, it's not a component... let's do this!
 
-		bIsFliteredOut = IsUnexposedMemberAction(Filter, BlueprintAction);
+		bIsFilteredOut = IsUnexposedMemberAction(Filter, BlueprintAction);
 		break;
 	}
 
-	return bIsFliteredOut;
+	return bIsFilteredOut;
 }
 
 //------------------------------------------------------------------------------
