Add "Convert function to event" option for function nodes and a "convert event to function" option for events

diff --git a/Engine/Source/Editor/BlueprintGraph/Private/EdGraphSchema_K2.cpp b/Engine/Source/Editor/BlueprintGraph/Private/EdGraphSchema_K2.cpp
--- a/Engine/Source/Editor/BlueprintGraph/Private/EdGraphSchema_K2.cpp
+++ b/Engine/Source/Editor/BlueprintGraph/Private/EdGraphSchema_K2.cpp
@@ -1661,6 +1661,16 @@ void UEdGraphSchema_K2::GetContextMenuActions(const UEdGraph* CurrentGraph, cons
 					MenuBuilder->AddMenuEntry( FGraphEditorCommands::Get().CollapseSelectionToMacro );
 					MenuBuilder->AddMenuEntry( FGraphEditorCommands::Get().ExpandNodes );
 
+					if (InGraphNode->IsA(UK2Node_FunctionEntry::StaticClass()))
+					{
+						MenuBuilder->AddMenuEntry(FGraphEditorCommands::Get().ConvertFunctionToEvent);
+					}
+
+					if (InGraphNode->IsA(UK2Node_Event::StaticClass()))
+					{
+						MenuBuilder->AddMenuEntry(FGraphEditorCommands::Get().ConvertEventToFunction);
+					}
+
 					if(InGraphNode->IsA(UK2Node_Composite::StaticClass()))
 					{
 						MenuBuilder->AddMenuEntry( FGraphEditorCommands::Get().PromoteSelectionToFunction );
diff --git a/Engine/Source/Editor/GraphEditor/Private/GraphEditorActions.cpp b/Engine/Source/Editor/GraphEditor/Private/GraphEditorActions.cpp
--- a/Engine/Source/Editor/GraphEditor/Private/GraphEditorActions.cpp
+++ b/Engine/Source/Editor/GraphEditor/Private/GraphEditorActions.cpp
@@ -44,6 +44,8 @@ void FGraphEditorCommandsImpl::RegisterCommands()
 	UI_COMMAND( ExpandNodes, "Expand Node", "Expands the node's internal graph into the current graph and removes this node.", EUserInterfaceActionType::Button, FInputChord() )
 	UI_COMMAND( CollapseSelectionToFunction, "Collapse to Function", "Collapses selected nodes into a single function node.", EUserInterfaceActionType::Button, FInputChord() )
 	UI_COMMAND( CollapseSelectionToMacro, "Collapse to Macro", "Collapses selected nodes into a single macro node.", EUserInterfaceActionType::Button, FInputChord() )
+	UI_COMMAND( ConvertFunctionToEvent, "Convert function to event", "Converts selected function to an event and removes the function defintion.", EUserInterfaceActionType::Button, FInputChord() )
+	UI_COMMAND( ConvertEventToFunction, "Convert event to function", "Converts the selected event to a function graph and removes this event node.", EUserInterfaceActionType::Button, FInputChord() )
 
 	UI_COMMAND( AlignNodesTop, "Align Top", "Aligns the top edges of the selected nodes", EUserInterfaceActionType::Button, FInputChord(EModifierKey::Shift, EKeys::W) )
 	UI_COMMAND( AlignNodesMiddle, "Align Middle", "Aligns the vertical middles of the selected nodes", EUserInterfaceActionType::Button, FInputChord(EModifierKey::Shift|EModifierKey::Alt, EKeys::W) )
diff --git a/Engine/Source/Editor/GraphEditor/Public/GraphEditorActions.h b/Engine/Source/Editor/GraphEditor/Public/GraphEditorActions.h
--- a/Engine/Source/Editor/GraphEditor/Public/GraphEditorActions.h
+++ b/Engine/Source/Editor/GraphEditor/Public/GraphEditorActions.h
@@ -63,6 +63,8 @@ public:
 	TSharedPtr< FUICommandInfo > ExpandNodes;
 	TSharedPtr< FUICommandInfo > CollapseSelectionToFunction;
 	TSharedPtr< FUICommandInfo > CollapseSelectionToMacro;
+	TSharedPtr< FUICommandInfo > ConvertFunctionToEvent;
+	TSharedPtr< FUICommandInfo > ConvertEventToFunction;
 
 	// Alignment commands
 	TSharedPtr< FUICommandInfo > AlignNodesTop;
diff --git a/Engine/Source/Editor/Kismet/Private/BlueprintEditor.cpp b/Engine/Source/Editor/Kismet/Private/BlueprintEditor.cpp
--- a/Engine/Source/Editor/Kismet/Private/BlueprintEditor.cpp
+++ b/Engine/Source/Editor/Kismet/Private/BlueprintEditor.cpp
@@ -1146,6 +1146,18 @@ TSharedRef<SGraphEditor> FBlueprintEditor::CreateGraphEditorWidget(TSharedRef<FT
 				FIsActionButtonVisible::CreateSP( this, &FBlueprintEditor::NewDocument_IsVisibleForType, CGT_NewMacroGraph )
 				);
 
+			GraphEditorCommands->MapAction(FGraphEditorCommands::Get().ConvertFunctionToEvent,
+				FExecuteAction::CreateSP(this, &FBlueprintEditor::OnConvertFunctionToEvent),
+				FCanExecuteAction::CreateSP(this, &FBlueprintEditor::CanConvertFunctionToEvent),
+				FIsActionChecked()
+			);
+
+			GraphEditorCommands->MapAction(FGraphEditorCommands::Get().ConvertEventToFunction,
+				FExecuteAction::CreateSP(this, &FBlueprintEditor::OnConvertEventToFunction),
+				FCanExecuteAction::CreateSP(this, &FBlueprintEditor::CanConvertEventToFunction),
+				FIsActionChecked()
+			);
+
 			GraphEditorCommands->MapAction( FGraphEditorCommands::Get().PromoteSelectionToFunction,
 				FExecuteAction::CreateSP( this, &FBlueprintEditor::OnPromoteSelectionToFunction ),
 				FCanExecuteAction::CreateSP( this, &FBlueprintEditor::CanPromoteSelectionToFunction ),
@@ -5241,7 +5253,7 @@ void FBlueprintEditor::OnExpandNodes()
 			if(MacroGraph)
 			{
 				// Clone the graph so that we do not delete the original
-				UEdGraph* ClonedGraph = FEdGraphUtilities::CloneGraph(MacroGraph, NULL);
+				UEdGraph* ClonedGraph = FEdGraphUtilities::CloneGraph(MacroGraph, nullptr);
 				ExpandNode(SelectedMacroInstanceNode, ClonedGraph, /*inout*/ ExpandedNodes);
 
 				ClonedGraph->MarkPendingKill();
@@ -5260,7 +5272,7 @@ void FBlueprintEditor::OnExpandNodes()
 		}
 		else if (UK2Node_CallFunction* SelectedCallFunctionNode = Cast<UK2Node_CallFunction>(*NodeIt))
 		{
-			const UEdGraphNode* ResultEventNode = NULL;
+			const UEdGraphNode* ResultEventNode = nullptr;
 			UEdGraph* FunctionGraph = SelectedCallFunctionNode->GetFunctionGraph(ResultEventNode);
 
 			// We should never get here when attempting to expand a call function that calls an event.
@@ -5269,72 +5281,507 @@ void FBlueprintEditor::OnExpandNodes()
 			if(FunctionGraph)
 			{
 				// Clone the graph so that we do not delete the original
-				UEdGraph* ClonedGraph = FEdGraphUtilities::CloneGraph(FunctionGraph, NULL);
+				UEdGraph* ClonedGraph = FEdGraphUtilities::CloneGraph(FunctionGraph, nullptr);
 				ExpandNode(SelectedCallFunctionNode, ClonedGraph, ExpandedNodes);
 
 				ClonedGraph->MarkPendingKill();
 			}
 		}
+		UEdGraphNode* SourceNode = CastChecked<UEdGraphNode>(*NodeIt);
+		check(SourceNode);
+		MoveNodesToAveragePos(ExpandedNodes, FVector2D(SourceNode->NodePosX, SourceNode->NodePosY), bExpandedNodesNeedUniqueGuid);
+	}
+
+	FBlueprintEditorUtils::MarkBlueprintAsStructurallyModified(GetBlueprintObj());
+}
+
+bool FBlueprintEditor::CanExpandNodes() const
+{
+	// Does the selection set contain any composite nodes that are legal to expand?
+	const FGraphPanelSelectionSet SelectedNodes = GetSelectedNodes();
+	for (FGraphPanelSelectionSet::TConstIterator NodeIt(SelectedNodes); NodeIt; ++NodeIt)
+	{
+		if (Cast<UK2Node_Composite>(*NodeIt))
+		{
+			return true;
+		}
+		else if (UK2Node_MacroInstance* SelectedMacroInstanceNode = Cast<UK2Node_MacroInstance>(*NodeIt))
+		{
+			return SelectedMacroInstanceNode->GetMacroGraph() != nullptr;
+		}
+		else if (UK2Node_CallFunction* SelectedCallFunctionNode = Cast<UK2Node_CallFunction>(*NodeIt))
+		{
+			// If ResultEventNode is non-NULL, it means it is sourced by an event, we do not want to expand events
+			const UEdGraphNode* ResultEventNode = nullptr;
+			return SelectedCallFunctionNode->GetFunctionGraph(ResultEventNode) != nullptr && ResultEventNode == nullptr;
+		}
+	}
+
+	return false;
+}
+
+void FBlueprintEditor::MoveNodesToAveragePos(TSet<UEdGraphNode*>& AverageNodes, FVector2D SourcePos, bool bExpandedNodesNeedUniqueGuid /* = false */) const
+{
+	if (AverageNodes.Num() > 0)
+	{
+		FVector2D AvgNodePosition(0.0f, 0.0f);
+
+		for (TSet<UEdGraphNode*>::TIterator It(AverageNodes); It; ++It)
+		{
+			UEdGraphNode* Node = *It;
+			AvgNodePosition.X += Node->NodePosX;
+			AvgNodePosition.Y += Node->NodePosY;
+		}
+
+		float InvNumNodes = 1.0f / float(AverageNodes.Num());
+		AvgNodePosition.X *= InvNumNodes;
+		AvgNodePosition.Y *= InvNumNodes;
+
+		TSharedPtr<SGraphEditor> FocusedGraphEd = FocusedGraphEdPtr.Pin();
 
-		if (ExpandedNodes.Num() > 0)
+		for (UEdGraphNode* ExpandedNode : AverageNodes)
 		{
-			FVector2D AvgNodePosition(0.0f, 0.0f);
+			ExpandedNode->NodePosX = (ExpandedNode->NodePosX - AvgNodePosition.X) + SourcePos.X;
+			ExpandedNode->NodePosY = (ExpandedNode->NodePosY - AvgNodePosition.Y) + SourcePos.Y;
+
+			ExpandedNode->SnapToGrid(SNodePanel::GetSnapGridSize());
 
-			for (TSet<UEdGraphNode*>::TIterator It(ExpandedNodes); It; ++It)
+			if (bExpandedNodesNeedUniqueGuid)
 			{
-				UEdGraphNode* Node = *It;
-				AvgNodePosition.X += Node->NodePosX;
-				AvgNodePosition.Y += Node->NodePosY;
+				ExpandedNode->CreateNewGuid();
 			}
 
-			float InvNumNodes = 1.0f / float(ExpandedNodes.Num());
-			AvgNodePosition.X *= InvNumNodes;
-			AvgNodePosition.Y *= InvNumNodes;
+			//Add expanded node to selection
+			FocusedGraphEd->SetNodeSelection(ExpandedNode, true);
+		}
+	}
+}
+
+void FBlueprintEditor::OnConvertFunctionToEvent()
+{
+	if (UEdGraphNode* SelectedNode = GetSingleSelectedNode())
+	{
+		if (UK2Node_FunctionEntry* SelectedCallFunctionNode = Cast<UK2Node_FunctionEntry>(SelectedNode))
+		{
+			ConvertFunctionToEvent(SelectedCallFunctionNode);
+		}
+	}
+	else
+	{
+		FCompilerResultsLog LogResults;
+		LogResults.Error(*LOCTEXT("MultipleNodesSelectred_Error", "Only one node can be selected for conversion!").ToString());
+		FMessageLog MessageLog("BlueprintLog");
+		MessageLog.NewPage(LOCTEXT("OnConvertEventToFunctionLabel", "Convert Event to Function"));
+		MessageLog.AddMessages(LogResults.Messages);
+		MessageLog.Notify(LOCTEXT("OnConvertEventToFunctionError", "Convert Event to Function Failed!"));
+	}
+}
+
+void FBlueprintEditor::ConvertFunctionToEvent(UK2Node_FunctionEntry* SelectedCallFunctionNode)
+{
+	check(SelectedCallFunctionNode);
+	UBlueprint* NodeBP = GetBlueprintObj();
+	UEdGraph* FunctionGraph = SelectedCallFunctionNode->GetGraph();
 
-			//Remove source node from selection
-			UEdGraphNode* SourceNode = CastChecked<UEdGraphNode>(*NodeIt);
-			FocusedGraphEd->SetNodeSelection(SourceNode, false);
+	// Create a new event node with the old function name
+	UFunction* Func = SelectedCallFunctionNode->FindSignatureFunction();
+	UEdGraph* EventGraph = FBlueprintEditorUtils::FindEventGraph(NodeBP);
 
-			for (UEdGraphNode* ExpandedNode : ExpandedNodes)
+	if (Func && EventGraph && FunctionGraph)
+	{
+		const FScopedTransaction Transaction(FGraphEditorCommands::Get().ConvertFunctionToEvent->GetLabel());
+		NodeBP->Modify();
+		EventGraph->Modify();
+
+		// Keep track of the old connections from the entry node
+		TMap<FString, TSet<UEdGraphPin*>> PinConnections;
+		GetPinConnectionMap(SelectedCallFunctionNode, PinConnections);
+
+		FName EventName = Func->GetFName();
+		UClass* const OverrideFuncClass = CastChecked<UClass>(Func->GetOuter())->GetAuthoritativeClass();
+		
+		UK2Node_Event* NewEventNode = nullptr;
+		FVector2D SpawnPos = EventGraph->GetGoodPlaceForNewNode();
+		
+		// Was this function implemented in as an override?
+		UFunction* ParentFunction = FindField<UFunction>(NodeBP->ParentClass, EventName);
+		bool bIsOverrideFunc = Func->GetSuperFunction() || (ParentFunction != nullptr);
+		const UEdGraphSchema_K2* K2Schema = GetDefault<UEdGraphSchema_K2>();
+
+		if (bIsOverrideFunc)
+		{
+			NewEventNode = FEdGraphSchemaAction_K2NewNode::SpawnNode<UK2Node_Event>(
+				EventGraph,
+				SpawnPos,
+				EK2NewNodeFlags::SelectNewNode,
+				[EventName, OverrideFuncClass](UK2Node_Event* NewInstance)
+				{
+					NewInstance->EventReference.SetExternalMember(EventName, OverrideFuncClass);
+					NewInstance->bOverrideFunction = true;
+				}
+			);
+		}
+		else
+		{
+			// Spawn a new node, we have to do it this way to avoid renaming on creation of the custom event
+			UEdGraphNode* NewNode = nullptr;
+			NewNode = NewObject<UEdGraphNode>(EventGraph, UK2Node_CustomEvent::StaticClass());
+			check(NewNode != nullptr);
+			NewNode->CreateNewGuid();
+
+			NewNode->NodePosX = SpawnPos.X;
+			NewNode->NodePosY = SpawnPos.Y;
+
+			NewNode->SetFlags(RF_Transactional);
+			NewNode->AllocateDefaultPins();
+			NewNode->PostPlacedNewNode();
+			EventGraph->Modify();
+			// the FBlueprintMenuActionItem should do the selecting
+			EventGraph->AddNode(NewNode, /*bFromUI =*/false, /*bSelectNewNode =*/false);
+			
+			NewEventNode = Cast<UK2Node_CustomEvent>(NewNode);
+			NewEventNode->CustomFunctionName = EventName;
+			NewEventNode->bOverrideFunction = false;
+
+			// Add every type of user pin that we need tot he new event node
+			for (TSharedPtr<FUserPinInfo> Pin : SelectedCallFunctionNode->UserDefinedPins)
 			{
-				ExpandedNode->NodePosX = (ExpandedNode->NodePosX - AvgNodePosition.X) + SourceNode->NodePosX;
-				ExpandedNode->NodePosY = (ExpandedNode->NodePosY - AvgNodePosition.Y) + SourceNode->NodePosY;
+				NewEventNode->CreateUserDefinedPin(Pin->PinName, Pin->PinType, Pin->DesiredPinDirection);
+			}
 
-				ExpandedNode->SnapToGrid(SNodePanel::GetSnapGridSize());
+			K2Schema->ReconstructNode(*NewEventNode);
+		}
+		
+		// Keep track of any nodes that have been expanded out of the function graph
+		TSet<UEdGraphNode*> ExpandedNodes;
 
-				if (bExpandedNodesNeedUniqueGuid)
+		UEdGraphNode* Entry = nullptr;
+		UEdGraphNode* Result = nullptr;
+		FunctionGraph->Modify();
+		MoveNodesToGraph(FunctionGraph->Nodes, EventGraph, ExpandedNodes, &Entry, &Result, false);
+
+		MoveNodesToAveragePos(ExpandedNodes, FVector2D(SpawnPos.X + 500.0f, SpawnPos.Y));
+		
+		if (Entry)
+		{
+			Entry->DestroyNode();
+		}
+
+		if (Result)
+		{
+			Result->DestroyNode();
+		}
+		// Connect any pins that need to be set from the old function
+		if (NewEventNode)
+		{
+			// Link the nodes from the original function entry node to the new event node
+			ReconnectPinMap(NewEventNode, PinConnections);
+			FKismetEditorUtilities::BringKismetToFocusAttentionOnObject(NewEventNode, false);
+		}
+
+		// Remove the old function graph
+		FBlueprintEditorUtils::RemoveGraph(NodeBP, FunctionGraph, EGraphRemoveFlags::Recompile);
+		FunctionGraph->MarkPendingKill();
+		FBlueprintEditorUtils::MarkBlueprintAsStructurallyModified(NodeBP);
+	}
+}
+
+void FBlueprintEditor::GetPinConnectionMap(UEdGraphNode* Node, TMap<FString, TSet<UEdGraphPin*>>& OutPinConnections) const
+{
+	check(Node);
+	
+	for (UEdGraphPin* Pin : Node->Pins)
+	{
+		if (Pin != nullptr)
+		{
+			FString PinName = Pin->GetName();
+
+			// If this is the first time seeing this pin name, add a new set
+			if (!OutPinConnections.Contains(PinName))
+			{
+				OutPinConnections.Add(PinName, TSet<UEdGraphPin*>());
+			}
+			else
+			{
+				// There are no pins connected to this
+				continue;
+			}
+
+			for (UEdGraphPin* ConnectedPin : Pin->LinkedTo)
+			{
+				if (ConnectedPin != nullptr)
 				{
-					ExpandedNode->CreateNewGuid();
+					// Add to the array of nodes at this pin name
+					OutPinConnections[PinName].Add(ConnectedPin);
 				}
+			}
+		}
+	}
+}
 
-				//Add expanded node to selection
-				FocusedGraphEd->SetNodeSelection(ExpandedNode, true);
+bool FBlueprintEditor::CanConvertFunctionToEvent() const
+{
+	UEdGraphNode* const SelectedNode = GetSingleSelectedNode();
+	if (SelectedNode)
+	{
+		if (UK2Node_FunctionEntry* const SelectedCallFunctionNode = Cast<UK2Node_FunctionEntry>(SelectedNode))
+		{
+			UFunction* const Func = SelectedCallFunctionNode->FindSignatureFunction();
+			
+			// If the function has no output parameters and is not an interface or construction script
+			if (Func && !UEdGraphSchema_K2::HasFunctionAnyOutputParameter(Func) &&
+				!FBlueprintEditorUtils::IsInterfaceFunction(GetBlueprintObj(), Func) &&
+				SelectedCallFunctionNode->FunctionReference.GetMemberName() != UEdGraphSchema_K2::FN_UserConstructionScript)
+			{
+				return true;
 			}
 		}
 	}
+	// We cannot convert functions that have output parameters to events
+	return false;
+}
 
-	FBlueprintEditorUtils::MarkBlueprintAsStructurallyModified(GetBlueprintObj());
+void FBlueprintEditor::OnConvertEventToFunction()
+{
+	const FScopedTransaction Transaction(FGraphEditorCommands::Get().ConvertEventToFunction->GetLabel());
+
+	UBlueprint* const NodeBP = GetBlueprintObj();
+	NodeBP->Modify();
+	
+	// Keep track of the current selected node
+	UEdGraphNode* const SelectedNode = GetSingleSelectedNode();
+
+	if (UK2Node_Event* SelectedEventNode = Cast<UK2Node_Event>(SelectedNode))
+	{
+		const UEdGraphSchema_K2* Schema = GetDefault<UEdGraphSchema_K2>();
+
+		// Break any connections to the delegate pin, because that won't exist 
+		{
+			UEdGraphPin* DelegatePin = SelectedEventNode->FindPin(UK2Node_Event::DelegateOutputName);
+			if (DelegatePin)
+			{
+				DelegatePin->BreakAllPinLinks();
+			}
+		}
+		// Refresh this node before getting the pin names to ensure that we get the most up to date ones
+		Schema->ReconstructNode(*SelectedEventNode);
+
+		// Keep track of the old connections from the event node
+		TMap<FString, TSet<UEdGraphPin*>> PinConnections;
+		GetPinConnectionMap(SelectedNode, PinConnections);
+
+		TArray<UEdGraphNode*> CollapsableNodes = GetAllConnectedNodes(SelectedNode);
+
+		UFunction* const FunctionSig = FFunctionFromNodeHelper::FunctionFromNode(SelectedEventNode);
+		UEdGraph* const SourceGraph = SelectedEventNode->GetGraph();
+		SourceGraph->Modify();
+
+		const FName OriginalEventName = SelectedEventNode->GetFunctionName();
+
+		if (FunctionSig && SourceGraph)
+		{
+			// Check if this is an override function
+			UFunction* ParentFunction = FindField<UFunction>(NodeBP->ParentClass, OriginalEventName);
+			const bool bIsOverrideFunc = FunctionSig->GetSuperFunction() || (ParentFunction != nullptr);
+			UClass* const OverrideFuncClass = FBlueprintEditorUtils::GetOverrideFunctionClass(NodeBP, OriginalEventName);
+
+			// Create a new function graph
+			// We have to start it out with a temporary graph name because otherwise the flags will get renamed uniquely
+			UEdGraph* NewGraph = nullptr;
+			static const FString TempFuncGraph = "TEMP_FUNC_GRAPH";
+			const FName TempFuncGraphName = FBlueprintEditorUtils::GenerateUniqueGraphName(NodeBP, TempFuncGraph);
+			NewGraph = FBlueprintEditorUtils::CreateNewGraph(NodeBP, TempFuncGraphName, SourceGraph->GetClass(), SourceGraph->GetSchema() ? SourceGraph->GetSchema()->GetClass() : Schema->GetClass());
+
+			if (bIsOverrideFunc)
+			{
+				FBlueprintEditorUtils::AddFunctionGraph<UClass>(NodeBP, NewGraph, /*bIsUserCreated=*/ false, OverrideFuncClass);
+			}
+			else
+			{
+				FBlueprintEditorUtils::AddFunctionGraph<UFunction>(NodeBP, NewGraph, /*bIsUserCreated=*/ true, FunctionSig);
+			}
+			
+			// Get the new entry node
+			TArray<UK2Node_FunctionEntry*> EntryNodes;
+			NewGraph->GetNodesOfClass<UK2Node_FunctionEntry>(EntryNodes);
+			check(EntryNodes.Num() > 0);
+			UK2Node_FunctionEntry* NewEntryNode = EntryNodes[0];
+
+			// Reconstruct the pins and add a proper function reference to the parent function
+			if (bIsOverrideFunc && NewEntryNode)
+			{
+				FGuid GraphGuid;
+				FBlueprintEditorUtils::GetFunctionGuidFromClassByFieldName(FBlueprintEditorUtils::GetMostUpToDateClass(OverrideFuncClass), OriginalEventName, /* Out */GraphGuid);
+				NewEntryNode->FunctionReference.SetExternalMember(OriginalEventName, OverrideFuncClass, GraphGuid);
+				Schema->ReconstructNode(*NewEntryNode);
+			}
+			
+			NewGraph->Modify();
+
+			UEdGraphNode* InvalidNode = nullptr;
+			FCompilerResultsLog LogResults;
+
+			// Check to make sure that this node can actually be put in event graph
+			for (UEdGraphNode* const Node : CollapsableNodes)
+			{
+				if (Node && (!Schema->CanEncapuslateNode(*Node) || !Node->CanPasteHere(NewGraph)))
+				{
+					// Tell the user why we can't send to function graph
+					LogResults.Error(*LOCTEXT("InvalidNode_Error", "@@ cannot be placed in function graph").ToString(), Node);
+					InvalidNode = Node;
+				}
+			}
+
+			// If this is invalid, then cancel this operation and focus on the node that is invalid
+			if (InvalidNode != nullptr)
+			{
+				FMessageLog MessageLog("BlueprintLog");
+				MessageLog.NewPage(LOCTEXT("OnConvertEventToFunctionLabel", "Convert Event to Function"));
+				MessageLog.AddMessages(LogResults.Messages);
+				MessageLog.Notify(LOCTEXT("OnConvertEventToFunctionError", "Convert Event to Function Failed!"));
+
+				FKismetEditorUtilities::BringKismetToFocusAttentionOnObject(InvalidNode, false);
+
+				// Get rid of the function graph
+				FBlueprintEditorUtils::RemoveGraph(NodeBP, NewGraph);
+				NewGraph->MarkPendingKill();
+				return;
+			}
+			
+			// Move the nodes to the new graph
+			UEdGraphNode* Entry = nullptr;
+			UEdGraphNode* Result = nullptr;
+			TSet<UEdGraphNode*> ExpandedNodes;
+			MoveNodesToGraph(CollapsableNodes, NewGraph, ExpandedNodes, &Entry, &Result);
+			
+			// Link the new nodes accordingly
+			ReconnectPinMap(NewEntryNode, PinConnections);
+			if (NewEntryNode)
+			{
+				MoveNodesToAveragePos(ExpandedNodes, FVector2D(NewEntryNode->NodePosX + 500.0f, NewEntryNode->NodePosY));
+			}
+			else
+			{
+				MoveNodesToAveragePos(ExpandedNodes, NewGraph->GetGoodPlaceForNewNode());
+			}
+
+			FBlueprintEditorUtils::UpdateTransactionalFlags(NodeBP);
+			
+			FVector2D NewFuncCallSpawn( SelectedEventNode->NodePosX, SelectedEventNode->NodePosY );
+
+			FBlueprintEditorUtils::RemoveNode(NodeBP, SelectedEventNode, /* bDontRecompile= */ true);
+
+			// Cache the string conversion of the event name
+			const FString OriginalEventNameString = OriginalEventName.ToString();
+
+			// If there is already another object in the same scope with this name, rename it.
+			UObject* ExistingObject = StaticFindObject(/*Class=*/ nullptr, NodeBP, *OriginalEventNameString, true);
+			if (ExistingObject)
+			{
+				check(ExistingObject->GetOuter() == NewGraph->GetOuter());
+
+				ExistingObject->Rename(nullptr, nullptr, REN_DontCreateRedirectors);
+			}
+
+			// Rename the function graph to the original name
+			FBlueprintEditorUtils::RenameGraph(NewGraph, *OriginalEventNameString);
+
+			// Create a function call the to the newly created graph in place of the old event
+			IBlueprintNodeBinder::FBindingSet Bindings;
+			UEdGraphNode* OutFunctionNode = UBlueprintFunctionNodeSpawner::Create(FindField<UFunction>(NodeBP->SkeletonGeneratedClass, OriginalEventName))->Invoke(SourceGraph, Bindings, NewFuncCallSpawn);
+		}
+
+		FBlueprintEditorUtils::MarkBlueprintAsStructurallyModified(NodeBP);
+	}
 }
 
-bool FBlueprintEditor::CanExpandNodes() const
+void FBlueprintEditor::ReconnectPinMap(UEdGraphNode* Node, const TMap<FString, TSet<UEdGraphPin*>>& PinConnections)
 {
-	// Does the selection set contain any composite nodes that are legal to expand?
-	const FGraphPanelSelectionSet SelectedNodes = GetSelectedNodes();
-	for (FGraphPanelSelectionSet::TConstIterator NodeIt(SelectedNodes); NodeIt; ++NodeIt)
+	check(Node);
+
+	for (UEdGraphPin* const NewPin : Node->Pins)
 	{
-		if (Cast<UK2Node_Composite>(*NodeIt))
+		const FString& NewPinName = NewPin->GetName();
+		if (PinConnections.Contains(NewPinName))
 		{
-			return true;
+			// Connect the new pins here
+			for (UEdGraphPin* OldPin : PinConnections[NewPinName])
+			{
+				NewPin->MakeLinkTo(OldPin);
+			}
 		}
-		else if (UK2Node_MacroInstance* SelectedMacroInstanceNode = Cast<UK2Node_MacroInstance>(*NodeIt))
+	}
+}
+
+TArray<UEdGraphNode*> FBlueprintEditor::GetAllConnectedNodes(UEdGraphNode* const SourceNode) const
+{
+	TArray<UEdGraphNode*> OutNodes;
+
+	if (SourceNode == nullptr)
+	{
+		return OutNodes;
+	}
+
+	const UEdGraphSchema_K2* Schema = GetDefault<UEdGraphSchema_K2>();
+
+	TQueue<UEdGraphNode*> Queue;
+	Queue.Enqueue(SourceNode);
+	
+	TMap<int32, UEdGraphNode*> VisitedNodes;
+	VisitedNodes.Add(SourceNode->GetUniqueID(), SourceNode);
+
+	while (!Queue.IsEmpty())
+	{
+		UEdGraphNode* Current = nullptr;
+		Queue.Dequeue(Current);
+
+		if (Current != nullptr)
 		{
-			return SelectedMacroInstanceNode->GetMacroGraph() != NULL;
+			// Don't add the original node in the list
+			if (Current != SourceNode)
+			{
+				OutNodes.Add(Current);
+			}
+			
+			// For every pin that this node has
+			for (UEdGraphPin* CurrentPin : Current->Pins)
+			{
+				// Look at what pins are connected to that
+				for (UEdGraphPin* LinkedPin : CurrentPin->LinkedTo)
+				{
+					UEdGraphNode* OwningNode = LinkedPin->GetOwningNode();
+
+					// If we can encapsulate the owner and it has not been visited
+					if (OwningNode && !VisitedNodes.Contains(OwningNode->GetUniqueID()))
+					{
+						// Mark as visited and add to the queue
+						VisitedNodes.Add(OwningNode->GetUniqueID(), OwningNode);
+						Queue.Enqueue(OwningNode);
+					}
+				}
+			}
 		}
-		else if (UK2Node_CallFunction* SelectedCallFunctionNode = Cast<UK2Node_CallFunction>(*NodeIt))
+	}
+
+	return OutNodes;
+}
+
+bool FBlueprintEditor::CanConvertEventToFunction() const
+{
+	// Only allow when only a _single_ event node is selected, not allowed on the anim graph
+	UEdGraphNode* SelectedNode = GetSingleSelectedNode();
+	if (SelectedNode && !IsEditingAnimGraph())
+	{
+		if (UK2Node_Event* SelectedEventNode = Cast<UK2Node_Event>(SelectedNode))
 		{
-			// If ResultEventNode is non-NULL, it means it is sourced by an event, we do not want to expand events
-			const UEdGraphNode* ResultEventNode = NULL;
-			return SelectedCallFunctionNode->GetFunctionGraph(ResultEventNode) != NULL && ResultEventNode == NULL;
+			UFunction* const Func = FFunctionFromNodeHelper::FunctionFromNode(SelectedEventNode);
+
+			// If the function has no output parameters it is valid to make a function
+			if (!FBlueprintEditorUtils::IsInterfaceFunction(GetBlueprintObj(), Func) &&
+				!SelectedEventNode->IsInterfaceEventNode())
+			{
+				return true;
+			}
 		}
 	}
 
@@ -5492,7 +5939,7 @@ void FBlueprintEditor::DeleteSelectedNodes()
 				}
 
 				UK2Node* K2Node = Cast<UK2Node>(Node);
-				if (K2Node != NULL && K2Node->NodeCausesStructuralBlueprintChange())
+				if (K2Node != nullptr && K2Node->NodeCausesStructuralBlueprintChange())
 				{
 					bNeedToModifyStructurally = true;
 				}
@@ -5663,7 +6110,7 @@ bool FBlueprintEditor::CanCopyNodes() const
 	for (FGraphPanelSelectionSet::TConstIterator SelectedIter(SelectedNodes); SelectedIter; ++SelectedIter)
 	{
 		UEdGraphNode* Node = Cast<UEdGraphNode>(*SelectedIter);
-		if ((Node != NULL) && Node->CanDuplicateNode())
+		if ((Node != nullptr) && Node->CanDuplicateNode())
 		{
 			return true;
 		}
@@ -7226,71 +7673,14 @@ void FBlueprintEditor::ExpandNode(UEdGraphNode* InNodeToExpand, UEdGraph* InSour
 	SourceGraph->Modify();
 	InNodeToExpand->Modify();
 
-	UEdGraphNode* Entry = NULL;
-	UEdGraphNode* Result = NULL;
+	UEdGraphNode* Entry = nullptr;
+	UEdGraphNode* Result = nullptr;
 
 	const bool bIsCollapsedGraph = InNodeToExpand->IsA<UK2Node_Composite>();
 
-	// Move the nodes over, remembering any that are boundary nodes
-	while (SourceGraph->Nodes.Num())
-	{
-		UEdGraphNode* Node = SourceGraph->Nodes.Pop();
-
-		Node->Modify();
-		Node->Rename(/*NewName=*/ NULL, /*NewOuter=*/ DestinationGraph);
-
-		// We do not check CanPasteHere when determining CanCollapseNodes, unlike CanCollapseSelectionToFunction/Macro,
-		// so when expanding a collapsed graph we don't want to check the CanPasteHere function:
-		if (!bIsCollapsedGraph && !Node->CanPasteHere(DestinationGraph))
-		{
-			Node->BreakAllNodeLinks();
-			continue;
-		}
-
-		// Successfully added the node to the graph, we may need to remove flags
-		if (Node->HasAllFlags(RF_Transient) && !DestinationGraph->HasAllFlags(RF_Transient))
-		{
-			Node->SetFlags(RF_Transactional);
-			Node->ClearFlags(RF_Transient);
-			TArray<UObject*> Subobjects;
-			GetObjectsWithOuter(Node, Subobjects);
-			for (UObject* Subobject : Subobjects)
-			{
-				Subobject->ClearFlags(RF_Transient);
-				Subobject->SetFlags(RF_Transactional);
-			}
-		}
-		DestinationGraph->Nodes.Add(Node);
-
-		// Want to test exactly against tunnel, we shouldn't collapse embedded collapsed
-		// nodes or macros, only the tunnels in/out of the collapsed graph
-		if (Node->GetClass() == UK2Node_Tunnel::StaticClass())
-		{
-			UK2Node_Tunnel* TunnelNode = Cast<UK2Node_Tunnel>(Node);
-			if(TunnelNode->bCanHaveOutputs)
-			{
-				Entry = Node;
-			}
-			else if(TunnelNode->bCanHaveInputs)
-			{
-				Result = Node;
-			}
-		}
-		else if (Node->GetClass() == UK2Node_FunctionEntry::StaticClass())
-		{
-			Entry = Node;
-		}
-		else if(Node->GetClass() == UK2Node_FunctionResult::StaticClass())
-		{
-			Result = Node;
-		}
-		else
-		{
-			OutExpandedNodes.Add(Node);
-		}
-	}
+	MoveNodesToGraph(SourceGraph->Nodes, DestinationGraph, OutExpandedNodes, &Entry, &Result, bIsCollapsedGraph);
 
-	UEdGraphPin* OutputExecPinReconnect = NULL;
+	UEdGraphPin* OutputExecPinReconnect = nullptr;
 	if(UK2Node_CallFunction* CallFunction = Cast<UK2Node_CallFunction>(InNodeToExpand))
 	{
 		UEdGraphPin* ThenPin = CallFunction->GetThenPin();
@@ -7341,6 +7731,74 @@ void FBlueprintEditor::ExpandNode(UEdGraphNode* InNodeToExpand, UEdGraph* InSour
 	}
 }
 
+void FBlueprintEditor::MoveNodesToGraph(TArray<UEdGraphNode*>& SourceNodes, UEdGraph* DestinationGraph, TSet<UEdGraphNode*>& OutExpandedNodes, UEdGraphNode** OutEntry, UEdGraphNode** OutResult, const bool bIsCollapsedGraph)
+{
+	// Move the nodes over, remembering any that are boundary nodes
+	while (SourceNodes.Num())
+	{
+		UEdGraphNode* Node = SourceNodes.Pop();
+		UEdGraph* OriginalGraph = Node->GetGraph();
+
+		Node->Modify();
+		OriginalGraph->Modify();
+		Node->Rename(/*NewName=*/ nullptr, /*NewOuter=*/ DestinationGraph, REN_DontCreateRedirectors);
+
+		// Remove the node from the original graph
+		OriginalGraph->RemoveNode(Node, false);
+		
+		// We do not check CanPasteHere when determining CanCollapseNodes, unlike CanCollapseSelectionToFunction/Macro,
+		// so when expanding a collapsed graph we don't want to check the CanPasteHere function:
+		if (!bIsCollapsedGraph && !Node->CanPasteHere(DestinationGraph))
+		{
+			Node->BreakAllNodeLinks();
+			continue;
+		}
+
+		// Successfully added the node to the graph, we may need to remove flags
+		if (Node->HasAllFlags(RF_Transient) && !DestinationGraph->HasAllFlags(RF_Transient))
+		{
+			Node->SetFlags(RF_Transactional);
+			Node->ClearFlags(RF_Transient);
+			TArray<UObject*> Subobjects;
+			GetObjectsWithOuter(Node, Subobjects);
+			for (UObject* Subobject : Subobjects)
+			{
+				Subobject->ClearFlags(RF_Transient);
+				Subobject->SetFlags(RF_Transactional);
+			}
+		}
+		
+		DestinationGraph->AddNode(Node, /* bFromUI */ false, /* bSelectNewNode */ false);
+		
+		// Want to test exactly against tunnel, we shouldn't collapse embedded collapsed
+		// nodes or macros, only the tunnels in/out of the collapsed graph
+		if (Node->GetClass() == UK2Node_Tunnel::StaticClass())
+		{
+			UK2Node_Tunnel* TunnelNode = Cast<UK2Node_Tunnel>(Node);
+			if (TunnelNode->bCanHaveOutputs)
+			{
+				*OutEntry = Node;
+			}
+			else if (TunnelNode->bCanHaveInputs)
+			{
+				*OutResult = Node;
+			}
+		}
+		else if (Node->GetClass() == UK2Node_FunctionEntry::StaticClass())
+		{
+			*OutEntry = Node;
+		}
+		else if (Node->GetClass() == UK2Node_FunctionResult::StaticClass())
+		{
+			*OutResult = Node;
+		}
+		else
+		{
+			OutExpandedNodes.Add(Node);
+		}
+	}
+}
+
 void FBlueprintEditor::SaveEditedObjectState()
 {
 	check(IsEditingSingleBlueprint());
diff --git a/Engine/Source/Editor/Kismet/Public/BlueprintEditor.h b/Engine/Source/Editor/Kismet/Public/BlueprintEditor.h
--- a/Engine/Source/Editor/Kismet/Public/BlueprintEditor.h
+++ b/Engine/Source/Editor/Kismet/Public/BlueprintEditor.h
@@ -45,6 +45,7 @@ class UUserDefinedEnum;
 class UUserDefinedStruct;
 class UBlueprintEditorOptions;
 struct Rect;
+class UK2Node_FunctionEntry;
 
 /* Enums to use when grouping the blueprint members in the list panel. The order here will determine the order in the list */
 namespace NodeSectionID
@@ -815,6 +816,53 @@ protected:
 	void OnExpandNodes();
 	bool CanExpandNodes() const;
 
+	/**
+	* Move the given set of nodes to an average spot near the Source position
+	* 
+	* @param AverageNodes					The nodes to move
+	* @param SourcePos						The source position used to average the nodes around
+	* @param bExpandedNodesNeedUniqueGuid	If true then a new Guid will be generated for each node in the set
+	*/
+	void MoveNodesToAveragePos(TSet<UEdGraphNode*>& AverageNodes, FVector2D SourcePos, bool bExpandedNodesNeedUniqueGuid = false) const;
+
+	void OnConvertFunctionToEvent();
+	bool CanConvertFunctionToEvent() const;
+
+	/*
+	* Given a function node, move all nodes from the function out of the function graph,
+	* create an event with the same function name, and connect nodes to that event. 
+	* 
+	* @param SelectedCallFunctionNode	The function node to convert to an event
+	*/
+	void ConvertFunctionToEvent(UK2Node_FunctionEntry* SelectedCallFunctionNode);
+
+	/**
+	* Output a map of pin names to a set of connections given a function entry pin
+	* 
+	* @param Node				The entry pin to gather connections from
+	* @param OutPinConnections	Output map of connection data
+	*/
+	void GetPinConnectionMap(UEdGraphNode* Node, TMap<FString, TSet<UEdGraphPin*>>& OutPinConnections) const;
+
+	/**
+	* Reconnect the pin map to the given node
+	* 
+	* @param Node				The node to connect pins to
+	* @param PinConnections		Map of pin connections to set
+	*/
+	void ReconnectPinMap(UEdGraphNode* Node, const TMap<FString, TSet<UEdGraphPin*>>& PinConnections);
+
+	void OnConvertEventToFunction();
+	bool CanConvertEventToFunction() const;
+
+	/**
+	* Get all connected nodes and add them to a set of out nodes
+	* 
+	* @param Node		The node to get all connects to 
+	* @return	TArray of all connected nodes
+	*/
+	TArray<UEdGraphNode*> GetAllConnectedNodes(UEdGraphNode* const Node) const;
+
 	void OnAlignTop();
 	void OnAlignMiddle();
 	void OnAlignBottom();
@@ -964,6 +1012,19 @@ protected:
 	 * Expands passed in node */
 	static void ExpandNode(UEdGraphNode* InNodeToExpand, UEdGraph* InSourceGraph, TSet<UEdGraphNode*>& OutExpandedNodes);
 
+	/**
+	* Move every node from the source graph to the destination graph. Add Each node that is moved to the OutExpandedNodes set.
+	* If the source graph is a function graph, keep track of the entry and result nodes in the given Out Parameters. 
+	*
+	* @param SourceNodes		Nodes to move
+	* @param DestinationGraph	Graph to move nodes to
+	* @param OutExpandedNodes	Set of each node that was moved from the source to destination graph
+	* @param OutEntry			Pointer to the function entry node
+	* @param OutResult			Pointer to the function result node
+	* @param bIsCollapsedGraph	Whether or not the source graph is collapsed
+	**/
+	static void MoveNodesToGraph(TArray<UEdGraphNode*>& SourceNodes, UEdGraph* DestinationGraph, TSet<UEdGraphNode*>& OutExpandedNodes, UEdGraphNode** OutEntry, UEdGraphNode** OutResult, const bool bIsCollapsedGraph = false);
+
 	/** Start editing the defaults for this blueprint */
 	void OnStartEditingDefaultsClicked();
 
