PR #6076: Add mouse scrolling to SGraphTitleBar - helps with long path navigati. (Contributed by [REDACTED_PERSON])

Minor cleanup of NULL and forward declarations as well.

diff --git a/Engine/Source/Editor/Kismet/Private/SGraphTitleBar.cpp b/Engine/Source/Editor/Kismet/Private/SGraphTitleBar.cpp
--- a/Engine/Source/Editor/Kismet/Private/SGraphTitleBar.cpp
+++ b/Engine/Source/Editor/Kismet/Private/SGraphTitleBar.cpp
@@ -7,6 +7,7 @@
 #include "Misc/Paths.h"
 #include "Widgets/SBoxPanel.h"
 #include "Widgets/Layout/SBorder.h"
+#include "Widgets/Layout/SScrollBox.h"
 #include "Widgets/Layout/SSeparator.h"
 #include "Widgets/Images/SImage.h"
 #include "Widgets/Text/STextBlock.h"
@@ -33,7 +34,7 @@ SGraphTitleBar::~SGraphTitleBar()
 
 const FSlateBrush* SGraphTitleBar::GetTypeGlyph() const
 {
-	check(EdGraphObj != NULL);
+	check(EdGraphObj != nullptr);
 	return FBlueprintEditor::GetGlyphForGraph(EdGraphObj, true);
 }
 
@@ -52,7 +53,7 @@ FText SGraphTitleBar::GetTitleForOneCrumb(const UEdGraph* Graph)
 
 FText SGraphTitleBar::GetTitleExtra() const
 {
-	check(EdGraphObj != NULL);
+	check(EdGraphObj != nullptr);
 
 	FText ExtraText = Kismet2Ptr.Pin()->GetGraphDecorationString(EdGraphObj);
 
@@ -133,55 +134,68 @@ void SGraphTitleBar::Construct( const FArguments& InArgs )
 					SNew(SHorizontalBox)
 					+SHorizontalBox::Slot()
 					.AutoWidth()
-					.Padding( 10,5 )
+					.Padding( 10.0f,5.0f )
 					.VAlign(VAlign_Center)
 					[
 						SNew(SImage)
 						.Image( this, &SGraphTitleBar::GetTypeGlyph )
 					]
-
-					// show fake 'root' breadcrumb for the title
-					+SHorizontalBox::Slot()
-					.AutoWidth()
-					.VAlign(VAlign_Center)
-					.Padding(BreadcrumbTrailPadding)
-					[
-						SNew(STextBlock)
-						.Text(this, &SGraphTitleBar::GetBlueprintTitle )
-						.TextStyle( FEditorStyle::Get(), TEXT("GraphBreadcrumbButtonText") )
-						.Visibility( this, &SGraphTitleBar::IsGraphBlueprintNameVisible )
-					]
 					+SHorizontalBox::Slot()
-					.AutoWidth()
-					.VAlign(VAlign_Center)
-					[
-						SNew(SImage)
-						.Image( BreadcrumbButtonImage )
-						.Visibility( this, &SGraphTitleBar::IsGraphBlueprintNameVisible )
-					]
-
-					// New style breadcrumb
-					+SHorizontalBox::Slot()
-					.AutoWidth()
-					.VAlign(VAlign_Center)
-					[
-						SAssignNew(BreadcrumbTrail, SBreadcrumbTrail<UEdGraph*>)
-						.ButtonStyle(FEditorStyle::Get(), "GraphBreadcrumbButton")
-						.TextStyle(FEditorStyle::Get(), "GraphBreadcrumbButtonText")
-						.ButtonContentPadding( BreadcrumbTrailPadding )
-						.DelimiterImage( BreadcrumbButtonImage )
-						.PersistentBreadcrumbs( true )
-						.OnCrumbClicked( this, &SGraphTitleBar::OnBreadcrumbClicked )
-					]
-
-					+SHorizontalBox::Slot()
-					.AutoWidth()
-					.VAlign(VAlign_Center)
+					.FillWidth(1.f)
 					[
-						SNew(STextBlock)
-						.Font( FCoreStyle::GetDefaultFontStyle("Regular", 14) )
-						.ColorAndOpacity( FLinearColor(1,1,1,0.5) )
-						.Text( this, &SGraphTitleBar::GetTitleExtra )
+						SAssignNew(BreadcrumbTrailScrollBox, SScrollBox)
+						.Orientation(Orient_Horizontal)
+						.ScrollBarVisibility(EVisibility::Collapsed)
+
+						+SScrollBox::Slot()
+						.Padding(0.f)
+						.VAlign(VAlign_Center)
+						[
+							SNew(SHorizontalBox)
+							// show fake 'root' breadcrumb for the title
+							+SHorizontalBox::Slot()
+							.AutoWidth()
+							.VAlign(VAlign_Center)
+							.Padding(BreadcrumbTrailPadding)
+							[
+								SNew(STextBlock)
+								.Text(this, &SGraphTitleBar::GetBlueprintTitle )
+								.TextStyle( FEditorStyle::Get(), TEXT("GraphBreadcrumbButtonText") )
+								.Visibility( this, &SGraphTitleBar::IsGraphBlueprintNameVisible )
+							]
+							+SHorizontalBox::Slot()
+							.AutoWidth()
+							.VAlign(VAlign_Center)
+							[
+								SNew(SImage)
+								.Image( BreadcrumbButtonImage )
+								.Visibility( this, &SGraphTitleBar::IsGraphBlueprintNameVisible )
+							]
+
+							// New style breadcrumb
+							+SHorizontalBox::Slot()
+							.AutoWidth()
+							.VAlign(VAlign_Center)
+							[
+								SAssignNew(BreadcrumbTrail, SBreadcrumbTrail<UEdGraph*>)
+								.ButtonStyle(FEditorStyle::Get(), "GraphBreadcrumbButton")
+								.TextStyle(FEditorStyle::Get(), "GraphBreadcrumbButtonText")
+								.ButtonContentPadding( BreadcrumbTrailPadding )
+								.DelimiterImage( BreadcrumbButtonImage )
+								.PersistentBreadcrumbs( true )
+								.OnCrumbClicked( this, &SGraphTitleBar::OnBreadcrumbClicked )
+							]
+
+							+SHorizontalBox::Slot()
+							.AutoWidth()
+							.VAlign(VAlign_Center)
+							[
+								SNew(STextBlock)
+								.Font( FCoreStyle::GetDefaultFontStyle("Regular", 14) )
+								.ColorAndOpacity( FLinearColor(1,1,1,0.5) )
+								.Text( this, &SGraphTitleBar::GetTitleExtra )
+							]
+						]
 					]
 				]
 			]
@@ -189,6 +203,7 @@ void SGraphTitleBar::Construct( const FArguments& InArgs )
 	];
 
 	RebuildBreadcrumbTrail();
+	BreadcrumbTrailScrollBox->ScrollToEnd();
 
 	UBlueprint* BlueprintObj = FBlueprintEditorUtils::FindBlueprintForGraph(this->EdGraphObj);
 	if (BlueprintObj)
@@ -208,7 +223,7 @@ void SGraphTitleBar::RebuildBreadcrumbTrail()
 {
 	// Build up a stack of graphs so we can pop them in reverse order and create breadcrumbs
 	TArray<UEdGraph*> Stack;
-	for (UEdGraph* OuterChain = EdGraphObj; OuterChain != NULL; OuterChain = GetOuterGraph(OuterChain))
+	for (UEdGraph* OuterChain = EdGraphObj; OuterChain != nullptr; OuterChain = GetOuterGraph(OuterChain))
 	{
 		Stack.Push(OuterChain);
 	}
@@ -216,7 +231,7 @@ void SGraphTitleBar::RebuildBreadcrumbTrail()
 	BreadcrumbTrail->ClearCrumbs(false);
 
 	//Get the last object in the array
-	UEdGraph* LastObj = NULL;
+	UEdGraph* LastObj = nullptr;
 	if( Stack.Num() > 0 )
 	{
 		LastObj = Stack[Stack.Num() -1];
@@ -248,7 +263,7 @@ UEdGraph* SGraphTitleBar::GetOuterGraph( UObject* Obj )
 			}
 		}
 	}
-	return NULL;
+	return nullptr;
 }
 
 FText SGraphTitleBar::GetBlueprintTitle() const
diff --git a/Engine/Source/Editor/Kismet/Private/SGraphTitleBar.h b/Engine/Source/Editor/Kismet/Private/SGraphTitleBar.h
--- a/Engine/Source/Editor/Kismet/Private/SGraphTitleBar.h
+++ b/Engine/Source/Editor/Kismet/Private/SGraphTitleBar.h
@@ -12,6 +12,8 @@
 class FBlueprintEditor;
 class UEdGraph;
 struct FSlateBrush;
+class SFunctionEditor;
+class SScrollBox;
 
 //////////////////////////////////////////////////////////////////////////
 // SGraphTitleBar
@@ -20,7 +22,7 @@ class SGraphTitleBar : public SCompoundWidget
 {
 public:
 	SLATE_BEGIN_ARGS( SGraphTitleBar )
-		: _EdGraphObj(NULL)
+		: _EdGraphObj(nullptr)
 		, _Kismet2()
 		{}
 
@@ -46,7 +48,9 @@ protected:
 	UEdGraph* EdGraphObj;
 
 	/** Pointer to the function editor widget */
-	TWeakPtr<class SFunctionEditor>	FuncEditorPtr;
+	TWeakPtr<SFunctionEditor>	FuncEditorPtr;
+
+	TSharedPtr<SScrollBox> BreadcrumbTrailScrollBox;
 
 	/** Breadcrumb trail widget */
 	TSharedPtr< SBreadcrumbTrail<UEdGraph*> > BreadcrumbTrail;
