Event graph now focuses on the created event WITHOUT control click from the details panel

diff --git a/Engine/Source/Editor/Kismet/Private/BlueprintEditor.cpp b/Engine/Source/Editor/Kismet/Private/BlueprintEditor.cpp
--- a/Engine/Source/Editor/Kismet/Private/BlueprintEditor.cpp
+++ b/Engine/Source/Editor/Kismet/Private/BlueprintEditor.cpp
@@ -8561,7 +8561,7 @@ TSharedPtr<SGraphEditor> FBlueprintEditor::OpenGraphAndBringToFront(UEdGraph* Gr
 	SetCurrentMode(FBlueprintEditorApplicationModes::StandardBlueprintEditorMode);
 
 	// Next, try to make sure there is a copy open
-	TSharedPtr<SDockTab> TabWithGraph = OpenDocument(Graph, FDocumentTracker::CreateHistoryEvent);
+	TSharedPtr<SDockTab> TabWithGraph = OpenDocument(Graph, FDocumentTracker::NavigatingCurrentDocument);
 
 	// We know that the contents of the opened tabs will be a graph editor.
 	TSharedRef<SGraphEditor> NewGraphEditor = StaticCastSharedRef<SGraphEditor>(TabWithGraph->GetContent());
