Only add a new function call to a converted event node if the function is blueprint callable.

Break all connections to any return nodes before converting the function to an event.

Re-order the CanConvertToEvent function to be above the call for actually converting it for better code readability

diff --git a/Engine/Source/Editor/Kismet/Private/BlueprintEditor.cpp b/Engine/Source/Editor/Kismet/Private/BlueprintEditor.cpp
--- a/Engine/Source/Editor/Kismet/Private/BlueprintEditor.cpp
+++ b/Engine/Source/Editor/Kismet/Private/BlueprintEditor.cpp
@@ -5365,6 +5365,28 @@
 	}
 }
 
+bool FBlueprintEditor::CanConvertFunctionToEvent() const
+{
+	UEdGraphNode* const SelectedNode = GetSingleSelectedNode();
+	if (SelectedNode)
+	{
+		if (UK2Node_FunctionEntry* const SelectedCallFunctionNode = Cast<UK2Node_FunctionEntry>(SelectedNode))
+		{
+			UFunction* const Func = SelectedCallFunctionNode->FindSignatureFunction();
+
+			// If the function has no output parameters and is not an interface or construction script
+			if (Func && !UEdGraphSchema_K2::HasFunctionAnyOutputParameter(Func) &&
+				!FBlueprintEditorUtils::IsInterfaceFunction(GetBlueprintObj(), Func) &&
+				SelectedCallFunctionNode->FunctionReference.GetMemberName() != UEdGraphSchema_K2::FN_UserConstructionScript)
+			{
+				return true;
+			}
+		}
+	}
+	// We cannot convert functions that have output parameters to events
+	return false;
+}
+
 void FBlueprintEditor::OnConvertFunctionToEvent()
 {
 	if (UEdGraphNode* SelectedNode = GetSingleSelectedNode())
@@ -5401,6 +5423,17 @@
 		NodeBP->Modify();
 		EventGraph->Modify();
 
+		// Find all return nodes and break their connections
+		TArray< UK2Node_FunctionTerminator*> FunctionReturnNodes;
+		FunctionGraph->GetNodesOfClass(FunctionReturnNodes);
+		for (UK2Node_FunctionTerminator* Node : FunctionReturnNodes)
+		{
+			if (Node)
+			{
+				Node->BreakAllNodeLinks();
+			}
+		}
+
 		// Keep track of the old connections from the entry node
 		TMap<FString, TSet<UEdGraphPin*>> PinConnections;
 		GetPinConnectionMap(SelectedCallFunctionNode, PinConnections);
@@ -5527,25 +5560,25 @@
 	}
 }
 
-bool FBlueprintEditor::CanConvertFunctionToEvent() const
+bool FBlueprintEditor::CanConvertEventToFunction() const
 {
-	UEdGraphNode* const SelectedNode = GetSingleSelectedNode();
-	if (SelectedNode)
+	// Only allow when only a _single_ event node is selected, not allowed on the anim graph
+	UEdGraphNode* SelectedNode = GetSingleSelectedNode();
+	if (SelectedNode && !IsEditingAnimGraph())
 	{
-		if (UK2Node_FunctionEntry* const SelectedCallFunctionNode = Cast<UK2Node_FunctionEntry>(SelectedNode))
+		if (UK2Node_Event* SelectedEventNode = Cast<UK2Node_Event>(SelectedNode))
 		{
-			UFunction* const Func = SelectedCallFunctionNode->FindSignatureFunction();
-			
-			// If the function has no output parameters and is not an interface or construction script
-			if (Func && !UEdGraphSchema_K2::HasFunctionAnyOutputParameter(Func) &&
-				!FBlueprintEditorUtils::IsInterfaceFunction(GetBlueprintObj(), Func) &&
-				SelectedCallFunctionNode->FunctionReference.GetMemberName() != UEdGraphSchema_K2::FN_UserConstructionScript)
+			UFunction* const Func = FFunctionFromNodeHelper::FunctionFromNode(SelectedEventNode);
+
+			// If the function has no output parameters it is valid to make a function
+			if (!FBlueprintEditorUtils::IsInterfaceFunction(GetBlueprintObj(), Func) &&
+				!SelectedEventNode->IsInterfaceEventNode())
 			{
 				return true;
 			}
 		}
 	}
-	// We cannot convert functions that have output parameters to events
+
 	return false;
 }
 
@@ -5620,7 +5653,9 @@
 			{
 				FGuid GraphGuid;
 				FBlueprintEditorUtils::GetFunctionGuidFromClassByFieldName(FBlueprintEditorUtils::GetMostUpToDateClass(OverrideFuncClass), OriginalEventName, /* Out */GraphGuid);
-				NewEntryNode->FunctionReference.SetExternalMember(OriginalEventName, OverrideFuncClass, GraphGuid);
+		
+				// Set the external members of this function appropriately
+				NewEntryNode->FunctionReference.SetExternalMember(OriginalEventName, SelectedEventNode->EventReference.GetMemberParentClass(), GraphGuid);
 				Schema->ReconstructNode(*NewEntryNode);
 			}
 			
@@ -5694,9 +5729,13 @@
 			// Rename the function graph to the original name
 			FBlueprintEditorUtils::RenameGraph(NewGraph, *OriginalEventNameString);
 
-			// Create a function call the to the newly created graph in place of the old event
-			IBlueprintNodeBinder::FBindingSet Bindings;
-			UEdGraphNode* OutFunctionNode = UBlueprintFunctionNodeSpawner::Create(FindField<UFunction>(NodeBP->SkeletonGeneratedClass, OriginalEventName))->Invoke(SourceGraph, Bindings, NewFuncCallSpawn);
+			// If this function is blueprint callable, then spawn a function call node to it in place of the old event
+			UFunction* const NewFunction = FindField<UFunction>(NodeBP->SkeletonGeneratedClass, OriginalEventName);
+			if (NewFunction && NewFunction->HasAllFunctionFlags(FUNC_BlueprintCallable))
+			{
+				IBlueprintNodeBinder::FBindingSet Bindings;
+				UEdGraphNode* OutFunctionNode = UBlueprintFunctionNodeSpawner::Create(NewFunction)->Invoke(SourceGraph, Bindings, NewFuncCallSpawn);
+			}
 		}
 
 		FBlueprintEditorUtils::MarkBlueprintAsStructurallyModified(NodeBP);
@@ -5774,28 +5813,6 @@
 	return OutNodes;
 }
 
-bool FBlueprintEditor::CanConvertEventToFunction() const
-{
-	// Only allow when only a _single_ event node is selected, not allowed on the anim graph
-	UEdGraphNode* SelectedNode = GetSingleSelectedNode();
-	if (SelectedNode && !IsEditingAnimGraph())
-	{
-		if (UK2Node_Event* SelectedEventNode = Cast<UK2Node_Event>(SelectedNode))
-		{
-			UFunction* const Func = FFunctionFromNodeHelper::FunctionFromNode(SelectedEventNode);
-
-			// If the function has no output parameters it is valid to make a function
-			if (!FBlueprintEditorUtils::IsInterfaceFunction(GetBlueprintObj(), Func) &&
-				!SelectedEventNode->IsInterfaceEventNode())
-			{
-				return true;
-			}
-		}
-	}
-
-	return false;
-}
-
 void FBlueprintEditor::OnAlignTop()
 {
 	TSharedPtr<SGraphEditor> FocusedGraphEd = FocusedGraphEdPtr.Pin();
