Ensure that no function return nodes are added when an event is converted to a function

#rnx

diff --git a/Engine/Source/Editor/Kismet/Private/BlueprintEditor.cpp b/Engine/Source/Editor/Kismet/Private/BlueprintEditor.cpp
--- a/Engine/Source/Editor/Kismet/Private/BlueprintEditor.cpp
+++ b/Engine/Source/Editor/Kismet/Private/BlueprintEditor.cpp
@@ -5424,9 +5424,9 @@ void FBlueprintEditor::ConvertFunctionToEvent(UK2Node_FunctionEntry* SelectedCal
 		EventGraph->Modify();
 
 		// Find all return nodes and break their connections
-		TArray< UK2Node_FunctionTerminator*> FunctionReturnNodes;
+		TArray< UK2Node_FunctionResult*> FunctionReturnNodes;
 		FunctionGraph->GetNodesOfClass(FunctionReturnNodes);
-		for (UK2Node_FunctionTerminator* Node : FunctionReturnNodes)
+		for (UK2Node_FunctionResult* Node : FunctionReturnNodes)
 		{
 			if (Node)
 			{
@@ -5641,6 +5641,17 @@ void FBlueprintEditor::OnConvertEventToFunction()
 			{
 				FBlueprintEditorUtils::AddFunctionGraph<UFunction>(NodeBP, NewGraph, /*bIsUserCreated=*/ true, FunctionSig);
 			}
+
+			// Remove any return nodes that may have been added @see UE-78785
+			TArray<UK2Node_FunctionResult*> FunctionReturnNodes;
+			NewGraph->GetNodesOfClass(FunctionReturnNodes);
+			for (UK2Node_FunctionResult* Node : FunctionReturnNodes)
+			{
+				if (Node)
+				{
+					Node->DestroyNode();
+				}
+			}
 			
 			// Get the new entry node
 			TArray<UK2Node_FunctionEntry*> EntryNodes;
