Hiding Hide Unrelated Nodes button for data only BP
#rnx

diff --git a/Engine/Source/Editor/Kismet/Private/BlueprintEditor.cpp b/Engine/Source/Editor/Kismet/Private/BlueprintEditor.cpp
--- a/Engine/Source/Editor/Kismet/Private/BlueprintEditor.cpp
+++ b/Engine/Source/Editor/Kismet/Private/BlueprintEditor.cpp
@@ -1873,15 +1873,17 @@ void FBlueprintEditor::InitBlueprintEditor(
 		}
 	};
 
-	TSharedPtr<FExtender> ToolbarExtender = MakeShareable(new FExtender);
-	ToolbarExtender->AddToolBarExtension(
-		"Asset",
-		EExtensionHook::After,
-		GetToolkitCommands(),
-		FToolBarExtensionDelegate::CreateStatic( &Local::FillToolbar, GetToolkitCommands(), this )
-	);
-
-	AddToolbarExtender(ToolbarExtender);
+	if (!bShouldOpenInDefaultsMode)
+	{
+		TSharedPtr<FExtender> ToolbarExtender = MakeShareable(new FExtender);
+		ToolbarExtender->AddToolBarExtension(
+			"Asset",
+			EExtensionHook::After,
+			GetToolkitCommands(),
+			FToolBarExtensionDelegate::CreateStatic(&Local::FillToolbar, GetToolkitCommands(), this)
+		);
+		AddToolbarExtender(ToolbarExtender);
+	}
 
 	RegenerateMenusAndToolbars();
 
