Allow blueprint users to override the component class of a native introduced component.

diff --git a/Engine/Source/Developer/BlueprintCompilerCppBackend/Private/BlueprintCompilerCppBackendValueHelper.cpp b/Engine/Source/Developer/BlueprintCompilerCppBackend/Private/BlueprintCompilerCppBackendValueHelper.cpp
--- a/Engine/Source/Developer/BlueprintCompilerCppBackend/Private/BlueprintCompilerCppBackendValueHelper.cpp
+++ b/Engine/Source/Developer/BlueprintCompilerCppBackend/Private/BlueprintCompilerCppBackendValueHelper.cpp
@@ -1645,6 +1645,7 @@ void FEmitDefaultValueHelper::GenerateCustomDynamicClassInitialization(FEmitterL
 		Context.AddLine(FString::Printf(TEXT("ensure(0 == InDynamicClass->%s.Num());"), GET_MEMBER_NAME_STRING_CHECKED(UDynamicClass, DynamicBindingObjects)));
 		Context.AddLine(FString::Printf(TEXT("ensure(0 == InDynamicClass->%s.Num());"), GET_MEMBER_NAME_STRING_CHECKED(UDynamicClass, ComponentTemplates)));
 		Context.AddLine(FString::Printf(TEXT("ensure(0 == InDynamicClass->%s.Num());"), GET_MEMBER_NAME_STRING_CHECKED(UDynamicClass, Timelines)));
+		Context.AddLine(FString::Printf(TEXT("ensure(0 == InDynamicClass->%s.Num());"), GET_MEMBER_NAME_STRING_CHECKED(UDynamicClass, ComponentClassOverrides)));
 		Context.AddLine(FString::Printf(TEXT("ensure(nullptr == InDynamicClass->%s);"), GET_MEMBER_NAME_STRING_CHECKED(UDynamicClass, AnimClassImplementation)));
 		Context.AddLine(FString::Printf(TEXT("InDynamicClass->%s();"), GET_FUNCTION_NAME_STRING_CHECKED(UDynamicClass, AssembleReferenceTokenStream)));
 
@@ -1774,6 +1775,18 @@ void FEmitDefaultValueHelper::GenerateCustomDynamicClassInitialization(FEmitterL
 		CreateAndInitializeClassSubobjects(true, false);
 		CreateAndInitializeClassSubobjects(false, true);
 
+		for (const FBPComponentClassOverride& Override : BPGC->ComponentClassOverrides)
+		{
+			if (Override.ComponentClass)
+			{
+				Context.AddLine(FString::Printf(TEXT("InDynamicClass->ComponentClassOverrides.Emplace(MakeTuple(FName(\"%s\"), FindObject<UClass>(ANY_PACKAGE,TEXT(\"%s\"))));"), *Override.ComponentName.ToString(), *Override.ComponentClass->GetPathName()));
+			}
+			else
+			{
+				Context.AddLine(FString::Printf(TEXT("InDynamicClass->ComponentClassOverrides.Emplace(MakeTuple(FName(\"%s\"), nullptr));"), *Override.ComponentName.ToString(), *Override.ComponentClass->GetName()));
+			}
+		}
+
 		FBackendHelperAnim::CreateAnimClassData(Context);
 
 		Context.DecreaseIndent();
diff --git a/Engine/Source/Editor/Kismet/Private/BlueprintDetailsCustomization.cpp b/Engine/Source/Editor/Kismet/Private/BlueprintDetailsCustomization.cpp
--- a/Engine/Source/Editor/Kismet/Private/BlueprintDetailsCustomization.cpp
+++ b/Engine/Source/Editor/Kismet/Private/BlueprintDetailsCustomization.cpp
@@ -3802,7 +3802,7 @@ void FBlueprintGraphActionDetails::SetNetFlags( TWeakObjectPtr<UK2Node_EditableP
 				TypedEntryNode->SetExtraFlags(ExtraFlags);
 				bBlueprintModified = true;
 			}
-			if (UK2Node_CustomEvent * CustomEventNode = Cast<UK2Node_CustomEvent>(FunctionEntryNode.Get()))
+			if (UK2Node_CustomEvent* CustomEventNode = Cast<UK2Node_CustomEvent>(FunctionEntryNode.Get()))
 			{
 				CustomEventNode->FunctionFlags &= ~FlagsToClear;
 				CustomEventNode->FunctionFlags |= FlagsToSet;
@@ -3819,7 +3819,7 @@ void FBlueprintGraphActionDetails::SetNetFlags( TWeakObjectPtr<UK2Node_EditableP
 
 FText FBlueprintGraphActionDetails::GetCurrentReplicatedEventString() const
 {
-	const UK2Node_EditablePinBase * FunctionEntryNode = FunctionEntryNodePtr.Get();
+	const UK2Node_EditablePinBase* FunctionEntryNode = FunctionEntryNodePtr.Get();
 	const UK2Node_CustomEvent* CustomEvent = Cast<const UK2Node_CustomEvent>(FunctionEntryNode);
 
 	uint32 const ReplicatedNetMask = (FUNC_NetMulticast | FUNC_NetServer | FUNC_NetClient);
@@ -4682,7 +4682,7 @@ TSharedRef<ITableRow> FBlueprintGraphActionDetails::HandleGenerateRowAccessSpeci
 FText FBlueprintGraphActionDetails::GetCurrentAccessSpecifierName() const
 {
 	uint32 AccessSpecifierFlag = 0;
-	UK2Node_EditablePinBase * FunctionEntryNode = FunctionEntryNodePtr.Get();
+	UK2Node_EditablePinBase* FunctionEntryNode = FunctionEntryNodePtr.Get();
 	if(UK2Node_FunctionEntry* EntryNode = Cast<UK2Node_FunctionEntry>(FunctionEntryNode))
 	{
 		AccessSpecifierFlag = FUNC_AccessSpecifiers & EntryNode->GetFunctionFlags();
@@ -4698,7 +4698,7 @@ bool FBlueprintGraphActionDetails::IsAccessSpecifierVisible() const
 {
 	bool bSupportedType = false;
 	bool bIsEditable = false;
-	UK2Node_EditablePinBase * FunctionEntryNode = FunctionEntryNodePtr.Get();
+	UK2Node_EditablePinBase* FunctionEntryNode = FunctionEntryNodePtr.Get();
 	if(FunctionEntryNode)
 	{
 		UBlueprint* Blueprint = FunctionEntryNode->GetBlueprint();
@@ -4818,7 +4818,7 @@ bool FBlueprintGraphActionDetails::IsCustomEvent() const
 
 void FBlueprintGraphActionDetails::OnIsReliableReplicationFunctionModified(const ECheckBoxState NewCheckedState)
 {
-	UK2Node_EditablePinBase * FunctionEntryNode = FunctionEntryNodePtr.Get();
+	UK2Node_EditablePinBase* FunctionEntryNode = FunctionEntryNodePtr.Get();
 	UK2Node_CustomEvent* CustomEvent = Cast<UK2Node_CustomEvent>(FunctionEntryNode);
 	if( CustomEvent )
 	{
@@ -4828,7 +4828,7 @@ void FBlueprintGraphActionDetails::OnIsReliableReplicationFunctionModified(const
 			{
 				TypedEntryNode->AddExtraFlags(FUNC_NetReliable);
 			}
-			if (UK2Node_CustomEvent * CustomEventNode = Cast<UK2Node_CustomEvent>(FunctionEntryNode))
+			if (UK2Node_CustomEvent* CustomEventNode = Cast<UK2Node_CustomEvent>(FunctionEntryNode))
 			{
 				CustomEventNode->FunctionFlags |= FUNC_NetReliable;
 			}
@@ -4839,7 +4839,7 @@ void FBlueprintGraphActionDetails::OnIsReliableReplicationFunctionModified(const
 			{
 				TypedEntryNode->ClearExtraFlags(FUNC_NetReliable);
 			}
-			if (UK2Node_CustomEvent * CustomEventNode = Cast<UK2Node_CustomEvent>(FunctionEntryNode))
+			if (UK2Node_CustomEvent* CustomEventNode = Cast<UK2Node_CustomEvent>(FunctionEntryNode))
 			{
 				CustomEventNode->FunctionFlags &= ~FUNC_NetReliable;
 			}
@@ -4851,7 +4851,7 @@ void FBlueprintGraphActionDetails::OnIsReliableReplicationFunctionModified(const
 
 ECheckBoxState FBlueprintGraphActionDetails::GetIsReliableReplicatedFunction() const
 {
-	const UK2Node_EditablePinBase * FunctionEntryNode = FunctionEntryNodePtr.Get();
+	const UK2Node_EditablePinBase* FunctionEntryNode = FunctionEntryNodePtr.Get();
 	const UK2Node_CustomEvent* CustomEvent = Cast<const UK2Node_CustomEvent>(FunctionEntryNode);
 	if(!CustomEvent)
 	{
@@ -4871,7 +4871,7 @@ bool FBlueprintGraphActionDetails::IsPureFunctionVisible() const
 {
 	bool bSupportedType = false;
 	bool bIsEditable = false;
-	UK2Node_EditablePinBase * FunctionEntryNode = FunctionEntryNodePtr.Get();
+	UK2Node_EditablePinBase* FunctionEntryNode = FunctionEntryNodePtr.Get();
 	if(FunctionEntryNode)
 	{
 		UBlueprint* Blueprint = FunctionEntryNode->GetBlueprint();
@@ -4885,7 +4885,7 @@ bool FBlueprintGraphActionDetails::IsPureFunctionVisible() const
 
 void FBlueprintGraphActionDetails::OnIsPureFunctionModified( const ECheckBoxState NewCheckedState )
 {
-	UK2Node_EditablePinBase * FunctionEntryNode = FunctionEntryNodePtr.Get();
+	UK2Node_EditablePinBase* FunctionEntryNode = FunctionEntryNodePtr.Get();
 	UFunction* Function = FindFunction();
 	UK2Node_FunctionEntry* EntryNode = Cast<UK2Node_FunctionEntry>(FunctionEntryNode);
 	if (EntryNode && Function)
@@ -4916,7 +4916,7 @@ bool FBlueprintGraphActionDetails::IsConstFunctionVisible() const
 {
 	bool bSupportedType = false;
 	bool bIsEditable = false;
-	UK2Node_EditablePinBase * FunctionEntryNode = FunctionEntryNodePtr.Get();
+	UK2Node_EditablePinBase* FunctionEntryNode = FunctionEntryNodePtr.Get();
 	if(FunctionEntryNode)
 	{
 		bSupportedType = FunctionEntryNode->IsA<UK2Node_FunctionEntry>();
@@ -4927,7 +4927,7 @@ bool FBlueprintGraphActionDetails::IsConstFunctionVisible() const
 
 void FBlueprintGraphActionDetails::OnIsConstFunctionModified( const ECheckBoxState NewCheckedState )
 {
-	UK2Node_EditablePinBase * FunctionEntryNode = FunctionEntryNodePtr.Get();
+	UK2Node_EditablePinBase* FunctionEntryNode = FunctionEntryNodePtr.Get();
 	UFunction* Function = FindFunction();
 	UK2Node_FunctionEntry* EntryNode = Cast<UK2Node_FunctionEntry>(FunctionEntryNode);
 	if(EntryNode && Function)
@@ -4956,7 +4956,7 @@ ECheckBoxState FBlueprintGraphActionDetails::GetIsConstFunction() const
 
 FReply FBaseBlueprintGraphActionDetails::OnAddNewInputClicked()
 {
-	UK2Node_EditablePinBase * FunctionEntryNode = FunctionEntryNodePtr.Get();
+	UK2Node_EditablePinBase* FunctionEntryNode = FunctionEntryNodePtr.Get();
 
 	if( FunctionEntryNode )
 	{
@@ -5736,6 +5736,40 @@ void FBlueprintComponentDetails::CustomizeDetails(IDetailLayoutBuilder& DetailLa
 			]
 		];
 
+		UActorComponent* ComponentTemplate = (CachedNodePtr->IsNative() ? CachedNodePtr->GetComponentTemplate() : nullptr);
+		if (ComponentTemplate)
+		{
+			UClass* BaseClass = ComponentTemplate->GetClass();
+
+			if (const FBPComponentClassOverride* Override = BlueprintObj->ComponentClassOverrides.FindByKey(ComponentTemplate->GetFName()))
+			{
+				AActor* Owner = ComponentTemplate->GetOwner();
+				AActor* OwnerArchetype = CastChecked<AActor>(Owner->GetArchetype());
+				if (UActorComponent* ArchetypeComponent = Cast<UActorComponent>((UObject*)FindObjectWithOuter(OwnerArchetype, UActorComponent::StaticClass(), ComponentTemplate->GetFName())))
+				{
+					BaseClass = ArchetypeComponent->GetClass();
+				}
+			}
+
+			const FText ComponentClassTooltip = LOCTEXT("BlueprintComponentDetails_ComponentClassOverrideTooltip", "The class to use when creating this component for this class. This can only be done for components defined in native at this time.");
+
+			VariableCategory.AddCustomRow( LOCTEXT("BlueprintComponentDetails_ComponentClassOverride", "Component Class") )
+			.NameContent()
+			[
+				SNew(STextBlock)
+				.Text(LOCTEXT("BlueprintComponentDetails_ComponentClassOverrideLabel", "Component Class"))
+				.ToolTipText(ComponentClassTooltip)
+				.Font(IDetailLayoutBuilder::GetDetailFont())
+			]
+			.ValueContent()
+			[
+				SNew(SClassPropertyEntryBox)
+				.MetaClass(BaseClass)
+				.AllowNone(false)
+				.SelectedClass(this, &FBlueprintComponentDetails::GetSelectedEntryClass)
+				.OnSetClass(this, &FBlueprintComponentDetails::HandleNewEntryClassSelected)];
+		}
+
 		IDetailCategoryBuilder& SocketsCategory = DetailLayout.EditCategory("Sockets", LOCTEXT("BlueprintComponentDetailsCategory", "Sockets"), ECategoryPriority::Important);
 
 		SocketsCategory.AddCustomRow(LOCTEXT("BlueprintComponentDetails_Sockets", "Sockets"))
@@ -5995,6 +6029,87 @@ void FBlueprintComponentDetails::PopulateVariableCategories()
 	}
 }
 
+const UClass* FBlueprintComponentDetails::GetSelectedEntryClass() const
+{
+	check(CachedNodePtr.IsValid());
+
+	if (UActorComponent* ComponentTemplate = CachedNodePtr->GetComponentTemplate())
+	{
+		UBlueprint* BlueprintObj = GetBlueprintObj();
+		check(BlueprintObj);
+
+		do
+		{
+			if (const FBPComponentClassOverride* Override = BlueprintObj->ComponentClassOverrides.FindByKey(ComponentTemplate->GetFName()))
+			{
+				return Override->ComponentClass;
+			}
+
+			BlueprintObj = UBlueprint::GetBlueprintFromClass(Cast<UBlueprintGeneratedClass>(BlueprintObj->ParentClass));
+
+		} while (BlueprintObj);
+
+		return ComponentTemplate->GetClass();
+	}
+
+	return nullptr;
+
+}
+
+void FBlueprintComponentDetails::HandleNewEntryClassSelected(const UClass* NewEntryClass) const
+{
+	const UClass* PreviousClass = GetSelectedEntryClass();
+
+	if (PreviousClass != NewEntryClass)
+	{
+		check(CachedNodePtr.IsValid());
+
+		if (UActorComponent* ComponentTemplate = CachedNodePtr->GetComponentTemplate())
+		{
+			const FName ComponentTemplateName = ComponentTemplate->GetFName();
+
+			UBlueprint* BlueprintObj = GetBlueprintObj();
+			check(BlueprintObj);
+
+			if (FBPComponentClassOverride* Override = BlueprintObj->ComponentClassOverrides.FindByKey(ComponentTemplateName))
+			{
+				bool bRemoveEntry = false;
+				bool bFoundOverride = false;
+
+				UBlueprint* ParentBP = UBlueprint::GetBlueprintFromClass(Cast<UBlueprintGeneratedClass>(BlueprintObj->ParentClass));
+				while (ParentBP)
+				{
+					if (FBPComponentClassOverride* ParentOverride = BlueprintObj->ComponentClassOverrides.FindByKey(ComponentTemplateName))
+					{
+						bRemoveEntry = (ParentOverride->ComponentClass == NewEntryClass);
+						bFoundOverride = true;
+						break;
+					}
+				}
+				if (!bFoundOverride)
+				{
+					bRemoveEntry = (ComponentTemplate->GetClass() == NewEntryClass);
+				}
+				if (bRemoveEntry)
+				{
+					BlueprintObj->ComponentClassOverrides.RemoveAllSwap([ComponentTemplateName](const FBPComponentClassOverride& CCOverride) { return (CCOverride.ComponentName == ComponentTemplateName); });
+				}
+				else
+				{
+					Override->ComponentClass = const_cast<UClass*>(NewEntryClass);
+				}
+			}
+			else
+			{
+				BlueprintObj->ComponentClassOverrides.Emplace(FBPComponentClassOverride(ComponentTemplateName, const_cast<UClass*>(NewEntryClass)));
+			}
+
+			FBlueprintEditorUtils::MarkBlueprintAsStructurallyModified(BlueprintObj);
+		}
+	}
+}
+
+
 FText FBlueprintComponentDetails::GetSocketName() const
 {
 	check(CachedNodePtr.IsValid());
diff --git a/Engine/Source/Editor/Kismet/Private/BlueprintDetailsCustomization.h b/Engine/Source/Editor/Kismet/Private/BlueprintDetailsCustomization.h
--- a/Engine/Source/Editor/Kismet/Private/BlueprintDetailsCustomization.h
+++ b/Engine/Source/Editor/Kismet/Private/BlueprintDetailsCustomization.h
@@ -820,6 +820,8 @@ protected:
 	void OnVariableCategoryTextCommitted(const FText& NewText, ETextCommit::Type InTextCommit, FName VarName);
 	void OnVariableCategorySelectionChanged(TSharedPtr<FText> ProposedSelection, ESelectInfo::Type /*SelectInfo*/);
 	TSharedRef<ITableRow> MakeVariableCategoryViewWidget(TSharedPtr<FText> Item, const TSharedRef< STableViewBase >& OwnerTable);
+	const UClass* GetSelectedEntryClass() const;
+	void HandleNewEntryClassSelected(const UClass* NewEntryClass) const;
 
 	FText GetSocketName() const;
 	bool CanChangeSocket() const;
diff --git a/Engine/Source/Editor/KismetCompiler/Private/KismetCompiler.cpp b/Engine/Source/Editor/KismetCompiler/Private/KismetCompiler.cpp
--- a/Engine/Source/Editor/KismetCompiler/Private/KismetCompiler.cpp
+++ b/Engine/Source/Editor/KismetCompiler/Private/KismetCompiler.cpp
@@ -591,6 +591,44 @@ void FKismetCompilerContext::ValidateVariableNames()
 	}
 }
 
+void FKismetCompilerContext::ValidateComponentClassOverrides()
+{
+	if (UClass* ParentClass = Blueprint->ParentClass)
+	{
+		if (UObject* CDO = ParentClass->GetDefaultObject(false))
+		{
+			for (auto It(Blueprint->ComponentClassOverrides.CreateIterator()); It; ++It)
+			{
+				const FBPComponentClassOverride& Override = *It;
+				if (UObject* OverridenObject = (UObject*)FindObjectWithOuter(CDO, nullptr, Override.ComponentName))
+				{
+					if (Override.ComponentClass && !Override.ComponentClass->IsChildOf(OverridenObject->GetClass()))
+					{
+						MessageLog.Error(
+							*FText::Format(
+								LOCTEXT("InvalidOverride", "{0} is not a legal override for component {1} because it does not derive from {2}."),
+								FText::FromName(Override.ComponentClass->GetFName()),
+								FText::FromName(Override.ComponentName),
+								FText::FromName(OverridenObject->GetClass()->GetFName())
+								).ToString()
+						);
+					}
+				}
+				else
+				{
+					MessageLog.Note(
+						*FText::Format(
+							LOCTEXT("UnneededOverride", "Removing class override for component {0} that no longer exists."),
+							FText::FromName(Override.ComponentName)
+						).ToString()
+					);
+					It.RemoveCurrent();
+				}
+			}
+		}
+	}
+}
+
 void FKismetCompilerContext::ValidateTimelineNames()
 {
 	TSharedPtr<FKismetNameValidator> ParentBPNameValidator;
@@ -2449,6 +2487,7 @@ void FKismetCompilerContext::FinishCompilingClass(UClass* Class)
 		BPGClass->Timelines = Blueprint->Timelines;
 		BPGClass->SimpleConstructionScript = Blueprint->SimpleConstructionScript;
 		BPGClass->InheritableComponentHandler = Blueprint->InheritableComponentHandler;
+		BPGClass->ComponentClassOverrides = Blueprint->ComponentClassOverrides;
 	}
 
 	//@TODO: Not sure if doing this again is actually necessary
@@ -3952,6 +3991,8 @@ void FKismetCompilerContext::CompileClassLayout(EInternalCompilerFlags InternalF
 	// This validation requires CDO object.
 	ValidateVariableNames();
 
+	ValidateComponentClassOverrides();
+
 	OldCDO = NULL;
 	OldGenLinkerIdx = INDEX_NONE;
 	OldLinker = Blueprint->GetLinker();
diff --git a/Engine/Source/Editor/KismetCompiler/Public/KismetCompiler.h b/Engine/Source/Editor/KismetCompiler/Public/KismetCompiler.h
--- a/Engine/Source/Editor/KismetCompiler/Public/KismetCompiler.h
+++ b/Engine/Source/Editor/KismetCompiler/Public/KismetCompiler.h
@@ -267,7 +267,10 @@ public:
 
 	/** Ensures that all variables have valid names for compilation/replication */
 	void ValidateVariableNames();
-	
+
+	/** Ensures that all component class overrides are legal overrides of the parent class */
+	void ValidateComponentClassOverrides();
+
 	/** Creates a class variable for each entry in the Blueprint NewVars array */
 	virtual void CreateClassVariablesFromBlueprint();
 
diff --git a/Engine/Source/Editor/UnrealEd/Private/Kismet2/BlueprintEditorUtils.cpp b/Engine/Source/Editor/UnrealEd/Private/Kismet2/BlueprintEditorUtils.cpp
--- a/Engine/Source/Editor/UnrealEd/Private/Kismet2/BlueprintEditorUtils.cpp
+++ b/Engine/Source/Editor/UnrealEd/Private/Kismet2/BlueprintEditorUtils.cpp
@@ -3116,7 +3116,7 @@ bool FBlueprintEditorUtils::IsDataOnlyBlueprint(const UBlueprint* Blueprint)
 		return false;
 	}
 
-	if (Blueprint->ComponentTemplates.Num() || Blueprint->Timelines.Num())
+	if (Blueprint->ComponentTemplates.Num() > 0 || Blueprint->Timelines.Num() > 0 || Blueprint->ComponentClassOverrides.Num() > 0)
 	{
 		return false;
 	}
diff --git a/Engine/Source/Editor/UnrealEd/Private/Kismet2/Kismet2.cpp b/Engine/Source/Editor/UnrealEd/Private/Kismet2/Kismet2.cpp
--- a/Engine/Source/Editor/UnrealEd/Private/Kismet2/Kismet2.cpp
+++ b/Engine/Source/Editor/UnrealEd/Private/Kismet2/Kismet2.cpp
@@ -1124,9 +1124,25 @@ static void ConformComponentsUtils::ConformRemovedNativeComponents(UObject* BpCd
 		{
 			// Keep track of components inherited from the native super class that are still valid.
 			NewNativeComponents.Add(Component);
-
 			continue;
 		}
+
+		// If we have overriden the class of a native component then ensure that the component still exists and that the overriden class is a valid subclass of it
+		if (const FBPComponentClassOverride* BPCO = CastChecked<UBlueprintGeneratedClass>(BlueprintClass)->ComponentClassOverrides.FindByKey(Component->GetFName()))
+		{
+			if (BPCO->ComponentClass == Component->GetClass())
+			{
+				if (UObject* OverridenComponent = (UObject*)FindObjectWithOuter(NativeCDO, UActorComponent::StaticClass(), Component->GetFName()))
+				{
+					if (Component->IsA(OverridenComponent->GetClass()))
+					{
+						NewNativeComponents.Add(Component);
+						continue;
+					}
+				}
+			}
+		}
+
 		// else, the component has been removed from our native super class
 
 		Component->DestroyComponent(/*bPromoteChildren =*/false);
diff --git a/Engine/Source/Runtime/CoreUObject/Private/UObject/Class.cpp b/Engine/Source/Runtime/CoreUObject/Private/UObject/Class.cpp
--- a/Engine/Source/Runtime/CoreUObject/Private/UObject/Class.cpp
+++ b/Engine/Source/Runtime/CoreUObject/Private/UObject/Class.cpp
@@ -5316,6 +5316,11 @@ void UDynamicClass::AddReferencedObjects(UObject* InThis, FReferenceCollector& C
 	Collector.AddReferencedObjects(This->ComponentTemplates, This);
 	Collector.AddReferencedObjects(This->Timelines, This);
 
+	for (TPair<FName, UClass*>& Override : This->ComponentClassOverrides)
+	{
+		Collector.AddReferencedObject(Override.Value);
+	}
+
 	Collector.AddReferencedObject(This->AnimClassImplementation, This);
 
 	Super::AddReferencedObjects(This, Collector);
@@ -5343,6 +5348,7 @@ void UDynamicClass::PurgeClass(bool bRecompilingOnLoad)
 	DynamicBindingObjects.Empty();
 	ComponentTemplates.Empty();
 	Timelines.Empty();
+	ComponentClassOverrides.Empty();
 
 	AnimClassImplementation = nullptr;
 }
@@ -5369,6 +5375,17 @@ UObject* UDynamicClass::FindArchetype(UClass* ArchetypeClass, const FName Archet
 		(SuperClass ? SuperClass->FindArchetype(ArchetypeClass, ArchetypeName) : nullptr);
 }
 
+void UDynamicClass::SetupObjectInitializer(FObjectInitializer& ObjectInitializer) const
+{
+	for (const TPair<FName, UClass*>& Override : ComponentClassOverrides)
+	{
+		ObjectInitializer.SetDefaultSubobjectClass(Override.Key, Override.Value);
+	}
+
+	GetSuperClass()->SetupObjectInitializer(ObjectInitializer);
+}
+
+
 UStructProperty* UDynamicClass::FindStructPropertyChecked(const TCHAR* PropertyName) const
 {
 	return FindFieldChecked<UStructProperty>(this, PropertyName);
diff --git a/Engine/Source/Runtime/CoreUObject/Private/UObject/UObjectGlobals.cpp b/Engine/Source/Runtime/CoreUObject/Private/UObject/UObjectGlobals.cpp
--- a/Engine/Source/Runtime/CoreUObject/Private/UObject/UObjectGlobals.cpp
+++ b/Engine/Source/Runtime/CoreUObject/Private/UObject/UObjectGlobals.cpp
@@ -2557,6 +2557,11 @@ FObjectInitializer::FObjectInitializer(UObject* InObj, UObject* InObjectArchetyp
 	LastConstructedObject = ThreadContext.ConstructedObject;
 	ThreadContext.ConstructedObject = Obj;
 	ThreadContext.PushInitializer(this);
+
+	if (Obj)
+	{
+		Obj->GetClass()->SetupObjectInitializer(*this);
+	}
 }
 
 /**
diff --git a/Engine/Source/Runtime/CoreUObject/Public/UObject/Class.h b/Engine/Source/Runtime/CoreUObject/Public/UObject/Class.h
--- a/Engine/Source/Runtime/CoreUObject/Public/UObject/Class.h
+++ b/Engine/Source/Runtime/CoreUObject/Public/UObject/Class.h
@@ -2549,6 +2549,11 @@ public:
 	 */
 	virtual void InitPropertiesFromCustomList(uint8* DataPtr, const uint8* DefaultDataPtr) {}
 
+	/**
+	 * Allows class to provide data to the object initializer that can affect how native class subobjects are created.
+	 */
+	virtual void SetupObjectInitializer(FObjectInitializer& ObjectInitializer) const {}
+
 	/**
 	 * Get the name of the CDO for the this class
 	 * @return The name of the CDO
@@ -2888,6 +2893,7 @@ public:
 	virtual UObject* CreateDefaultObject();
 	virtual void PurgeClass(bool bRecompilingOnLoad) override;
 	virtual UObject* FindArchetype(UClass* ArchetypeClass, const FName ArchetypeName) const override;
+	virtual void SetupObjectInitializer(FObjectInitializer& ObjectInitializer) const override;
 
 	/** Find a struct property, called from generated code */
 	UStructProperty* FindStructPropertyChecked(const TCHAR* PropertyName) const;
@@ -2906,6 +2912,9 @@ public:
 	TArray<UObject*> ComponentTemplates;
 	TArray<UObject*> Timelines;
 
+	/** Array of blueprint overrides of component classes in parent classes */
+	TArray<TPair<FName, UClass*>> ComponentClassOverrides;
+
 	/** IAnimClassInterface (UAnimClassData) or null */
 	UObject* AnimClassImplementation;
 
diff --git a/Engine/Source/Runtime/Engine/Classes/Engine/Blueprint.h b/Engine/Source/Runtime/Engine/Classes/Engine/Blueprint.h
--- a/Engine/Source/Runtime/Engine/Classes/Engine/Blueprint.h
+++ b/Engine/Source/Runtime/Engine/Classes/Engine/Blueprint.h
@@ -11,6 +11,7 @@
 #include "Engine/EngineTypes.h"
 #include "EdGraph/EdGraphPin.h"
 #include "Engine/BlueprintCore.h"
+#include "Engine/BlueprintGeneratedClass.h"
 #include "UObject/SoftObjectPath.h"
 #include "Blueprint/BlueprintSupport.h"
 #include "Blueprint.generated.h"
@@ -594,6 +595,10 @@ public:
 	UPROPERTY()
 	TArray<class UTimelineTemplate*> Timelines;
 
+	/** Array of blueprint overrides of component classes in parent classes */
+	UPROPERTY()
+	TArray<FBPComponentClassOverride> ComponentClassOverrides;
+
 	/** Stores data to override (in children classes) components (created by SCS) from parent classes */
 	UPROPERTY()
 	class UInheritableComponentHandler* InheritableComponentHandler;
diff --git a/Engine/Source/Runtime/Engine/Classes/Engine/BlueprintGeneratedClass.h b/Engine/Source/Runtime/Engine/Classes/Engine/BlueprintGeneratedClass.h
--- a/Engine/Source/Runtime/Engine/Classes/Engine/BlueprintGeneratedClass.h
+++ b/Engine/Source/Runtime/Engine/Classes/Engine/BlueprintGeneratedClass.h
@@ -591,6 +591,37 @@ private:
 	TArray<uint8> CachedPropertyData;
 };
 
+/** Utility struct to store class overrides for components. */
+USTRUCT()
+struct FBPComponentClassOverride
+{
+	GENERATED_BODY()
+
+	/** The component name an override is being specified for. */
+	UPROPERTY()
+	FName ComponentName;
+
+	/** The class to use when constructing the component. */
+	UPROPERTY()
+	UClass* ComponentClass;
+
+	FBPComponentClassOverride()
+		: ComponentClass(nullptr)
+	{
+	}
+
+	FBPComponentClassOverride(FName InComponentName, UClass* InComponentClass)
+		: ComponentName(InComponentName)
+		, ComponentClass(InComponentClass)
+	{
+	}
+
+	bool operator==(const FName OtherComponentName) const
+	{
+		return (ComponentName == OtherComponentName);
+	}
+};
+
 UCLASS()
 class ENGINE_API UBlueprintGeneratedClass : public UClass
 {
@@ -626,6 +657,10 @@ public:
 	UPROPERTY()
 	TArray<class UTimelineTemplate*> Timelines;
 
+	/** Array of blueprint overrides of component classes in parent classes */
+	UPROPERTY()
+	TArray<FBPComponentClassOverride> ComponentClassOverrides;
+
 	/** 'Simple' construction script - graph of components to instance */
 	UPROPERTY()
 	class USimpleConstructionScript* SimpleConstructionScript;
@@ -715,6 +750,7 @@ public:
 	virtual UObject* FindArchetype(UClass* ArchetypeClass, const FName ArchetypeName) const override;
 
 	virtual void InitPropertiesFromCustomList(uint8* DataPtr, const uint8* DefaultDataPtr) override;
+	virtual void SetupObjectInitializer(FObjectInitializer& ObjectInitializer) const override;
 
 protected:
 
diff --git a/Engine/Source/Runtime/Engine/Private/BlueprintGeneratedClass.cpp b/Engine/Source/Runtime/Engine/Private/BlueprintGeneratedClass.cpp
--- a/Engine/Source/Runtime/Engine/Private/BlueprintGeneratedClass.cpp
+++ b/Engine/Source/Runtime/Engine/Private/BlueprintGeneratedClass.cpp
@@ -5,6 +5,7 @@
 #include "Stats/StatsMisc.h"
 #include "UObject/UObjectHash.h"
 #include "UObject/CoreNet.h"
+#include "UObject/CoreRedirects.h"
 #include "UObject/Package.h"
 #include "UObject/LinkerLoad.h"
 #include "Serialization/ObjectReader.h"
@@ -127,6 +128,28 @@ void UBlueprintGeneratedClass::PostLoad()
 	}
 #endif // WITH_EDITORONLY_DATA
 
+	// Update any component names that have been redirected
+	if (!FPlatformProperties::RequiresCookedData())
+	{
+		for (FBPComponentClassOverride& Override : ComponentClassOverrides)
+		{
+			const FString ComponentName = Override.ComponentName.ToString();
+			UClass* ClassToCheck = this;
+			while (ClassToCheck)
+			{
+				if (const TMap<FString, FString>* ValueChanges = FCoreRedirects::GetValueRedirects(ECoreRedirectFlags::Type_Class, ClassToCheck))
+				{
+					if (const FString* NewComponentName = ValueChanges->Find(ComponentName))
+					{
+						Override.ComponentName = **NewComponentName;
+						break;
+					}
+				}
+				ClassToCheck = ClassToCheck->GetSuperClass();
+			}
+		}
+	}
+
 	AssembleReferenceTokenStream(true);
 }
 
@@ -603,6 +626,16 @@ void UBlueprintGeneratedClass::UpdateCustomPropertyListForPostConstruction()
 	bCustomPropertyListForPostConstructionInitialized = true;
 }
 
+void UBlueprintGeneratedClass::SetupObjectInitializer(FObjectInitializer& ObjectInitializer) const
+{
+	for (const FBPComponentClassOverride& Override : ComponentClassOverrides)
+	{
+		ObjectInitializer.SetDefaultSubobjectClass(Override.ComponentName, Override.ComponentClass);
+	}
+
+	GetSuperClass()->SetupObjectInitializer(ObjectInitializer);
+}
+
 void UBlueprintGeneratedClass::InitPropertiesFromCustomList(uint8* DataPtr, const uint8* DefaultDataPtr)
 {
 	FScopeLock SerializeAndPostLoadLock(&SerializeAndPostLoadCritical);
@@ -1428,6 +1461,15 @@ void UBlueprintGeneratedClass::GetDefaultObjectPreloadDependencies(TArray<UObjec
 			OutDeps.Add(ComponentTemplate);
 		}
 	}
+
+	// Add the classes that will be used for overriding components defined in base classes
+	for (const FBPComponentClassOverride& Override : ComponentClassOverrides)
+	{
+		if (Override.ComponentClass)
+		{
+			OutDeps.Add(Override.ComponentClass);
+		}
+	}
 }
 
 bool UBlueprintGeneratedClass::NeedsLoadForServer() const
