Improve messaging for converting events and functions. There is now a failure message that pops up instead of the button just being disabled for the user.

#rnx

diff --git a/Engine/Source/Editor/Kismet/Private/BlueprintEditor.cpp b/Engine/Source/Editor/Kismet/Private/BlueprintEditor.cpp
--- a/Engine/Source/Editor/Kismet/Private/BlueprintEditor.cpp
+++ b/Engine/Source/Editor/Kismet/Private/BlueprintEditor.cpp
@@ -5372,38 +5372,79 @@ bool FBlueprintEditor::CanConvertFunctionToEvent() const
 	{
 		if (UK2Node_FunctionEntry* const SelectedCallFunctionNode = Cast<UK2Node_FunctionEntry>(SelectedNode))
 		{
-			UFunction* const Func = SelectedCallFunctionNode->FindSignatureFunction();
-
-			// If the function has no output parameters and is not an interface or construction script
-			if (Func && !UEdGraphSchema_K2::HasFunctionAnyOutputParameter(Func) &&
-				!FBlueprintEditorUtils::IsInterfaceFunction(GetBlueprintObj(), Func) &&
-				SelectedCallFunctionNode->FunctionReference.GetMemberName() != UEdGraphSchema_K2::FN_UserConstructionScript)
+			if (SelectedCallFunctionNode->FindSignatureFunction())
 			{
 				return true;
 			}
 		}
 	}
-	// We cannot convert functions that have output parameters to events
+
 	return false;
 }
 
 void FBlueprintEditor::OnConvertFunctionToEvent()
 {
+	FCompilerResultsLog LogResults;
+	FMessageLog MessageLog("BlueprintLog");
+	FText SpecificErrorMessage;
+
 	if (UEdGraphNode* SelectedNode = GetSingleSelectedNode())
 	{
 		if (UK2Node_FunctionEntry* SelectedCallFunctionNode = Cast<UK2Node_FunctionEntry>(SelectedNode))
 		{
-			ConvertFunctionToEvent(SelectedCallFunctionNode);
+			UFunction* Func = SelectedCallFunctionNode->FindSignatureFunction();
+			check(Func);
+
+			// Make sure there are no output parameters
+			if (UEdGraphSchema_K2::HasFunctionAnyOutputParameter(Func))
+			{
+				LogResults.Error(*LOCTEXT("FunctionHasOutput_Error", "A function can only be converted if it does not have any output parameters.").ToString());
+				SpecificErrorMessage = LOCTEXT("FunctionHasOutput_Error_Title", "Function cannot have output parameters");
+			}
+			// Make sure this is not an interface function
+			else if (FBlueprintEditorUtils::IsInterfaceFunction(GetBlueprintObj(), Func))
+			{
+				LogResults.Error(*LOCTEXT("FunctionIsFromInterface_Error", "Functions from interfaces cannot be converted to events.").ToString());
+				SpecificErrorMessage = LOCTEXT("FunctionIsFromInterface_Error_Title", "Cannot convert interface functions");
+			}
+			// Make sure this is not a blueprint/macro function library
+			else if (GetBlueprintObj()->BlueprintType == BPTYPE_FunctionLibrary || GetBlueprintObj()->BlueprintType == BPTYPE_MacroLibrary)
+			{
+				LogResults.Error(*LOCTEXT("BlueprintFunctionLibarary_Error", "Cannot convert functions from blueprint or macro libraries.").ToString());
+				SpecificErrorMessage = LOCTEXT("BlueprintFunctionLibarary_Error_Title", "Cannot convert blueprint or macro library functions");
+			}
+			// Ensure that this is no the construction script
+			else if (SelectedCallFunctionNode->FunctionReference.GetMemberName() == UEdGraphSchema_K2::FN_UserConstructionScript)
+			{
+				LogResults.Error(*LOCTEXT("ConvertConstructionScript_Error", "Cannot convert the construction script!").ToString());
+				SpecificErrorMessage = LOCTEXT("ConvertConstructionScript_Error_Title", "Cannot convert construction script");
+			}
+			// Make sure we are not on the animation graph
+			else if (IsEditingAnimGraph())
+			{
+				LogResults.Error(*LOCTEXT("ConvertAnimGraph_Error", "Cannot convert functions on the anim graph!").ToString());
+				SpecificErrorMessage = LOCTEXT("ConvertAnimGraph_Error_Title", "Cannot convert on the anim graph");
+			}
+			else
+			{
+				ConvertFunctionToEvent(SelectedCallFunctionNode);
+			}
 		}
 	}
 	else
 	{
-		FCompilerResultsLog LogResults;
 		LogResults.Error(*LOCTEXT("MultipleNodesSelectred_Error", "Only one node can be selected for conversion!").ToString());
-		FMessageLog MessageLog("BlueprintLog");
+	}
+
+	// Show the log results if there were any errors
+	if (LogResults.NumErrors)
+	{
 		MessageLog.NewPage(LOCTEXT("OnConvertEventToFunctionLabel", "Convert Event to Function"));
 		MessageLog.AddMessages(LogResults.Messages);
-		MessageLog.Notify(LOCTEXT("OnConvertEventToFunctionError", "Convert Event to Function Failed!"));
+
+		FFormatNamedArguments Arguments;
+		Arguments.Add(TEXT("SpecificErrorMessage"), SpecificErrorMessage);
+		MessageLog.Notify(FText::Format(LOCTEXT("OnConvertEventToFunctionError", "Convert Event to Function Failed!\n{SpecificErrorMessage}"), Arguments));
 	}
 }
 
@@ -5564,36 +5605,72 @@ bool FBlueprintEditor::CanConvertEventToFunction() const
 {
 	// Only allow when only a _single_ event node is selected, not allowed on the anim graph
 	UEdGraphNode* SelectedNode = GetSingleSelectedNode();
-	if (SelectedNode && !IsEditingAnimGraph())
+	if (SelectedNode && SelectedNode->IsA<UK2Node_Event>())
+	{
+		return true;
+	}
+
+	return false;
+}
+
+void FBlueprintEditor::OnConvertEventToFunction()
+{
+	FCompilerResultsLog LogResults;
+	FMessageLog MessageLog("BlueprintLog");
+	FText SpecificErrorMessage;
+
+	if (UEdGraphNode* SelectedNode = GetSingleSelectedNode())
 	{
 		if (UK2Node_Event* SelectedEventNode = Cast<UK2Node_Event>(SelectedNode))
 		{
 			UFunction* const Func = FFunctionFromNodeHelper::FunctionFromNode(SelectedEventNode);
 
-			// If the function has no output parameters it is valid to make a function
-			if (!FBlueprintEditorUtils::IsInterfaceFunction(GetBlueprintObj(), Func) &&
-				!SelectedEventNode->IsInterfaceEventNode())
+			// Ensure we are not trying to do this on an interface event
+			if (FBlueprintEditorUtils::IsInterfaceFunction(GetBlueprintObj(), Func) || SelectedEventNode->IsInterfaceEventNode())
 			{
-				return true;
+				LogResults.Error(*LOCTEXT("EventIsOnInterface_Error", "Only non-interface events can be converted to functions!").ToString());
+				SpecificErrorMessage = LOCTEXT("EventIsOnInterface_Error_Title", "Cannot convert interface events");
+			}
+			// Make sure that we are not on the animation graph
+			else if (IsEditingAnimGraph())
+			{
+				LogResults.Error(*LOCTEXT("ConvertAnimGraphEvent_Error", "Cannot convert events on the anim graph!").ToString());
+				SpecificErrorMessage = LOCTEXT("ConvertAnimGraphEvent_Error_Title", "Cannot convert on the anim graph");
+			}
+			else
+			{
+				ConvertEventToFunction(SelectedEventNode);
 			}
 		}
 	}
+	else
+	{
+		SpecificErrorMessage = LOCTEXT("MultipleNodesSelectred_Error_Title", "Only one node can be selected for conversion");
+		LogResults.Error(*LOCTEXT("MultipleNodesSelectred_Error", "Only one node can be selected for conversion!").ToString());
+	}
 
-	return false;
+	// Show the log results if there were any errors
+	if (LogResults.NumErrors)
+	{
+		MessageLog.NewPage(LOCTEXT("OnConvertEventToFunctionLabel", "Convert Event to Function"));
+		MessageLog.AddMessages(LogResults.Messages);
+
+		// Format the title node
+		FFormatNamedArguments Arguments;
+		Arguments.Add(TEXT("SpecificErrorMessage"), SpecificErrorMessage);
+		MessageLog.Notify(FText::Format(LOCTEXT("OnConvertEventToFunctionError", "Convert Event to Function Failed!\n{SpecificErrorMessage}"), Arguments));
+	}
 }
 
-void FBlueprintEditor::OnConvertEventToFunction()
-{
-	const FScopedTransaction Transaction(FGraphEditorCommands::Get().ConvertEventToFunction->GetLabel());
+void FBlueprintEditor::ConvertEventToFunction(UK2Node_Event* SelectedEventNode)
+{	
+	if (SelectedEventNode)
+	{
+		const FScopedTransaction Transaction(FGraphEditorCommands::Get().ConvertEventToFunction->GetLabel());
 
-	UBlueprint* const NodeBP = GetBlueprintObj();
-	NodeBP->Modify();
-	
-	// Keep track of the current selected node
-	UEdGraphNode* const SelectedNode = GetSingleSelectedNode();
+		UBlueprint* const NodeBP = GetBlueprintObj();
+		NodeBP->Modify();
 
-	if (UK2Node_Event* SelectedEventNode = Cast<UK2Node_Event>(SelectedNode))
-	{
 		const UEdGraphSchema_K2* Schema = GetDefault<UEdGraphSchema_K2>();
 
 		// Break any connections to the delegate pin, because that won't exist 
@@ -5609,9 +5686,9 @@ void FBlueprintEditor::OnConvertEventToFunction()
 
 		// Keep track of the old connections from the event node
 		TMap<FString, TSet<UEdGraphPin*>> PinConnections;
-		GetPinConnectionMap(SelectedNode, PinConnections);
+		GetPinConnectionMap(SelectedEventNode, PinConnections);
 
-		TArray<UEdGraphNode*> CollapsableNodes = GetAllConnectedNodes(SelectedNode);
+		TArray<UEdGraphNode*> CollapsableNodes = GetAllConnectedNodes(SelectedEventNode);
 
 		UFunction* const FunctionSig = FFunctionFromNodeHelper::FunctionFromNode(SelectedEventNode);
 		UEdGraph* const SourceGraph = SelectedEventNode->GetGraph();
diff --git a/Engine/Source/Editor/Kismet/Public/BlueprintEditor.h b/Engine/Source/Editor/Kismet/Public/BlueprintEditor.h
--- a/Engine/Source/Editor/Kismet/Public/BlueprintEditor.h
+++ b/Engine/Source/Editor/Kismet/Public/BlueprintEditor.h
@@ -46,6 +46,7 @@ class UUserDefinedStruct;
 class UBlueprintEditorOptions;
 struct Rect;
 class UK2Node_FunctionEntry;
+class UK2Node_Event;
 
 /* Enums to use when grouping the blueprint members in the list panel. The order here will determine the order in the list */
 namespace NodeSectionID
@@ -853,6 +854,7 @@ protected:
 	void ReconnectPinMap(UEdGraphNode* Node, const TMap<FString, TSet<UEdGraphPin*>>& PinConnections);
 
 	void OnConvertEventToFunction();
+	void ConvertEventToFunction(UK2Node_Event* SelectedEventNode);
 	bool CanConvertEventToFunction() const;
 
 	/**
