PR #5895: Disable Pass-by-Ref input modifier in macros (Contributed by [REDACTED_PERSON])

#rnx

diff --git a/Engine/Source/Editor/Kismet/Private/BlueprintDetailsCustomization.cpp b/Engine/Source/Editor/Kismet/Private/BlueprintDetailsCustomization.cpp
--- a/Engine/Source/Editor/Kismet/Private/BlueprintDetailsCustomization.cpp
+++ b/Engine/Source/Editor/Kismet/Private/BlueprintDetailsCustomization.cpp
@@ -2888,8 +2888,24 @@ void FBlueprintGraphArgumentLayout::GenerateChildContent( IDetailChildrenBuilder
 			}
 		}
 
+		bool bMacroGraph = false;
+		if (TargetNode)
+		{
+			if (const UBlueprint* Blueprint = TargetNode->GetBlueprint())
+			{
+				if (Blueprint->BlueprintType == BPTYPE_MacroLibrary)
+				{
+					bMacroGraph = true;
+				}
+				else if (const UEdGraph* Graph = TargetNode->GetGraph()) 
+				{
+					bMacroGraph = Blueprint->MacroGraphs.Contains(Graph);
+				}
+			}
+		}
+
 		// Exec pins can't be passed by reference
-		if (FoundPin == nullptr || !UEdGraphSchema_K2::IsExecPin(*FoundPin))
+		if ((FoundPin == nullptr || (!UEdGraphSchema_K2::IsExecPin(*FoundPin)) && !bMacroGraph))
 		{
 			ChildrenBuilder.AddCustomRow(LOCTEXT("FunctionArgDetailsPassByReference", "Pass-by-Reference"))
 				.NameContent()
