Fix logical op parentheses error on linux

#rnx

diff --git a/Engine/Source/Editor/Kismet/Private/BlueprintDetailsCustomization.cpp b/Engine/Source/Editor/Kismet/Private/BlueprintDetailsCustomization.cpp
--- a/Engine/Source/Editor/Kismet/Private/BlueprintDetailsCustomization.cpp
+++ b/Engine/Source/Editor/Kismet/Private/BlueprintDetailsCustomization.cpp
@@ -2905,7 +2905,7 @@ void FBlueprintGraphArgumentLayout::GenerateChildContent( IDetailChildrenBuilder
 		}
 
 		// Exec pins can't be passed by reference
-		if ((FoundPin == nullptr || (!UEdGraphSchema_K2::IsExecPin(*FoundPin)) && !bMacroGraph))
+		if (FoundPin && !UEdGraphSchema_K2::IsExecPin(*FoundPin) && !bMacroGraph)
 		{
 			ChildrenBuilder.AddCustomRow(LOCTEXT("FunctionArgDetailsPassByReference", "Pass-by-Reference"))
 				.NameContent()
