Disable drag and drop on events that are not marked as BP callable from the MyBlueprint window

diff --git a/Engine/Source/Editor/Kismet/Private/BPFunctionDragDropAction.cpp b/Engine/Source/Editor/Kismet/Private/BPFunctionDragDropAction.cpp
--- a/Engine/Source/Editor/Kismet/Private/BPFunctionDragDropAction.cpp
+++ b/Engine/Source/Editor/Kismet/Private/BPFunctionDragDropAction.cpp
@@ -7,6 +7,7 @@
 #include "EdGraphSchema_K2_Actions.h"
 #include "K2Node_CustomEvent.h"
 #include "K2Node_MacroInstance.h"
+#include "Kismet2/KismetEditorUtilities.h"
 
 #include "Kismet2/BlueprintEditorUtils.h"
 #include "BlueprintNodeBinder.h"
@@ -162,11 +163,18 @@ bool FKismetDragDropAction::ActionWillShowExistingNode() const
 			else if (SourceAction->GetTypeId() == FEdGraphSchemaAction_K2Event::StaticGetTypeId())
 			{
 				FEdGraphSchemaAction_K2Event* FuncAction = (FEdGraphSchemaAction_K2Event*)SourceAction.Get();
-				UK2Node_CustomEvent* CustomEvent = Cast<UK2Node_CustomEvent>(FuncAction->NodeTemplate);
-				// Drag and dropping custom event's will place a Call Function and will not focus the existing event
-				if (CustomEvent == nullptr)
+				// Drag and dropping custom events will place a Call Function and will not focus the existing event
+				if (UK2Node_Event* Event = CastChecked<UK2Node_Event>(FuncAction->NodeTemplate))
 				{
-					bWillFocusOnExistingNode = true;
+					UFunction* Function = FFunctionFromNodeHelper::FunctionFromNode(Event);
+					if (Function && (Function->FunctionFlags & (FUNC_BlueprintCallable | FUNC_BlueprintPure)))
+					{
+						bWillFocusOnExistingNode = false;
+					}
+					else
+					{
+						bWillFocusOnExistingNode = true;
+					}
 				}
 			}
 		}
@@ -235,7 +243,7 @@ FReply FKismetFunctionDragDropAction::DroppedOnPin(FVector2D ScreenPosition, FVe
 		if (!CanBeDroppedDelegate.IsBound() || CanBeDroppedDelegate.Execute(nullptr, Graph, CannotDropReason))
 		{
 			UFunction const* Function = GetFunctionProperty();
-			if ((Function != NULL) && UEdGraphSchema_K2::CanUserKismetCallFunction(Function))
+			if ((Function != nullptr) && UEdGraphSchema_K2::CanUserKismetCallFunction(Function))
 			{
 				AnalyticCallback.ExecuteIfBound();
 
@@ -257,6 +265,13 @@ FReply FKismetFunctionDragDropAction::DroppedOnPin(FVector2D ScreenPosition, FVe
 			}
 		}
 	}
+	else if (SourceAction->GetTypeId() == FEdGraphSchemaAction_K2Event::StaticGetTypeId())
+	{
+		if (FEdGraphSchemaAction_K2Event* FuncAction = (FEdGraphSchemaAction_K2Event*)SourceAction.Get())
+		{
+			FKismetEditorUtilities::BringKismetToFocusAttentionOnObject(FuncAction->NodeTemplate);
+		}
+	}
 
 	return Reply;
 }
@@ -288,7 +303,10 @@ UBlueprintFunctionNodeSpawner* FKismetFunctionDragDropAction::GetDropAction(UEdG
 
 			if (UFunction const* Function = GetFunctionProperty())
 			{
-				return UBlueprintFunctionNodeSpawner::Create(Function);
+				if (Function->FunctionFlags & (FUNC_BlueprintCallable | FUNC_BlueprintPure))
+				{
+					return UBlueprintFunctionNodeSpawner::Create(Function);
+				}
 			}
 		}
 	}
diff --git a/Engine/Source/Editor/Kismet/Private/SMyBlueprint.cpp b/Engine/Source/Editor/Kismet/Private/SMyBlueprint.cpp
--- a/Engine/Source/Editor/Kismet/Private/SMyBlueprint.cpp
+++ b/Engine/Source/Editor/Kismet/Private/SMyBlueprint.cpp
@@ -1567,7 +1567,7 @@ FReply SMyBlueprint::OnActionDragged( const TArray< TSharedPtr<FEdGraphSchemaAct
 		return FReply::Unhandled();
 	}
 
-	TSharedPtr<FEdGraphSchemaAction> InAction( InActions.Num() > 0 ? InActions[0] : NULL );
+	TSharedPtr<FEdGraphSchemaAction> InAction( InActions.Num() > 0 ? InActions[0] : nullptr );
 	if(InAction.IsValid())
 	{
 		auto AnalyticsDelegate = FNodeCreationAnalytic::CreateSP( this, &SMyBlueprint::UpdateNodeCreation );
@@ -1655,17 +1655,32 @@ FReply SMyBlueprint::OnActionDragged( const TArray< TSharedPtr<FEdGraphSchemaAct
 		{	
 			// Check if it's a custom event, it is preferable to drop a call function for custom events than to focus on the node
 			FEdGraphSchemaAction_K2Event* FuncAction = (FEdGraphSchemaAction_K2Event*)InAction.Get();
-			if (UK2Node_CustomEvent* CustomEvent = Cast<UK2Node_CustomEvent>(FuncAction->NodeTemplate))
+			if (UK2Node_Event* Event = Cast<UK2Node_Event>(FuncAction->NodeTemplate))
 			{
-				return FReply::Handled().BeginDragDrop(FKismetFunctionDragDropAction::New(InAction, CustomEvent->GetFunctionName(), GetBlueprintObj()->SkeletonGeneratedClass, FMemberReference(), AnalyticsDelegate, FKismetDragDropAction::FCanBeDroppedDelegate()));
-			}
-			else
-			{
-				// don't need a valid FCanBeDroppedDelegate because this entry means we already have this 
-				// event placed (so this action will just focus it)
-				TSharedRef<FKismetDragDropAction> DragOperation = FKismetDragDropAction::New(InAction, AnalyticsDelegate, FKismetDragDropAction::FCanBeDroppedDelegate());
+				UFunction* const Function = FFunctionFromNodeHelper::FunctionFromNode(Event);
 
-				return FReply::Handled().BeginDragDrop(DragOperation);
+				// Callback function to report that the user cannot drop this function in the graph
+				auto CanDragDropAction = [](TSharedPtr<FEdGraphSchemaAction> /*DropAction*/, UEdGraph* /*HoveredGraphIn*/, FText& ImpededReasonOut, UFunction* Func)->bool
+				{
+					// If this function is not BP callable then don't let it be dropped
+					if (Func && !(Func->FunctionFlags & (FUNC_BlueprintCallable | FUNC_BlueprintPure)))
+					{
+						ImpededReasonOut = LOCTEXT("NonBlueprintCallableEvent", "This event was not marked as Blueprint Callable and cannot be placed in a graph!");
+						return false;
+					}
+
+					return true;
+				};
+
+				TSharedRef< FKismetFunctionDragDropAction> DragOperation =
+					FKismetFunctionDragDropAction::New(
+						InAction, Function->GetFName(),
+						GetBlueprintObj()->SkeletonGeneratedClass,
+						FMemberReference(),
+						AnalyticsDelegate,
+						FKismetDragDropAction::FCanBeDroppedDelegate::CreateLambda(CanDragDropAction, Function)
+					);
+				return FReply::Handled().BeginDragDrop(DragOperation);	
 			}
 		}
 	}
