Add keyword support for Custom Events

User can now type in additional keywords to the details panel of a custom event and search for them in the BP context menu

diff --git a/Engine/Source/Editor/BlueprintGraph/Classes/K2Node_CustomEvent.h b/Engine/Source/Editor/BlueprintGraph/Classes/K2Node_CustomEvent.h
--- a/Engine/Source/Editor/BlueprintGraph/Classes/K2Node_CustomEvent.h
+++ b/Engine/Source/Editor/BlueprintGraph/Classes/K2Node_CustomEvent.h
@@ -13,6 +13,7 @@
 class FBlueprintActionDatabaseRegistrar;
 class INameValidatorInterface;
 class UEdGraph;
+struct FKismetUserDeclaredFunctionMetadata;
 
 UCLASS(MinimalAPI)
 class UK2Node_CustomEvent : public UK2Node_Event
@@ -103,7 +104,15 @@ class UK2Node_CustomEvent : public UK2Node_Event
 private:
 	/** Constructing FText strings can be costly, so we cache the node's title */
 	FNodeTextCache CachedNodeTitle;
-};
 
+	/** Custom event metadata that can be used for adding custom keywords */
+	UPROPERTY()
+	FKismetUserDeclaredFunctionMetadata MetaData;
 
+public:
 
+	FKismetUserDeclaredFunctionMetadata& GetUserDefinedMetaData()
+	{
+		return MetaData;
+	}
+};
\ No newline at end of file
diff --git a/Engine/Source/Editor/Kismet/Private/BlueprintCompilationManager.cpp b/Engine/Source/Editor/Kismet/Private/BlueprintCompilationManager.cpp
--- a/Engine/Source/Editor/Kismet/Private/BlueprintCompilationManager.cpp
+++ b/Engine/Source/Editor/Kismet/Private/BlueprintCompilationManager.cpp
@@ -2389,7 +2389,7 @@ UClass* FBlueprintCompilationManagerImpl::FastGenerateSkeletonClass(UBlueprint*
 	UField** CurrentFieldStorageLocation = &Ret->Children;
 	
 	// Helper function for making UFunctions generated for 'event' nodes, e.g. custom event and timelines
-	const auto MakeEventFunction = [&CurrentFieldStorageLocation, MakeFunction, Schema]( FName InName, EFunctionFlags ExtraFnFlags, const TArray<UEdGraphPin*>& InputPins, const TArray< TSharedPtr<FUserPinInfo> >& UserPins, UFunction* InSourceFN, bool bInCallInEditor, bool bIsDeprecated, const FString& DeprecationMessage )
+	const auto MakeEventFunction = [&CurrentFieldStorageLocation, MakeFunction, Schema]( FName InName, EFunctionFlags ExtraFnFlags, const TArray<UEdGraphPin*>& InputPins, const TArray< TSharedPtr<FUserPinInfo> >& UserPins, UFunction* InSourceFN, bool bInCallInEditor, bool bIsDeprecated, const FString& DeprecationMessage, FKismetUserDeclaredFunctionMetadata* UserDefinedMetaData = nullptr)
 	{
 		UField** CurrentParamStorageLocation = nullptr;
 
@@ -2423,6 +2423,11 @@ UClass* FBlueprintCompilationManagerImpl::FastGenerateSkeletonClass(UBlueprint*
 				NewFunction->SetMetaData(FBlueprintMetadata::MD_CallInEditor, TEXT( "true" ));
 			}
 
+			if (UserDefinedMetaData)
+			{
+				NewFunction->SetMetaData(FBlueprintMetadata::MD_FunctionKeywords, *(UserDefinedMetaData->Keywords).ToString());
+			}
+
 			NewFunction->Bind();
 			NewFunction->StaticLink(true);
 		}
@@ -2460,6 +2465,7 @@ UClass* FBlueprintCompilationManagerImpl::FastGenerateSkeletonClass(UBlueprint*
 			FString DeprecationMessage;
 			bool bIsDeprecated = false;
 			bool bCallInEditor = false;
+			FKismetUserDeclaredFunctionMetadata* UserMetaData = nullptr;
 			if(UK2Node_CustomEvent* CustomEvent = Cast<UK2Node_CustomEvent>(Event))
 			{
 				bCallInEditor = CustomEvent->bCallInEditor;
@@ -2468,6 +2474,7 @@ UClass* FBlueprintCompilationManagerImpl::FastGenerateSkeletonClass(UBlueprint*
 				{
 					DeprecationMessage = CustomEvent->DeprecationMessage;
 				}
+				UserMetaData = &(CustomEvent->GetUserDefinedMetaData());
 			}
 			MakeEventFunction(
 				CompilerContext.GetEventStubFunctionName(Event), 
@@ -2477,7 +2484,8 @@ UClass* FBlueprintCompilationManagerImpl::FastGenerateSkeletonClass(UBlueprint*
 				Event->FindEventSignatureFunction(),
 				bCallInEditor,
 				bIsDeprecated,
-				DeprecationMessage
+				DeprecationMessage,
+				UserMetaData
 			);
 		}
 	}
diff --git a/Engine/Source/Editor/Kismet/Private/BlueprintDetailsCustomization.cpp b/Engine/Source/Editor/Kismet/Private/BlueprintDetailsCustomization.cpp
--- a/Engine/Source/Editor/Kismet/Private/BlueprintDetailsCustomization.cpp
+++ b/Engine/Source/Editor/Kismet/Private/BlueprintDetailsCustomization.cpp
@@ -3493,6 +3493,27 @@ void FBlueprintGraphActionDetails::CustomizeDetails( IDetailLayoutBuilder& Detai
 										FSlateIcon(),
 										FUIAction(FExecuteAction::CreateStatic( &FBlueprintGraphActionDetails::SetNetFlags, FunctionEntryNodePtr, static_cast<uint32>(FUNC_NetClient) ), CanExecuteDelegate));
 
+			const FString DocLink = TEXT("Shared/Editors/BlueprintEditor/GraphDetails");
+			TSharedPtr<SToolTip> KeywordsTooltip = IDocumentation::Get()->CreateToolTip(LOCTEXT("EditEventKeywords_Tooltip", "Keywords for searching for the event."), nullptr, DocLink, TEXT("Keywords"));
+			Category.AddCustomRow(LOCTEXT("EventsKeywordsLabel", "Keywords"))
+			.NameContent()
+			[
+				SNew(STextBlock)
+					.Text(LOCTEXT("KeywordsLabel", "Keywords"))
+					.ToolTip(KeywordsTooltip)
+					.Font(IDetailLayoutBuilder::GetDetailFont())
+			]
+			.ValueContent()
+			[
+				SNew(SEditableTextBox)
+					.Text(this, &FBlueprintGraphActionDetails::OnGetKeywordsText)
+					.OnTextCommitted(this, &FBlueprintGraphActionDetails::OnKeywordsTextCommitted)
+					.ToolTip(KeywordsTooltip)
+					.RevertTextOnEscape(true)
+					.Font(IDetailLayoutBuilder::GetDetailFont())
+			];
+
+
 			Category.AddCustomRow( LOCTEXT( "FunctionReplicate", "Replicates" ) )
 			.NameContent()
 			[
@@ -4460,10 +4481,13 @@ UK2Node_EditablePinBase* FBlueprintGraphActionDetails::GetEditableNode() const
 
 UFunction* FBlueprintGraphActionDetails::FindFunction() const
 {
-	UEdGraph* Graph = GetGraph();
-	if(Graph)
+	if (UK2Node_CustomEvent* EventNode = Cast<UK2Node_CustomEvent>(FunctionEntryNodePtr.Get()))
 	{
-		if(UBlueprint* Blueprint = FBlueprintEditorUtils::FindBlueprintForGraph(Graph))
+		return FFunctionFromNodeHelper::FunctionFromNode(EventNode);
+	}
+	else if (UEdGraph* Graph = GetGraph())
+	{
+		if (UBlueprint* Blueprint = FBlueprintEditorUtils::FindBlueprintForGraph(Graph))
 		{
 			UClass* Class = Blueprint->SkeletonGeneratedClass;
 
@@ -4477,7 +4501,7 @@ UFunction* FBlueprintGraphActionDetails::FindFunction() const
 			}
 		}
 	}
-	return NULL;
+	return nullptr;
 }
 
 FKismetUserDeclaredFunctionMetadata* FBlueprintGraphActionDetails::GetMetadataBlock() const
@@ -4492,8 +4516,12 @@ FKismetUserDeclaredFunctionMetadata* FBlueprintGraphActionDetails::GetMetadataBl
 		// Must be exactly a tunnel, not a macro instance
 		return &(TunnelNode->MetaData);
 	}
+	else if (UK2Node_CustomEvent* EventNode = Cast<UK2Node_CustomEvent>(FunctionEntryNode))
+	{
+		return &(EventNode->GetUserDefinedMetaData());
+	}
 
-	return NULL;
+	return nullptr;
 }
 
 FText FBlueprintGraphActionDetails::OnGetTooltipText() const
