Allow interfaces to show up in the "My Blueprint" panel in child BP's

#rnx

diff --git a/Engine/Source/Editor/Kismet/Private/SMyBlueprint.cpp b/Engine/Source/Editor/Kismet/Private/SMyBlueprint.cpp
--- a/Engine/Source/Editor/Kismet/Private/SMyBlueprint.cpp
+++ b/Engine/Source/Editor/Kismet/Private/SMyBlueprint.cpp
@@ -1313,7 +1313,27 @@ void SMyBlueprint::CollectAllActions(FGraphActionListBuilderBase& OutAllActions)
 		}
 	}
 
-	// Also function implemented for interfaces
+	auto IsInAnimBPLambda = [&BlueprintObj](const FName FunctionName, FText& FunctionCategory) -> bool
+	{
+		if (BlueprintObj->SkeletonGeneratedClass != nullptr)
+		{
+			if (UFunction * Function = BlueprintObj->SkeletonGeneratedClass->FindFunctionByName(FunctionName))
+			{
+				FunctionCategory = Function->GetMetaDataText(FBlueprintMetadata::MD_FunctionCategory, TEXT("UObjectCategory"), Function->GetFullGroupName(false));
+
+				if (IAnimClassInterface * AnimClassInterface = IAnimClassInterface::GetFromClass(BlueprintObj->SkeletonGeneratedClass))
+				{
+					if (IAnimClassInterface::IsAnimBlueprintFunction(AnimClassInterface, Function))
+					{
+						return true;
+					}
+				}
+			}
+		}
+		return false;
+	};
+
+	// Also functions implemented from interfaces
 	for (int32 i=0; i < BlueprintObj->ImplementedInterfaces.Num(); i++)
 	{
 		FBPInterfaceDescription& InterfaceDesc = BlueprintObj->ImplementedInterfaces[i];
@@ -1322,33 +1342,18 @@ void SMyBlueprint::CollectAllActions(FGraphActionListBuilderBase& OutAllActions)
 		{
 			for (TFieldIterator<UFunction> FunctionIt(InterfaceClass, EFieldIteratorFlags::IncludeSuper); FunctionIt; ++FunctionIt)
 			{
-				const FName FunctionName = (*FunctionIt)->GetFName();
+				const UFunction* Function = *FunctionIt;
+				const FName FunctionName = Function->GetFName();
 
 				if (FunctionName != UEdGraphSchema_K2::FN_ExecuteUbergraphBase)
 				{
-					FString FunctionTooltip = FunctionName.ToString();
-					FString FunctionDesc = FunctionName.ToString();
+					FText FunctionTooltip = Function->GetToolTipText();
+					FText FunctionDesc = K2Schema->GetFriendlySignatureName(Function);
 
 					FText FunctionCategory;
-					bool bIsAnimFunction = false;
+					bool bIsAnimFunction = IsInAnimBPLambda(FunctionName, FunctionCategory);
 
-					if (BlueprintObj->SkeletonGeneratedClass != nullptr)
-					{
-						if (UFunction * Function = BlueprintObj->SkeletonGeneratedClass->FindFunctionByName(FunctionName))
-						{
-							FunctionCategory = Function->GetMetaDataText(FBlueprintMetadata::MD_FunctionCategory, TEXT("UObjectCategory"), Function->GetFullGroupName(false));
-
-							if (IAnimClassInterface * AnimClassInterface = IAnimClassInterface::GetFromClass(BlueprintObj->SkeletonGeneratedClass))
-							{
-								if (IAnimClassInterface::IsAnimBlueprintFunction(AnimClassInterface, Function))
-								{
-									bIsAnimFunction = true;
-								}
-							}
-						}
-					}
-
-					TSharedPtr<FEdGraphSchemaAction_K2Graph> NewFuncAction = MakeShareable(new FEdGraphSchemaAction_K2Graph(EEdGraphSchemaAction_K2Graph::Interface, FunctionCategory, FText::FromString(FunctionDesc), FText::FromString(FunctionTooltip), 1, bIsAnimFunction ? NodeSectionID::ANIMLAYER : NodeSectionID::INTERFACE));
+					TSharedPtr<FEdGraphSchemaAction_K2Graph> NewFuncAction = MakeShareable(new FEdGraphSchemaAction_K2Graph(EEdGraphSchemaAction_K2Graph::Interface, FunctionCategory, FunctionDesc, FunctionTooltip, 1, bIsAnimFunction ? NodeSectionID::ANIMLAYER : NodeSectionID::INTERFACE));
 
 					NewFuncAction->FuncName = FunctionName;
 					OutAllActions.AddAction(NewFuncAction);
@@ -1380,27 +1385,31 @@ void SMyBlueprint::CollectAllActions(FGraphActionListBuilderBase& OutAllActions)
 		for (int32 Idx = 0; Idx < TempClass->Interfaces.Num(); ++Idx)
 		{
 			FImplementedInterface const& I = TempClass->Interfaces[Idx];
-			if (!I.bImplementedByK2)
+			
+			// same as above, make a function?
+			for (TFieldIterator<UFunction> FunctionIt(I.Class, EFieldIteratorFlags::IncludeSuper); FunctionIt; ++FunctionIt)
 			{
-				// same as above, make a function?
-				for (TFieldIterator<UFunction> FunctionIt(I.Class, EFieldIteratorFlags::IncludeSuper); FunctionIt; ++FunctionIt)
-				{
-					const UFunction* Function = *FunctionIt;
-					const FName FunctionName = Function->GetFName();
+				const UFunction* Function = *FunctionIt;
+				const FName FunctionName = Function->GetFName();
 
-					if (UEdGraphSchema_K2::CanKismetOverrideFunction(Function) && !ImplementedFunctionCache.Contains(FunctionName) && !UEdGraphSchema_K2::FunctionCanBePlacedAsEvent(Function))
-					{
-						FText FunctionTooltip = Function->GetToolTipText();
-						FText FunctionDesc = K2Schema->GetFriendlySignatureName(Function);
+				if (UEdGraphSchema_K2::CanKismetOverrideFunction(Function) && !ImplementedFunctionCache.Contains(FunctionName))
+				{
+					FText FunctionTooltip = Function->GetToolTipText();
+					FText FunctionDesc = K2Schema->GetFriendlySignatureName(Function);
 
-						FText FunctionCategory = Function->GetMetaDataText(FBlueprintMetadata::MD_FunctionCategory, TEXT("UObjectCategory"), Function->GetFullGroupName(false));
+					FText FunctionCategory = Function->GetMetaDataText(FBlueprintMetadata::MD_FunctionCategory, TEXT("UObjectCategory"), Function->GetFullGroupName(false));
+					bool bIsAnimFunction = IsInAnimBPLambda(FunctionName, FunctionCategory);
 
-						TSharedPtr<FEdGraphSchemaAction_K2Graph> NewFuncAction = MakeShareable(new FEdGraphSchemaAction_K2Graph(EEdGraphSchemaAction_K2Graph::Function, FunctionCategory, FunctionDesc, FunctionTooltip, 1, NodeSectionID::INTERFACE));
-						NewFuncAction->FuncName = FunctionName;
+					TSharedPtr<FEdGraphSchemaAction_K2Graph> NewFuncAction = MakeShareable(new FEdGraphSchemaAction_K2Graph(EEdGraphSchemaAction_K2Graph::Interface, FunctionCategory, FunctionDesc, FunctionTooltip, 1, bIsAnimFunction ? NodeSectionID::ANIMLAYER : NodeSectionID::INTERFACE));
+					NewFuncAction->FuncName = FunctionName;
+					
+					if (!OverridableFunctionNames.Contains(FunctionName))
+					{
 						OverridableFunctionActions.Add(NewFuncAction);
 						OverridableFunctionNames.Add(FunctionName);
-						OutAllActions.AddAction(NewFuncAction);
 					}
+
+					OutAllActions.AddAction(NewFuncAction);
 				}
 			}
 		}
