A more complete fix for including Hide Unrelated Nodes only where there are graph nodes.
#rnx

diff --git a/Engine/Source/Editor/Kismet/Private/BlueprintEditor.cpp b/Engine/Source/Editor/Kismet/Private/BlueprintEditor.cpp
--- a/Engine/Source/Editor/Kismet/Private/BlueprintEditor.cpp
+++ b/Engine/Source/Editor/Kismet/Private/BlueprintEditor.cpp
@@ -1847,44 +1847,6 @@ void FBlueprintEditor::InitBlueprintEditor(
 
 	InitalizeExtenders();
 
-	struct Local
-	{
-		static void FillToolbar(FToolBarBuilder& ToolbarBuilder, const TSharedRef< FUICommandList > ToolkitCommands, FBlueprintEditor* BlueprintEditor)
-		{
-			ToolbarBuilder.BeginSection("Graph");
-			{
-				ToolbarBuilder.AddToolBarButton(
-					FBlueprintEditorCommands::Get().ToggleHideUnrelatedNodes,
-					NAME_None,
-					TAttribute<FText>(),
-					TAttribute<FText>(),
-					FSlateIcon(FEditorStyle::GetStyleSetName(), "GraphEditor.ToggleHideUnrelatedNodes")
-				);
-				ToolbarBuilder.AddComboButton(
-					FUIAction(),
-					FOnGetContent::CreateSP(BlueprintEditor, &FBlueprintEditor::MakeHideUnrelatedNodesOptionsMenu),
-					LOCTEXT("HideUnrelatedNodesOptions", "Focus Related Nodes Options"),
-					LOCTEXT("HideUnrelatedNodesOptionsMenu", "Focus Related Nodes options menu"),
-					TAttribute<FSlateIcon>(),
-					true
-				);
-			}
-			ToolbarBuilder.EndSection();
-		}
-	};
-
-	if (!bShouldOpenInDefaultsMode)
-	{
-		TSharedPtr<FExtender> ToolbarExtender = MakeShareable(new FExtender);
-		ToolbarExtender->AddToolBarExtension(
-			"Asset",
-			EExtensionHook::After,
-			GetToolkitCommands(),
-			FToolBarExtensionDelegate::CreateStatic(&Local::FillToolbar, GetToolkitCommands(), this)
-		);
-		AddToolbarExtender(ToolbarExtender);
-	}
-
 	RegenerateMenusAndToolbars();
 
 	RegisterApplicationModes(InBlueprints, bShouldOpenInDefaultsMode, bNewlyCreated);
diff --git a/Engine/Source/Editor/Kismet/Private/SBlueprintEditorToolbar.cpp b/Engine/Source/Editor/Kismet/Private/SBlueprintEditorToolbar.cpp
--- a/Engine/Source/Editor/Kismet/Private/SBlueprintEditorToolbar.cpp
+++ b/Engine/Source/Editor/Kismet/Private/SBlueprintEditorToolbar.cpp
@@ -559,6 +559,22 @@ void FBlueprintEditorToolbar::FillScriptingToolbar(FToolBarBuilder& ToolbarBuild
 
 	ToolbarBuilder.AddToolBarButton(FBlueprintEditorCommands::Get().FindInBlueprint);
 
+	ToolbarBuilder.AddToolBarButton(
+		FBlueprintEditorCommands::Get().ToggleHideUnrelatedNodes,
+		NAME_None,
+		TAttribute<FText>(),
+		TAttribute<FText>(),
+		FSlateIcon(FEditorStyle::GetStyleSetName(), "GraphEditor.ToggleHideUnrelatedNodes")
+	);
+	ToolbarBuilder.AddComboButton(
+		FUIAction(),
+		FOnGetContent::CreateSP(BlueprintEditorPtr.Get(), &FBlueprintEditor::MakeHideUnrelatedNodesOptionsMenu),
+		LOCTEXT("HideUnrelatedNodesOptions", "Focus Related Nodes Options"),
+		LOCTEXT("HideUnrelatedNodesOptionsMenu", "Focus Related Nodes options menu"),
+		TAttribute<FSlateIcon>(),
+		true
+	);
+
 	ToolbarBuilder.EndSection();
 }
 
