Fix crash when a component is added if you have a separator selected

diff --git a/Engine/Source/Editor/Kismet/Private/SSCSEditor.cpp b/Engine/Source/Editor/Kismet/Private/SSCSEditor.cpp
--- a/Engine/Source/Editor/Kismet/Private/SSCSEditor.cpp
+++ b/Engine/Source/Editor/Kismet/Private/SSCSEditor.cpp
@@ -5430,7 +5430,7 @@ FSCSEditorTreeNodePtrType SSCSEditor::FindParentForNewComponent(UActorComponent*
 	{
 		if (TargetParentNode.IsValid())
 		{
-			while (TargetParentNode->GetNodeType() == FSCSEditorTreeNode::ComponentNode)
+			while (TargetParentNode->GetNodeType() != FSCSEditorTreeNode::RootActorNode)
 			{
 				TargetParentNode = TargetParentNode->GetParent();
 			}
@@ -5440,7 +5440,7 @@ FSCSEditorTreeNodePtrType SSCSEditor::FindParentForNewComponent(UActorComponent*
 			TargetParentNode = GetActorNode();
 		}
 
-		check(TargetParentNode->GetNodeType() == FSCSEditorTreeNode::RootActorNode);
+		check(TargetParentNode.IsValid() && TargetParentNode->GetNodeType() == FSCSEditorTreeNode::RootActorNode);
 	}
 
 	return TargetParentNode;
