Update Engine and Plugins to use the new Getters / Setters for UActorComponent, USceneComponent, and AActor replicated properties

diff --git a/Engine/Source/Editor/Kismet/Private/SSCSEditor.cpp b/Engine/Source/Editor/Kismet/Private/SSCSEditor.cpp
--- a/Engine/Source/Editor/Kismet/Private/SSCSEditor.cpp
+++ b/Engine/Source/Editor/Kismet/Private/SSCSEditor.cpp
@@ -2584,9 +2584,9 @@ void SSCS_RowWidget::OnAttachToDropAction(const TArray<FSCSEditorTreeNodePtrType
 				FRotator OldRelativeRotation;
 				if(SceneComponentTemplate)
 				{
-					OldRelativeLocation = SceneComponentTemplate->RelativeLocation;
-					OldRelativeRotation = SceneComponentTemplate->RelativeRotation;
-					OldRelativeScale3D = SceneComponentTemplate->RelativeScale3D;
+					OldRelativeLocation = SceneComponentTemplate->GetRelativeLocation();
+					OldRelativeRotation = SceneComponentTemplate->GetRelativeRotation();
+					OldRelativeScale3D = SceneComponentTemplate->GetRelativeScale3D();
 				}
 
 				// Check for a valid parent node
@@ -2617,9 +2617,9 @@ void SSCS_RowWidget::OnAttachToDropAction(const TArray<FSCSEditorTreeNodePtrType
 						{
 							// If we find a match, save off the world position
 							FTransform ComponentToWorld = InstancedSceneComponent->GetComponentToWorld();
-							SceneComponentTemplate->RelativeLocation = ComponentToWorld.GetTranslation();
-							SceneComponentTemplate->RelativeRotation = ComponentToWorld.Rotator();
-							SceneComponentTemplate->RelativeScale3D = ComponentToWorld.GetScale3D();
+							SceneComponentTemplate->SetRelativeLocation_Direct(ComponentToWorld.GetTranslation());
+							SceneComponentTemplate->SetRelativeRotation_Direct(ComponentToWorld.Rotator());
+							SceneComponentTemplate->SetRelativeScale3D_Direct(ComponentToWorld.GetScale3D());
 						}
 					}
 				}
@@ -2632,26 +2632,26 @@ void SSCS_RowWidget::OnAttachToDropAction(const TArray<FSCSEditorTreeNodePtrType
 				if(SceneComponentTemplate && ParentSceneComponent && ParentSceneComponent->IsRegistered())
 				{
 					// If we find a match, calculate its new position relative to the scene root component instance in its current scene
-					FTransform ComponentToWorld(SceneComponentTemplate->RelativeRotation, SceneComponentTemplate->RelativeLocation, SceneComponentTemplate->RelativeScale3D);
+					FTransform ComponentToWorld(SceneComponentTemplate->GetRelativeRotation(), SceneComponentTemplate->GetRelativeLocation(), SceneComponentTemplate->GetRelativeScale3D());
 					FTransform ParentToWorld = SceneComponentTemplate->GetAttachSocketName() != NAME_None ? ParentSceneComponent->GetSocketTransform(SceneComponentTemplate->GetAttachSocketName(), RTS_World) : ParentSceneComponent->GetComponentToWorld();
 					FTransform RelativeTM = ComponentToWorld.GetRelativeTransform(ParentToWorld);
 
 					// Store new relative location value (if not set to absolute)
-					if(!SceneComponentTemplate->bAbsoluteLocation)
+					if(!SceneComponentTemplate->IsUsingAbsoluteLocation())
 					{
-						SceneComponentTemplate->RelativeLocation = RelativeTM.GetTranslation();
+						SceneComponentTemplate->SetRelativeLocation_Direct(RelativeTM.GetTranslation());
 					}
 
 					// Store new relative rotation value (if not set to absolute)
-					if(!SceneComponentTemplate->bAbsoluteRotation)
+					if(!SceneComponentTemplate->IsUsingAbsoluteRotation())
 					{
-						SceneComponentTemplate->RelativeRotation = RelativeTM.Rotator();
+						SceneComponentTemplate->SetRelativeRotation_Direct(RelativeTM.Rotator());
 					}
 
 					// Store new relative scale value (if not set to absolute)
-					if(!SceneComponentTemplate->bAbsoluteScale)
+					if(!SceneComponentTemplate->IsUsingAbsoluteScale())
 					{
-						SceneComponentTemplate->RelativeScale3D = RelativeTM.GetScale3D();
+						SceneComponentTemplate->SetRelativeScale3D_Direct(RelativeTM.GetScale3D());
 					}
 				}
 
@@ -2665,9 +2665,9 @@ void SSCS_RowWidget::OnAttachToDropAction(const TArray<FSCSEditorTreeNodePtrType
 						USceneComponent* InstancedSceneComponent = Cast<USceneComponent>(InstancedSceneComponents[InstanceIndex]);
 						if(InstancedSceneComponent != nullptr)
 						{
-							FComponentEditorUtils::ApplyDefaultValueChange(InstancedSceneComponent, InstancedSceneComponent->RelativeLocation, OldRelativeLocation, SceneComponentTemplate->RelativeLocation);
-							FComponentEditorUtils::ApplyDefaultValueChange(InstancedSceneComponent, InstancedSceneComponent->RelativeRotation, OldRelativeRotation, SceneComponentTemplate->RelativeRotation);
-							FComponentEditorUtils::ApplyDefaultValueChange(InstancedSceneComponent, InstancedSceneComponent->RelativeScale3D,  OldRelativeScale3D,  SceneComponentTemplate->RelativeScale3D);
+							FComponentEditorUtils::ApplyDefaultValueChange(InstancedSceneComponent, InstancedSceneComponent->GetRelativeLocation_DirectMutable(), OldRelativeLocation, SceneComponentTemplate->GetRelativeLocation());
+							FComponentEditorUtils::ApplyDefaultValueChange(InstancedSceneComponent, InstancedSceneComponent->GetRelativeRotation_DirectMutable(), OldRelativeRotation, SceneComponentTemplate->GetRelativeRotation());
+							FComponentEditorUtils::ApplyDefaultValueChange(InstancedSceneComponent, InstancedSceneComponent->GetRelativeScale3D_DirectMutable(),  OldRelativeScale3D,  SceneComponentTemplate->GetRelativeScale3D());
 						}
 					}
 				}
@@ -2729,9 +2729,9 @@ void SSCS_RowWidget::OnDetachFromDropAction(const TArray<FSCSEditorTreeNodePtrTy
 			if(SceneComponentTemplate)
 			{
 				// Cache current default values for propagation
-				OldRelativeLocation = SceneComponentTemplate->RelativeLocation;
-				OldRelativeRotation = SceneComponentTemplate->RelativeRotation;
-				OldRelativeScale3D = SceneComponentTemplate->RelativeScale3D;
+				OldRelativeLocation = SceneComponentTemplate->GetRelativeLocation();
+				OldRelativeRotation = SceneComponentTemplate->GetRelativeRotation();
+				OldRelativeScale3D = SceneComponentTemplate->GetRelativeScale3D();
 
 				// Save current state
 				SceneComponentTemplate->Modify();
@@ -2751,9 +2751,9 @@ void SSCS_RowWidget::OnDetachFromDropAction(const TArray<FSCSEditorTreeNodePtrTy
 				{
 					// If we find a match, save off the world position
 					FTransform ComponentToWorld = InstancedSceneComponent->GetComponentToWorld();
-					SceneComponentTemplate->RelativeLocation = ComponentToWorld.GetTranslation();
-					SceneComponentTemplate->RelativeRotation = ComponentToWorld.Rotator();
-					SceneComponentTemplate->RelativeScale3D = ComponentToWorld.GetScale3D();
+					SceneComponentTemplate->SetRelativeLocation_Direct(ComponentToWorld.GetTranslation());
+					SceneComponentTemplate->SetRelativeRotation_Direct(ComponentToWorld.Rotator());
+					SceneComponentTemplate->SetRelativeScale3D_Direct(ComponentToWorld.GetScale3D());
 				}
 			}
 
@@ -2767,26 +2767,26 @@ void SSCS_RowWidget::OnDetachFromDropAction(const TArray<FSCSEditorTreeNodePtrTy
 			if(SceneComponentTemplate && InstancedSceneRootComponent && InstancedSceneRootComponent->IsRegistered())
 			{
 				// If we find a match, calculate its new position relative to the scene root component instance in the preview scene
-				FTransform ComponentToWorld(SceneComponentTemplate->RelativeRotation, SceneComponentTemplate->RelativeLocation, SceneComponentTemplate->RelativeScale3D);
+				FTransform ComponentToWorld(SceneComponentTemplate->GetRelativeRotation(), SceneComponentTemplate->GetRelativeLocation(), SceneComponentTemplate->GetRelativeScale3D());
 				FTransform ParentToWorld = SceneComponentTemplate->GetAttachSocketName() != NAME_None ? InstancedSceneRootComponent->GetSocketTransform(SceneComponentTemplate->GetAttachSocketName(), RTS_World) : InstancedSceneRootComponent->GetComponentToWorld();
 				FTransform RelativeTM = ComponentToWorld.GetRelativeTransform(ParentToWorld);
 
 				// Store new relative location value (if not set to absolute)
-				if(!SceneComponentTemplate->bAbsoluteLocation)
+				if(!SceneComponentTemplate->IsUsingAbsoluteLocation())
 				{
-					SceneComponentTemplate->RelativeLocation = RelativeTM.GetTranslation();
+					SceneComponentTemplate->SetRelativeLocation_Direct(RelativeTM.GetTranslation());
 				}
 
 				// Store new relative rotation value (if not set to absolute)
-				if(!SceneComponentTemplate->bAbsoluteRotation)
+				if(!SceneComponentTemplate->IsUsingAbsoluteRotation())
 				{
-					SceneComponentTemplate->RelativeRotation = RelativeTM.Rotator();
+					SceneComponentTemplate->SetRelativeRotation_Direct(RelativeTM.Rotator());
 				}
 
 				// Store new relative scale value (if not set to absolute)
-				if(!SceneComponentTemplate->bAbsoluteScale)
+				if(!SceneComponentTemplate->IsUsingAbsoluteScale())
 				{
-					SceneComponentTemplate->RelativeScale3D = RelativeTM.GetScale3D();
+					SceneComponentTemplate->SetRelativeScale3D_Direct(RelativeTM.GetScale3D());
 				}
 			}
 
@@ -2800,9 +2800,9 @@ void SSCS_RowWidget::OnDetachFromDropAction(const TArray<FSCSEditorTreeNodePtrTy
 					USceneComponent* InstancedSceneComponent = Cast<USceneComponent>(InstancedSceneComponents[InstanceIndex]);
 					if(InstancedSceneComponent != nullptr)
 					{
-						FComponentEditorUtils::ApplyDefaultValueChange(InstancedSceneComponent, InstancedSceneComponent->RelativeLocation, OldRelativeLocation, SceneComponentTemplate->RelativeLocation);
-						FComponentEditorUtils::ApplyDefaultValueChange(InstancedSceneComponent, InstancedSceneComponent->RelativeRotation, OldRelativeRotation, SceneComponentTemplate->RelativeRotation);
-						FComponentEditorUtils::ApplyDefaultValueChange(InstancedSceneComponent, InstancedSceneComponent->RelativeScale3D,  OldRelativeScale3D,  SceneComponentTemplate->RelativeScale3D);
+						FComponentEditorUtils::ApplyDefaultValueChange(InstancedSceneComponent, InstancedSceneComponent->GetRelativeLocation_DirectMutable(), OldRelativeLocation, SceneComponentTemplate->GetRelativeLocation());
+						FComponentEditorUtils::ApplyDefaultValueChange(InstancedSceneComponent, InstancedSceneComponent->GetRelativeRotation_DirectMutable(), OldRelativeRotation, SceneComponentTemplate->GetRelativeRotation());
+						FComponentEditorUtils::ApplyDefaultValueChange(InstancedSceneComponent, InstancedSceneComponent->GetRelativeScale3D_DirectMutable(),  OldRelativeScale3D,  SceneComponentTemplate->GetRelativeScale3D());
 					}
 				}
 			}
@@ -2898,8 +2898,8 @@ void SSCS_RowWidget::OnMakeNewRootDropAction(FSCSEditorTreeNodePtrType DroppedNo
 				}
 
 				// Cache the current relative location and rotation values (for propagation)
-				const FVector OldRelativeLocation = SceneComponentTemplate->RelativeLocation;
-				const FRotator OldRelativeRotation = SceneComponentTemplate->RelativeRotation;
+				const FVector OldRelativeLocation = SceneComponentTemplate->GetRelativeLocation();
+				const FRotator OldRelativeRotation = SceneComponentTemplate->GetRelativeRotation();
 
 				// Reset the relative transform (location and rotation only; scale is preserved)
 				SceneComponentTemplate->SetRelativeLocation(FVector::ZeroVector);
@@ -2918,8 +2918,8 @@ void SSCS_RowWidget::OnMakeNewRootDropAction(FSCSEditorTreeNodePtrType DroppedNo
 						SceneComponentInstance->DetachFromComponent(DetachmentTransformRules);
 
 						// Propagate the default relative location & rotation reset from the template to the instance
-						FComponentEditorUtils::ApplyDefaultValueChange(SceneComponentInstance, SceneComponentInstance->RelativeLocation, OldRelativeLocation, SceneComponentTemplate->RelativeLocation);
-						FComponentEditorUtils::ApplyDefaultValueChange(SceneComponentInstance, SceneComponentInstance->RelativeRotation, OldRelativeRotation, SceneComponentTemplate->RelativeRotation);
+						FComponentEditorUtils::ApplyDefaultValueChange(SceneComponentInstance, SceneComponentInstance->GetRelativeLocation_DirectMutable(), OldRelativeLocation, SceneComponentTemplate->GetRelativeLocation());
+						FComponentEditorUtils::ApplyDefaultValueChange(SceneComponentInstance, SceneComponentInstance->GetRelativeRotation_DirectMutable(), OldRelativeRotation, SceneComponentTemplate->GetRelativeRotation());
 
 						// Must also reset the root component here, so that RerunConstructionScripts() will cache the correct root component instance data
 						AActor* Owner = SceneComponentInstance->GetOwner();
@@ -4241,7 +4241,7 @@ void SSCSEditor::OnDuplicateComponent()
 				// otherwise we'll end up with the square of the root's scale instead of being the same size.
 				if (OriginalNodePtr == GetSceneRootNode())
 				{
-					NewSceneComponent->RelativeScale3D = FVector(1.f);
+					NewSceneComponent->SetRelativeScale3D_Direct(FVector(1.f));
 				}
 				else
 				{
diff --git a/Engine/Source/Editor/LandscapeEditor/Private/ActorFactoryLandscape.cpp b/Engine/Source/Editor/LandscapeEditor/Private/ActorFactoryLandscape.cpp
--- a/Engine/Source/Editor/LandscapeEditor/Private/ActorFactoryLandscape.cpp
+++ b/Engine/Source/Editor/LandscapeEditor/Private/ActorFactoryLandscape.cpp
@@ -62,10 +62,10 @@ ALandscapePlaceholder::ALandscapePlaceholder(const FObjectInitializer& ObjectIni
 	if (SpriteComponent)
 	{
 		SpriteComponent->Sprite = ConstructorStatics.TerrainTexture.Get();
-		SpriteComponent->RelativeScale3D = FVector(0.5f, 0.5f, 0.5f);
+		SpriteComponent->SetRelativeScale3D(FVector(0.5f, 0.5f, 0.5f));
 		SpriteComponent->SetupAttachment(RootComponent);
-		SpriteComponent->RelativeLocation = FVector(0, 0, 100);
-		SpriteComponent->bAbsoluteScale = true;
+		SpriteComponent->SetRelativeLocation(FVector(0, 0, 100));
+		SpriteComponent->SetUsingAbsoluteScale(true);
 	}
 #endif
 }
diff --git a/Engine/Source/Editor/LandscapeEditor/Private/LandscapeEdModeComponentTools.cpp b/Engine/Source/Editor/LandscapeEditor/Private/LandscapeEdModeComponentTools.cpp
--- a/Engine/Source/Editor/LandscapeEditor/Private/LandscapeEdModeComponentTools.cpp
+++ b/Engine/Source/Editor/LandscapeEditor/Private/LandscapeEdModeComponentTools.cpp
@@ -1520,8 +1520,9 @@ public:
 			const float W = Gizmo->GetWidth() / (2 * ScaleXY);
 			const float H = Gizmo->GetHeight() / (2 * ScaleXY);
 
-			const float SignX = Gizmo->GetRootComponent()->RelativeScale3D.X > 0.0f ? 1.0f : -1.0f;
-			const float SignY = Gizmo->GetRootComponent()->RelativeScale3D.Y > 0.0f ? 1.0f : -1.0f;
+			const FVector GizmoScale3D = Gizmo->GetRootComponent()->GetRelativeScale3D();
+			const float SignX = GizmoScale3D.X > 0.0f ? 1.0f : -1.0f;
+			const float SignY = GizmoScale3D.Y > 0.0f ? 1.0f : -1.0f;
 
 			const float ScaleX = Gizmo->CachedWidth / Width * ScaleXY / Gizmo->CachedScaleXY;
 			const float ScaleY = Gizmo->CachedHeight / Height * ScaleXY / Gizmo->CachedScaleXY;
diff --git a/Engine/Source/Editor/LandscapeEditor/Private/LandscapeEdModeSplineTools.cpp b/Engine/Source/Editor/LandscapeEditor/Private/LandscapeEdModeSplineTools.cpp
--- a/Engine/Source/Editor/LandscapeEditor/Private/LandscapeEdModeSplineTools.cpp
+++ b/Engine/Source/Editor/LandscapeEditor/Private/LandscapeEdModeSplineTools.cpp
@@ -100,7 +100,7 @@ public:
 	{
 		Landscape->Modify();
 		Landscape->SplineComponent = NewObject<ULandscapeSplinesComponent>(Landscape, NAME_None, RF_Transactional);
-		Landscape->SplineComponent->RelativeScale3D = Scale3D;
+		Landscape->SplineComponent->SetRelativeScale3D_Direct(Scale3D);
 		Landscape->SplineComponent->AttachToComponent(Landscape->GetRootComponent(), FAttachmentTransformRules::KeepRelativeTransform);
 		Landscape->SplineComponent->ShowSplineEditorMesh(true);
 	}
@@ -883,7 +883,7 @@ public:
 					ToLandscape->Modify();
 					if (ToLandscape->SplineComponent == nullptr)
 					{
-						CreateSplineComponent(ToLandscape, FromProxy->SplineComponent->RelativeScale3D);
+						CreateSplineComponent(ToLandscape, FromProxy->SplineComponent->GetRelativeScale3D());
 						check(ToLandscape->SplineComponent);
 					}
 					ToLandscape->SplineComponent->Modify();
@@ -957,7 +957,7 @@ public:
 					ToLandscape->Modify();
 					if (ToLandscape->SplineComponent == nullptr)
 					{
-						CreateSplineComponent(ToLandscape, FromProxy->SplineComponent->RelativeScale3D);
+						CreateSplineComponent(ToLandscape, FromProxy->SplineComponent->GetRelativeScale3D());
 						check(ToLandscape->SplineComponent);
 					}
 					ToLandscape->SplineComponent->Modify();
@@ -1090,7 +1090,7 @@ public:
 			{
 				if (!Landscape->SplineComponent)
 				{
-					CreateSplineComponent(Landscape, FVector(1.0f) / Landscape->GetRootComponent()->RelativeScale3D);
+					CreateSplineComponent(Landscape, FVector(1.0f) / Landscape->GetRootComponent()->GetRelativeScale3D());
 					check(Landscape->SplineComponent);
 				}
 				SplinesComponent = Landscape->SplineComponent;
@@ -2255,7 +2255,7 @@ public:
 		}
 		if (!Landscape->SplineComponent)
 		{
-			CreateSplineComponent(Landscape, FVector(1.0f) / Landscape->GetRootComponent()->RelativeScale3D);
+			CreateSplineComponent(Landscape, FVector(1.0f) / Landscape->GetRootComponent()->GetRelativeScale3D());
 			check(Landscape->SplineComponent);
 		}
 		Landscape->SplineComponent->Modify();
diff --git a/Engine/Source/Editor/MeshPaintMode/Private/PaintModePainter.cpp b/Engine/Source/Editor/MeshPaintMode/Private/PaintModePainter.cpp
--- a/Engine/Source/Editor/MeshPaintMode/Private/PaintModePainter.cpp
+++ b/Engine/Source/Editor/MeshPaintMode/Private/PaintModePainter.cpp
@@ -2008,7 +2008,7 @@ void FPaintModePainter::UpdatePaintTargets(UObject* InObject, struct FPropertyCh
 {
 	AActor* Actor = Cast<AActor>(InObject);
 	if (InPropertyChangedEvent.Property && 
-		InPropertyChangedEvent.Property->GetName() == GET_MEMBER_NAME_CHECKED(USceneComponent, bVisible).ToString())
+		InPropertyChangedEvent.Property->GetName() == USceneComponent::GetVisiblePropertyName().ToString())
 	{
 		Refresh();
 	}
diff --git a/Engine/Source/Editor/MovieSceneTools/Private/TrackEditors/CameraAnimTrackEditor.cpp b/Engine/Source/Editor/MovieSceneTools/Private/TrackEditors/CameraAnimTrackEditor.cpp
--- a/Engine/Source/Editor/MovieSceneTools/Private/TrackEditors/CameraAnimTrackEditor.cpp
+++ b/Engine/Source/Editor/MovieSceneTools/Private/TrackEditors/CameraAnimTrackEditor.cpp
@@ -279,7 +279,7 @@ UCameraComponent* FCameraAnimTrackEditor::AcquireCameraComponentFromObjectGuid(c
 		}
 		else if (UCameraComponent* const CameraComp = Cast<UCameraComponent>(Obj))
 		{
-			if (CameraComp->bIsActive)
+			if (CameraComp->IsActive())
 			{
 				return CameraComp;
 			}
diff --git a/Engine/Source/Editor/MovieSceneTools/Private/TrackEditors/CameraShakeTrackEditor.cpp b/Engine/Source/Editor/MovieSceneTools/Private/TrackEditors/CameraShakeTrackEditor.cpp
--- a/Engine/Source/Editor/MovieSceneTools/Private/TrackEditors/CameraShakeTrackEditor.cpp
+++ b/Engine/Source/Editor/MovieSceneTools/Private/TrackEditors/CameraShakeTrackEditor.cpp
@@ -286,7 +286,7 @@ UCameraComponent* FCameraShakeTrackEditor::AcquireCameraComponentFromObjectGuid(
 		}
 		else if (UCameraComponent* const CameraComp = Cast<UCameraComponent>(Obj))
 		{
-			if (CameraComp->bIsActive)
+			if (CameraComp->IsActive())
 			{
 				return CameraComp;
 			}
diff --git a/Engine/Source/Editor/MovieSceneTools/Private/TrackEditors/TransformTrackEditor.cpp b/Engine/Source/Editor/MovieSceneTools/Private/TrackEditors/TransformTrackEditor.cpp
--- a/Engine/Source/Editor/MovieSceneTools/Private/TrackEditors/TransformTrackEditor.cpp
+++ b/Engine/Source/Editor/MovieSceneTools/Private/TrackEditors/TransformTrackEditor.cpp
@@ -365,9 +365,9 @@ void F3DTransformTrackEditor::OnPrePropertyChanged(UObject* InObject, const FEdi
 	UProperty* PropertyAboutToChange = InPropertyChain.GetActiveMemberNode()->GetValue();
 	const FName MemberPropertyName = PropertyAboutToChange != nullptr ? PropertyAboutToChange->GetFName() : NAME_None;
 	const bool bTransformationToChange =
-		(MemberPropertyName == GET_MEMBER_NAME_CHECKED(USceneComponent, RelativeLocation) ||
-		 MemberPropertyName == GET_MEMBER_NAME_CHECKED(USceneComponent, RelativeRotation) ||
-		 MemberPropertyName == GET_MEMBER_NAME_CHECKED(USceneComponent, RelativeScale3D));
+		(MemberPropertyName == USceneComponent::GetRelativeLocationPropertyName() ||
+		 MemberPropertyName == USceneComponent::GetRelativeRotationPropertyName() ||
+		 MemberPropertyName == USceneComponent::GetRelativeScale3DPropertyName());
 
 	if (InObject && bTransformationToChange)
 	{
@@ -379,9 +379,9 @@ void F3DTransformTrackEditor::OnPostPropertyChanged(UObject* InObject, FProperty
 {
 	const FName MemberPropertyName = InPropertyChangedEvent.MemberProperty != nullptr ? InPropertyChangedEvent.MemberProperty->GetFName() : NAME_None;
 	const bool bTransformationChanged =
-		(MemberPropertyName == GET_MEMBER_NAME_CHECKED(USceneComponent, RelativeLocation) ||
-			MemberPropertyName == GET_MEMBER_NAME_CHECKED(USceneComponent, RelativeRotation) ||
-			MemberPropertyName == GET_MEMBER_NAME_CHECKED(USceneComponent, RelativeScale3D));
+		(MemberPropertyName == USceneComponent::GetRelativeLocationPropertyName() ||
+			MemberPropertyName == USceneComponent::GetRelativeRotationPropertyName() ||
+			MemberPropertyName == USceneComponent::GetRelativeScale3DPropertyName());
 
 	if (InObject && bTransformationChanged)
 	{
diff --git a/Engine/Source/Editor/SequenceRecorderSections/Private/MovieSceneVisibilitySectionRecorder.cpp b/Engine/Source/Editor/SequenceRecorderSections/Private/MovieSceneVisibilitySectionRecorder.cpp
--- a/Engine/Source/Editor/SequenceRecorderSections/Private/MovieSceneVisibilitySectionRecorder.cpp
+++ b/Engine/Source/Editor/SequenceRecorderSections/Private/MovieSceneVisibilitySectionRecorder.cpp
@@ -76,7 +76,7 @@ void FMovieSceneVisibilitySectionRecorder::CreateSection(UObject* InObjectToReco
 		}
 		else if(AActor* Actor = Cast<AActor>(ObjectToRecord.Get()))
 		{
-			bWasVisible = !Actor->bHidden;
+			bWasVisible = !Actor->IsHidden();
 		}
 
 		FFrameRate   TickResolution = MovieSceneSection->GetTypedOuter<UMovieScene>()->GetTickResolution();
@@ -127,7 +127,7 @@ void FMovieSceneVisibilitySectionRecorder::Record(float CurrentTime)
 		}
 		else if(AActor* Actor = Cast<AActor>(ObjectToRecord.Get()))
 		{
-			bVisible = !Actor->bHidden;
+			bVisible = !Actor->IsHidden();
 		}
 
 		if(bVisible != bWasVisible)
diff --git a/Engine/Source/Editor/StatsViewer/Private/StatsPages/StaticMeshLightingInfoStatsPage.cpp b/Engine/Source/Editor/StatsViewer/Private/StatsPages/StaticMeshLightingInfoStatsPage.cpp
--- a/Engine/Source/Editor/StatsViewer/Private/StatsPages/StaticMeshLightingInfoStatsPage.cpp
+++ b/Engine/Source/Editor/StatsViewer/Private/StatsPages/StaticMeshLightingInfoStatsPage.cpp
@@ -141,7 +141,7 @@ struct StaticMeshLightingInfoStatsGenerator
 			{
 				ULightComponent* Light = AllLights[LightIndex];
 				// Only add enabled lights
-				if (Light->bVisible && Light->AffectsPrimitive(InComponent))
+				if (Light->GetVisibleFlag() && Light->AffectsPrimitive(InComponent))
 				{
 					// Check whether the light should use a light-map or shadow-map.
 					const bool bHasStaticLighting = Light->HasStaticLighting();
diff --git a/Engine/Source/Editor/UnrealEd/Private/Editor.cpp b/Engine/Source/Editor/UnrealEd/Private/Editor.cpp
--- a/Engine/Source/Editor/UnrealEd/Private/Editor.cpp
+++ b/Engine/Source/Editor/UnrealEd/Private/Editor.cpp
@@ -1385,9 +1385,9 @@ namespace EditorUtilities
 					const bool bIsIdentical = Property->Identical_InContainer( SourceComponent, TargetComponent );
 					const bool bIsComponent = !!( Property->PropertyFlags & ( CPF_InstancedReference | CPF_ContainsInstancedReference ) );
 					const bool bIsTransform =
-						Property->GetFName() == GET_MEMBER_NAME_CHECKED(USceneComponent, RelativeScale3D) ||
-						Property->GetFName() == GET_MEMBER_NAME_CHECKED(USceneComponent, RelativeLocation) ||
-						Property->GetFName() == GET_MEMBER_NAME_CHECKED(USceneComponent, RelativeRotation);
+						Property->GetFName() == USceneComponent::GetRelativeScale3DPropertyName() ||
+						Property->GetFName() == USceneComponent::GetRelativeLocationPropertyName() ||
+						Property->GetFName() == USceneComponent::GetRelativeRotationPropertyName();
 
 					if( !bIsTransient && !bIsIdentical && !bIsComponent && !SourceUCSModifiedProperties.Contains(Property)
 						&& ( !bIsTransform || SourceComponent != SourceActor->GetRootComponent() || ( !SourceActor->HasAnyFlags( RF_ClassDefaultObject | RF_ArchetypeObject ) && !TargetActor->HasAnyFlags( RF_ClassDefaultObject | RF_ArchetypeObject ) ) ) )
diff --git a/Engine/Source/Editor/UnrealEd/Private/EditorBrushBuilder.cpp b/Engine/Source/Editor/UnrealEd/Private/EditorBrushBuilder.cpp
--- a/Engine/Source/Editor/UnrealEd/Private/EditorBrushBuilder.cpp
+++ b/Engine/Source/Editor/UnrealEd/Private/EditorBrushBuilder.cpp
@@ -53,7 +53,7 @@ bool UEditorBrushBuilder::EndBrush( UWorld* InWorld, ABrush* InBrush )
 	ABrush* BuilderBrush = (InBrush != nullptr) ? InBrush : InWorld->GetDefaultBrush();
 
 	// Ensure the builder brush is unhidden.
-	BuilderBrush->bHidden = false;
+	BuilderBrush->SetHidden(false);
 	BuilderBrush->bHiddenEdLayer = false;
 
 	AActor* Actor = GEditor->GetSelectedActors()->GetTop<AActor>();
diff --git a/Engine/Source/Editor/UnrealEd/Private/EditorServer.cpp b/Engine/Source/Editor/UnrealEd/Private/EditorServer.cpp
--- a/Engine/Source/Editor/UnrealEd/Private/EditorServer.cpp
+++ b/Engine/Source/Editor/UnrealEd/Private/EditorServer.cpp
@@ -4058,7 +4058,7 @@ bool UEditorEngine::Map_Scale( UWorld* InWorld, const TCHAR* Str, FOutputDevice&
 			}
 			else if( Actor->GetRootComponent() != NULL )
 			{
-				Actor->GetRootComponent()->SetRelativeScale3D( Actor->GetRootComponent()->RelativeScale3D * Factor );
+				Actor->GetRootComponent()->SetRelativeScale3D(Actor->GetRootComponent()->GetRelativeScale3D() * Factor);
 			}
 
 			if( bScaleLocations )
diff --git a/Engine/Source/Editor/UnrealEd/Private/Factories/ActorFactory.cpp b/Engine/Source/Editor/UnrealEd/Private/Factories/ActorFactory.cpp
--- a/Engine/Source/Editor/UnrealEd/Private/Factories/ActorFactory.cpp
+++ b/Engine/Source/Editor/UnrealEd/Private/Factories/ActorFactory.cpp
@@ -769,7 +769,7 @@ void UActorFactoryPhysicsAsset::PostSpawnActor(UObject* Asset, AActor* NewActor)
 	NewSkelActor->GetSkeletalMeshComponent()->bBlendPhysics = true;
 
 	NewSkelActor->bAlwaysRelevant = true;
-	NewSkelActor->bReplicateMovement = true;
+	NewSkelActor->SetReplicatingMovement(true);
 	NewSkelActor->SetReplicates(true);
 
 	// Init Component
@@ -798,7 +798,7 @@ void UActorFactoryPhysicsAsset::PostCreateBlueprint( UObject* Asset, AActor* CDO
 		SkeletalPhysicsActor->GetSkeletalMeshComponent()->bBlendPhysics = true;
 
 		SkeletalPhysicsActor->bAlwaysRelevant = true;
-		SkeletalPhysicsActor->bReplicateMovement = true;
+		SkeletalPhysicsActor->SetReplicatingMovement(true);
 		SkeletalPhysicsActor->SetReplicates(true);
 	}
 }
diff --git a/Engine/Source/Editor/UnrealEd/Private/Fbx/FbxAnimationExport.cpp b/Engine/Source/Editor/UnrealEd/Private/Fbx/FbxAnimationExport.cpp
--- a/Engine/Source/Editor/UnrealEd/Private/Fbx/FbxAnimationExport.cpp
+++ b/Engine/Source/Editor/UnrealEd/Private/Fbx/FbxAnimationExport.cpp
@@ -461,7 +461,7 @@ void FFbxExporter::ExportAnimSequencesAsSingle( USkeletalMesh* SkelMesh, const A
 		// The Unreal transformation is different from FBX's Z-up: invert the Y-axis for translations and the Y/Z angle values in rotations.
 		BaseNode->LclTranslation.Set(Converter.ConvertToFbxPos(SkelMeshActor->GetActorLocation()));
 		BaseNode->LclRotation.Set(Converter.ConvertToFbxRot(SkelMeshActor->GetActorRotation().Euler()));
-		BaseNode->LclScaling.Set(Converter.ConvertToFbxScale(SkelMeshActor->GetRootComponent()->RelativeScale3D));
+		BaseNode->LclScaling.Set(Converter.ConvertToFbxScale(SkelMeshActor->GetRootComponent()->GetRelativeScale3D()));
 
 	}
 
@@ -526,7 +526,7 @@ void FFbxExporter::ExportMatineeGroup(class AMatineeActor* MatineeActor, USkelet
 		// The UE3 transformation is different from FBX's Z-up: invert the Y-axis for translations and the Y/Z angle values in rotations.
 		BaseNode->LclTranslation.Set(Converter.ConvertToFbxPos(Owner->GetActorLocation()));
 		BaseNode->LclRotation.Set(Converter.ConvertToFbxRot(Owner->GetActorRotation().Euler()));
-		BaseNode->LclScaling.Set(Converter.ConvertToFbxScale(Owner->GetRootComponent()->RelativeScale3D));
+		BaseNode->LclScaling.Set(Converter.ConvertToFbxScale(Owner->GetRootComponent()->GetRelativeScale3D()));
 	}
 	// Create the Skeleton
 	TArray<FbxNode*> BoneNodes;
diff --git a/Engine/Source/Editor/UnrealEd/Private/Fbx/FbxMainExport.cpp b/Engine/Source/Editor/UnrealEd/Private/Fbx/FbxMainExport.cpp
--- a/Engine/Source/Editor/UnrealEd/Private/Fbx/FbxMainExport.cpp
+++ b/Engine/Source/Editor/UnrealEd/Private/Fbx/FbxMainExport.cpp
@@ -4307,7 +4307,7 @@ void FFbxExporter::ExportLandscapeToFbx(ALandscapeProxy* Landscape, const TCHAR*
 			int32 VertX, VertY;
 			CDI.VertexIndexToXY(VertIndex, VertX, VertY);
 
-			FVector Position = CDI.GetLocalVertex(VertX, VertY) + Component->RelativeLocation;
+			FVector Position = CDI.GetLocalVertex(VertX, VertY) + Component->GetRelativeLocation();
 			FbxVector4 FbxPosition = FbxVector4(Position.X, -Position.Y, Position.Z);
 			ControlPoints[BaseVertIndex + VertIndex] = FbxPosition;
 
diff --git a/Engine/Source/Editor/UnrealEd/Private/HierarchicalLOD.cpp b/Engine/Source/Editor/UnrealEd/Private/HierarchicalLOD.cpp
--- a/Engine/Source/Editor/UnrealEd/Private/HierarchicalLOD.cpp
+++ b/Engine/Source/Editor/UnrealEd/Private/HierarchicalLOD.cpp
@@ -537,7 +537,7 @@ bool FHierarchicalLODBuilder::ShouldGenerateCluster(AActor* Actor, const bool bP
 		return false;
 	}
 
-	if (Actor->bHidden)
+	if (Actor->IsHidden())
 	{
 		return false;
 	}
diff --git a/Engine/Source/Editor/UnrealEd/Private/Kismet2/ComponentEditorUtils.cpp b/Engine/Source/Editor/UnrealEd/Private/Kismet2/ComponentEditorUtils.cpp
--- a/Engine/Source/Editor/UnrealEd/Private/Kismet2/ComponentEditorUtils.cpp
+++ b/Engine/Source/Editor/UnrealEd/Private/Kismet2/ComponentEditorUtils.cpp
@@ -723,7 +723,7 @@ void FComponentEditorUtils::AdjustComponentDelta(USceneComponent* Component, FVe
 	{
 		const FTransform ParentToWorldSpace = ParentSceneComp->GetSocketTransform(Component->GetAttachSocketName());
 
-		if (!Component->bAbsoluteLocation)
+		if (!Component->IsUsingAbsoluteLocation())
 		{
 			//transform the drag vector in relative to the parent transform
 			Drag = ParentToWorldSpace.InverseTransformVectorNoScale(Drag);
@@ -731,7 +731,7 @@ void FComponentEditorUtils::AdjustComponentDelta(USceneComponent* Component, FVe
 			Drag = Drag * ParentToWorldSpace.Inverse().GetScale3D();
 		}
 
-		if (!Component->bAbsoluteRotation)
+		if (!Component->IsUsingAbsoluteRotation())
 		{
 			Rotation = ( ParentToWorldSpace.Inverse().GetRotation() * Rotation.Quaternion() * ParentToWorldSpace.GetRotation() ).Rotator();
 		}
diff --git a/Engine/Source/Editor/UnrealEd/Private/Kismet2/Kismet2.cpp b/Engine/Source/Editor/UnrealEd/Private/Kismet2/Kismet2.cpp
--- a/Engine/Source/Editor/UnrealEd/Private/Kismet2/Kismet2.cpp
+++ b/Engine/Source/Editor/UnrealEd/Private/Kismet2/Kismet2.cpp
@@ -1536,17 +1536,17 @@ void FKismetEditorUtilities::AddComponentsToBlueprint(UBlueprint* Blueprint, TAr
 
 				const FTransform ComponentToWorld = SceneComponent->GetComponentTransform();
 				const FTransform RelativeTransform = ComponentToWorld.GetRelativeTransform(FirstAttachParent->GetComponentTransform());
-				if (!SceneComponent->bAbsoluteLocation)
+				if (!SceneComponent->IsUsingAbsoluteLocation())
 				{
-					SceneComponentTemplate->RelativeLocation = RelativeTransform.GetLocation();
+					SceneComponentTemplate->SetRelativeLocation_Direct(RelativeTransform.GetLocation());
 				}
-				if (!SceneComponent->bAbsoluteRotation)
+				if (!SceneComponent->IsUsingAbsoluteRotation())
 				{
-					SceneComponentTemplate->RelativeRotation = RelativeTransform.GetRotation().Rotator();
+					SceneComponentTemplate->SetRelativeRotation_Direct(RelativeTransform.GetRotation().Rotator());
 				}
-				if (!SceneComponent->bAbsoluteScale)
+				if (!SceneComponent->IsUsingAbsoluteScale())
 				{
-					SceneComponentTemplate->RelativeScale3D = RelativeTransform.GetScale3D();
+					SceneComponentTemplate->SetRelativeScale3D_Direct(RelativeTransform.GetScale3D());
 				}
 			}
 		}
@@ -1629,8 +1629,8 @@ struct FResetSceneComponentAfterCopy
 private:
 	static void Reset(USceneComponent* Component)
 	{
-		Component->RelativeLocation = FVector::ZeroVector;
-		Component->RelativeRotation = FRotator::ZeroRotator;
+		Component->SetRelativeLocation_Direct(FVector::ZeroVector);
+		Component->SetRelativeRotation_Direct(FRotator::ZeroRotator);
 
 		// Clear out the attachment info after having copied the properties from the source actor
 		Component->SetupAttachment(nullptr);
@@ -1682,7 +1682,7 @@ UBlueprint* FKismetEditorUtilities::CreateBlueprintFromActor(const FName Bluepri
 					// Copy relative scale from source to target.
 					if (USceneComponent* SrcSceneRoot = Actor->GetRootComponent())
 					{
-						DstSceneRoot->RelativeScale3D = SrcSceneRoot->RelativeScale3D;
+						DstSceneRoot->SetRelativeScale3D_Direct(SrcSceneRoot->GetRelativeScale3D());
 					}
 				}
 			}
diff --git a/Engine/Source/Editor/UnrealEd/Private/LevelEditorViewport.cpp b/Engine/Source/Editor/UnrealEd/Private/LevelEditorViewport.cpp
--- a/Engine/Source/Editor/UnrealEd/Private/LevelEditorViewport.cpp
+++ b/Engine/Source/Editor/UnrealEd/Private/LevelEditorViewport.cpp
@@ -211,19 +211,19 @@ namespace LevelEditorViewportClientHelper
 		switch (WidgetMode)
 		{
 		case FWidget::WM_Translate:
-			ValueProperty = FindField<UProperty>(USceneComponent::StaticClass(), GET_MEMBER_NAME_CHECKED(USceneComponent, RelativeLocation));
+			ValueProperty = FindField<UProperty>(USceneComponent::StaticClass(), USceneComponent::GetRelativeLocationPropertyName());
 			break;
 		case FWidget::WM_Rotate:
-			ValueProperty = FindField<UProperty>(USceneComponent::StaticClass(), GET_MEMBER_NAME_CHECKED(USceneComponent, RelativeRotation));
+			ValueProperty = FindField<UProperty>(USceneComponent::StaticClass(), USceneComponent::GetRelativeRotationPropertyName());
 			break;
 		case FWidget::WM_Scale:
-			ValueProperty = FindField<UProperty>(USceneComponent::StaticClass(), GET_MEMBER_NAME_CHECKED(USceneComponent, RelativeScale3D));
+			ValueProperty = FindField<UProperty>(USceneComponent::StaticClass(), USceneComponent::GetRelativeScale3DPropertyName());
 			break;
 		case FWidget::WM_TranslateRotateZ:
-			ValueProperty = FindField<UProperty>(USceneComponent::StaticClass(), GET_MEMBER_NAME_CHECKED(USceneComponent, RelativeLocation));
+			ValueProperty = FindField<UProperty>(USceneComponent::StaticClass(), USceneComponent::GetRelativeLocationPropertyName());
 			break;
 		case FWidget::WM_2D:
-			ValueProperty = FindField<UProperty>(USceneComponent::StaticClass(), GET_MEMBER_NAME_CHECKED(USceneComponent, RelativeLocation));
+			ValueProperty = FindField<UProperty>(USceneComponent::StaticClass(), USceneComponent::GetRelativeLocationPropertyName());
 			break;
 		default:
 			break;
@@ -2327,8 +2327,8 @@ void FLevelEditorViewportClient::UpdateViewForLockedActor(float DeltaTime)
 			{
 				// No attachment, so just use the relative location, so that we don't need to
 				// convert from a quaternion, which loses winding information.
-				SetViewLocation(Actor->GetRootComponent()->RelativeLocation);
-				SetViewRotation(Actor->GetRootComponent()->RelativeRotation);
+				SetViewLocation(Actor->GetRootComponent()->GetRelativeLocation());
+				SetViewRotation(Actor->GetRootComponent()->GetRelativeRotation());
 			}
 
 			if (bLockedCameraView)
@@ -3555,7 +3555,7 @@ UActorComponent* FLevelEditorViewportClient::FindViewComponentForActor(AActor co
 		for (UActorComponent* Comp : Comps)
 		{
 			FMinimalViewInfo DummyViewInfo;
-			if (Comp && Comp->bIsActive && Comp->GetEditorPreviewInfo(/*DeltaTime =*/0.0f, DummyViewInfo))
+			if (Comp && Comp->IsActive() && Comp->GetEditorPreviewInfo(/*DeltaTime =*/0.0f, DummyViewInfo))
 			{
 				if (Comp->IsSelected())
 				{
@@ -4030,7 +4030,7 @@ void FLevelEditorViewportClient::ModifyScale( AActor* InActor, FVector& ScaleDel
 {
 	if( InActor->GetRootComponent() )
 	{
-		const FVector CurrentScale = InActor->GetRootComponent()->RelativeScale3D;
+		const FVector CurrentScale = InActor->GetRootComponent()->GetRelativeScale3D();
 
 		const FBox LocalBox = InActor->GetComponentsBoundingBox( true );
 		const FVector ScaledExtents = LocalBox.GetExtent() * CurrentScale;
@@ -4061,8 +4061,8 @@ void FLevelEditorViewportClient::ModifyScale( USceneComponent* InComponent, FVec
 	}
 	check(PreDragTransform);
 	const FBox LocalBox = Actor->GetComponentsBoundingBox(true);
-	const FVector ScaledExtents = LocalBox.GetExtent() * InComponent->RelativeScale3D;
-	ValidateScale(PreDragTransform->GetScale3D(), InComponent->RelativeScale3D, ScaledExtents, ScaleDelta);
+	const FVector ScaledExtents = LocalBox.GetExtent() * InComponent->GetRelativeScale3D();
+	ValidateScale(PreDragTransform->GetScale3D(), InComponent->GetRelativeScale3D(), ScaledExtents, ScaleDelta);
 
 	if( ScaleDelta.IsNearlyZero() )
 	{
@@ -4140,7 +4140,7 @@ void FLevelEditorViewportClient::MouseMove(FViewport* InViewport, int32 x, int32
 			{
 				UDirectionalLightComponent* SunLight = *ComponentIt;
 
-				if (!SunLight->IsUsedAsAtmosphereSunLight() || SunLight->GetAtmosphereSunLightIndex()!= DesiredLightIndex || !SunLight->bVisible)
+				if (!SunLight->IsUsedAsAtmosphereSunLight() || SunLight->GetAtmosphereSunLightIndex()!= DesiredLightIndex || !SunLight->GetVisibleFlag())
 					continue;
 
 				float LightLuminance = SunLight->GetColoredLightBrightness().ComputeLuminance();
diff --git a/Engine/Source/Editor/UnrealEd/Private/PlayLevel.cpp b/Engine/Source/Editor/UnrealEd/Private/PlayLevel.cpp
--- a/Engine/Source/Editor/UnrealEd/Private/PlayLevel.cpp
+++ b/Engine/Source/Editor/UnrealEd/Private/PlayLevel.cpp
@@ -3241,7 +3241,7 @@ UGameInstance* UEditorEngine::CreatePIEGameInstance(int32 InPIEInstance, bool bI
 			ActorsThatWereSelected.Add( Actor );
 
 			AActor* SimActor = EditorUtilities::GetSimWorldCounterpartActor(Actor);
-			if (SimActor && !SimActor->bHidden && bInSimulateInEditor)
+			if (SimActor && !SimActor->IsHidden() && bInSimulateInEditor)
 			{
 				SelectActor( SimActor, true, false );
 			}
@@ -3802,7 +3802,7 @@ void UEditorEngine::ToggleBetweenPIEandSIE( bool bNewSession )
 		if (Actor.IsValid())
 		{
 			AActor* SimActor = EditorUtilities::GetSimWorldCounterpartActor(Actor.Get());
-			if (SimActor && !SimActor->bHidden)
+			if (SimActor && !SimActor->IsHidden())
 			{
 				SelectActor( SimActor, bIsSimulatingInEditor, false );
 			}
diff --git a/Engine/Source/Editor/VREditor/UI/VREditorFloatingUI.cpp b/Engine/Source/Editor/VREditor/UI/VREditorFloatingUI.cpp
--- a/Engine/Source/Editor/VREditor/UI/VREditorFloatingUI.cpp
+++ b/Engine/Source/Editor/VREditor/UI/VREditorFloatingUI.cpp
@@ -224,7 +224,7 @@ void AVREditorFloatingUI::CleanupWidgetReferences()
 
 void AVREditorFloatingUI::SetTransform(const FTransform& Transform)
 {
-	if (!bHidden)
+	if (!IsHidden())
 	{
 		const FVector AnimatedScale = CalculateAnimatedScale();
 		FTransform AnimatedTransform = Transform;
@@ -266,9 +266,9 @@ void AVREditorFloatingUI::UpdateFadingState(const float DeltaTime)
 		if (FadeAlpha > 0.0f + KINDA_SMALL_NUMBER)
 		{
 			// At least a little bit visible
-			if (bHidden)
+			if (IsHidden())
 			{
-				bHidden = false;
+				SetHidden(false);
 
 				// Iterate as floating UI children may have other mesh components
 				TInlineComponentArray<USceneComponent*> ComponentArray;
@@ -288,9 +288,9 @@ void AVREditorFloatingUI::UpdateFadingState(const float DeltaTime)
 		else if (FadeAlpha <= 0.0f + KINDA_SMALL_NUMBER)
 		{
 			// Fully invisible
-			if (!bHidden)
+			if (!IsHidden())
 			{
-				bHidden = true;
+				SetHidden(true);
 				// Iterate as floating UI children may have other mesh components
 				TInlineComponentArray<USceneComponent*> ComponentArray;
 				GetComponents(ComponentArray);
@@ -364,7 +364,7 @@ void AVREditorFloatingUI::ShowUI( const bool bShow, const bool bAllowFading, con
 
 		if (!bAllowFading)
 		{
-			bHidden = !bShow;
+			SetHidden(!bShow);
 			// Iterate as floating UI children may have other components
 			TInlineComponentArray<USceneComponent*> ComponentArray;
 			GetComponents(ComponentArray);
diff --git a/Engine/Source/Editor/VREditor/UI/VREditorRadialFloatingUI.cpp b/Engine/Source/Editor/VREditor/UI/VREditorRadialFloatingUI.cpp
--- a/Engine/Source/Editor/VREditor/UI/VREditorRadialFloatingUI.cpp
+++ b/Engine/Source/Editor/VREditor/UI/VREditorRadialFloatingUI.cpp
@@ -282,13 +282,13 @@ void AVREditorRadialFloatingUI::UpdateFadingState( const float DeltaTime )
 		}
 		ArrowMaterial->SetScalarParameterValue("Alpha", ArrowAlpha);
 	}
-	if ( FadeDelay > 0.f )
+	if (FadeDelay > 0.f)
 	{
 		FadeDelay -= DeltaTime;
 	}
 	else
 	{
-		if( bShouldBeVisible.GetValue() )
+		if (bShouldBeVisible.GetValue())
 		{
 			FadeAlpha += VREd::RadialUIFadeSpeed->GetFloat() * DeltaTime;
 		}
@@ -296,14 +296,14 @@ void AVREditorRadialFloatingUI::UpdateFadingState( const float DeltaTime )
 		{
 			FadeAlpha -= VREd::RadialUIFadeSpeed->GetFloat() * DeltaTime;
 		}
-		FadeAlpha = FMath::Clamp( FadeAlpha, 0.0f, 1.0f );
+		FadeAlpha = FMath::Clamp(FadeAlpha, 0.0f, 1.0f);
 
-		if( FadeAlpha > 0.0f + KINDA_SMALL_NUMBER )
+		if (FadeAlpha > 0.0f + KINDA_SMALL_NUMBER)
 		{
 			// At least a little bit visible
-			if( bHidden )
+			if(IsHidden())
 			{
-				bHidden =  false ;
+				SetHidden(false);
 				TInlineComponentArray<USceneComponent*> ComponentArray;
 				GetComponents(ComponentArray);
 				for (USceneComponent* Component : ComponentArray)
@@ -314,16 +314,16 @@ void AVREditorRadialFloatingUI::UpdateFadingState( const float DeltaTime )
 			}
 		}
 
-		if( FadeAlpha >= 1.0f - KINDA_SMALL_NUMBER )
+		if (FadeAlpha >= 1.0f - KINDA_SMALL_NUMBER)
 		{
 			// Fully visible
 		}
-		else if( FadeAlpha <= 0.0f + KINDA_SMALL_NUMBER )
+		else if (FadeAlpha <= 0.0f + KINDA_SMALL_NUMBER)
 		{
 			// Fully invisible
-			if( !bHidden )
+			if (!IsHidden())
 			{
-				bHidden = true;
+				SetHidden(true);
 				TInlineComponentArray<USceneComponent*> ComponentArray;
 				GetComponents(ComponentArray);
 				for (USceneComponent* Component : ComponentArray)
@@ -389,7 +389,7 @@ void AVREditorRadialFloatingUI::ShowUI( const bool bShow, const bool bAllowFadin
 
 		if( !bAllowFading )
 		{
-			bHidden = !bShow;
+			SetHidden(!bShow);
 			TInlineComponentArray<USceneComponent*> ComponentArray;
 			GetComponents(ComponentArray);
 			for (USceneComponent* Component : ComponentArray)
diff --git a/Engine/Source/Editor/VREditor/UI/VREditorUISystem.cpp b/Engine/Source/Editor/VREditor/UI/VREditorUISystem.cpp
--- a/Engine/Source/Editor/VREditor/UI/VREditorUISystem.cpp
+++ b/Engine/Source/Editor/VREditor/UI/VREditorUISystem.cpp
@@ -524,7 +524,7 @@ void UVREditorUISystem::OnPreviewInputAction(FEditorViewportClient& ViewportClie
 									// If we are clicking on an editable text field and the radial menu is not a numpad, show the numpad
 									if (WidgetPathUnderFinger.Widgets.Last().Widget->GetTypeAsString() == TEXT("SEditableText") && (bRadialMenuIsNumpad == false))
 									{
-										if (!QuickRadialMenu->bHidden)
+										if (!QuickRadialMenu->IsHidden())
 										{
 											bRadialMenuVisibleAtSwap = true;
 										}
@@ -1331,7 +1331,7 @@ bool UVREditorUISystem::IsShowingRadialMenu(const UVREditorInteractor* Interacto
 	{
 		const EControllerHand DockedToHand = QuickRadialMenu->GetDockedTo() == AVREditorFloatingUI::EDockedTo::LeftArm ? EControllerHand::Left : EControllerHand::Right;
 		const UVREditorInteractor* DockedToInteractor = GetOwner().GetHandInteractor(DockedToHand);
-		return DockedToInteractor == Interactor && !QuickRadialMenu->bHidden;
+		return DockedToInteractor == Interactor && !QuickRadialMenu->IsHidden();
 	}
 	else
 	{
diff --git a/Engine/Source/Editor/ViewportInteraction/ActorViewportTransformable.cpp b/Engine/Source/Editor/ViewportInteraction/ActorViewportTransformable.cpp
--- a/Engine/Source/Editor/ViewportInteraction/ActorViewportTransformable.cpp
+++ b/Engine/Source/Editor/ViewportInteraction/ActorViewportTransformable.cpp
@@ -31,16 +31,16 @@ void FActorViewportTransformable::ApplyTransform( const FTransform& NewTransform
 				ExistingTransform.GetScale3D() == NewTransform.GetScale3D();
 			GEditor->BroadcastBeginObjectMovement(*Actor);
 			
-			UProperty* TransformProperty = FindField<UProperty>(USceneComponent::StaticClass(), GET_MEMBER_NAME_CHECKED(USceneComponent, RelativeLocation));
+			UProperty* TransformProperty = FindField<UProperty>(USceneComponent::StaticClass(), USceneComponent::GetRelativeLocationPropertyName());
 			if (!bOnlyTranslationChanged)
 			{
 				if (ExistingTransform.GetRotation() != NewTransform.GetRotation())
 				{
-					TransformProperty = FindField<UProperty>(USceneComponent::StaticClass(), GET_MEMBER_NAME_CHECKED(USceneComponent, RelativeRotation));
+					TransformProperty = FindField<UProperty>(USceneComponent::StaticClass(), USceneComponent::GetRelativeRotationPropertyName());
 				}
 				else if (ExistingTransform.GetScale3D() != NewTransform.GetScale3D())
 				{
-					TransformProperty = FindField<UProperty>(USceneComponent::StaticClass(), GET_MEMBER_NAME_CHECKED(USceneComponent, RelativeScale3D));
+					TransformProperty = FindField<UProperty>(USceneComponent::StaticClass(), USceneComponent::GetRelativeScale3DPropertyName());
 				}
 			}
 			
diff --git a/Engine/Source/Editor/WorldBrowser/Private/Tiles/WorldTileCollectionModel.cpp b/Engine/Source/Editor/WorldBrowser/Private/Tiles/WorldTileCollectionModel.cpp
--- a/Engine/Source/Editor/WorldBrowser/Private/Tiles/WorldTileCollectionModel.cpp
+++ b/Engine/Source/Editor/WorldBrowser/Private/Tiles/WorldTileCollectionModel.cpp
@@ -1314,7 +1314,7 @@ FVector2D FWorldTileCollectionModel::SnapTranslationDeltaLandscape(const TShared
 																	float SnappingDistance)
 {
 	ALandscapeProxy* Landscape = LandscapeTile->GetLandscape();
-	FVector ComponentScale = Landscape->GetRootComponent()->RelativeScale3D*Landscape->ComponentSizeQuads;
+	FVector ComponentScale = Landscape->GetRootComponent()->GetRelativeScale3D()*Landscape->ComponentSizeQuads;
 	
 	return FVector2D(	FMath::GridSnap(InAbsoluteDelta.X, ComponentScale.X),
 						FMath::GridSnap(InAbsoluteDelta.Y, ComponentScale.Y));
diff --git a/Engine/Source/Editor/WorldBrowser/Private/Tiles/WorldTileModel.cpp b/Engine/Source/Editor/WorldBrowser/Private/Tiles/WorldTileModel.cpp
--- a/Engine/Source/Editor/WorldBrowser/Private/Tiles/WorldTileModel.cpp
+++ b/Engine/Source/Editor/WorldBrowser/Private/Tiles/WorldTileModel.cpp
@@ -517,7 +517,7 @@ void FWorldTileModel::UpdateLandscapeSectionsOffset(FIntPoint LevelOffset)
 	if (LandscapeProxy)
 	{
 		// Calculate new section coordinates for landscape
-		FVector	DrawScale = LandscapeProxy->GetRootComponent()->RelativeScale3D;
+		FVector DrawScale = LandscapeProxy->GetRootComponent()->GetRelativeScale3D();
 		FIntPoint QuadsSpaceOffset;
 		QuadsSpaceOffset.X = FMath::RoundToInt(LevelOffset.X / DrawScale.X);
 		QuadsSpaceOffset.Y = FMath::RoundToInt(LevelOffset.Y / DrawScale.Y);
diff --git a/Engine/Source/Runtime/AIModule/Private/EnvironmentQuery/EQSTestingPawn.cpp b/Engine/Source/Runtime/AIModule/Private/EnvironmentQuery/EQSTestingPawn.cpp
--- a/Engine/Source/Runtime/AIModule/Private/EnvironmentQuery/EQSTestingPawn.cpp
+++ b/Engine/Source/Runtime/AIModule/Private/EnvironmentQuery/EQSTestingPawn.cpp
@@ -74,7 +74,7 @@ AEQSTestingPawn::AEQSTestingPawn(const FObjectInitializer& ObjectInitializer)
 			static FConstructorStatics ConstructorStatics;
 
 			SpriteComponent->Sprite = ConstructorStatics.TextureObject.Get();
-			SpriteComponent->RelativeScale3D = FVector(1, 1, 1);
+			SpriteComponent->SetRelativeScale3D(FVector(1, 1, 1));
 			SpriteComponent->bHiddenInGame = true;
 			//SpriteComponent->Mobility = EComponentMobility::Static;
 			SpriteComponent->SpriteInfo.Category = ConstructorStatics.ID_Misc;
@@ -282,7 +282,7 @@ void AEQSTestingPawn::UpdateDrawing()
 		SpriteComponent->MarkRenderStateDirty();
 	}	
 
-	if (EdRenderComp != NULL && EdRenderComp->bVisible)
+	if (EdRenderComp != NULL && EdRenderComp->GetVisibleFlag())
 	{
 		EdRenderComp->MarkRenderStateDirty();
 
diff --git a/Engine/Source/Runtime/AIModule/Private/Navigation/NavLinkProxy.cpp b/Engine/Source/Runtime/AIModule/Private/Navigation/NavLinkProxy.cpp
--- a/Engine/Source/Runtime/AIModule/Private/Navigation/NavLinkProxy.cpp
+++ b/Engine/Source/Runtime/AIModule/Private/Navigation/NavLinkProxy.cpp
@@ -22,7 +22,7 @@ ANavLinkProxy::ANavLinkProxy(const FObjectInitializer& ObjectInitializer) : Supe
 	USceneComponent* SceneComponent = CreateDefaultSubobject<USceneComponent>(TEXT("PositionComponent"));
 	RootComponent = SceneComponent;
 
-	bHidden = true;
+	SetHidden(true);
 
 #if WITH_EDITORONLY_DATA
 	EdRenderComp = CreateDefaultSubobject<UNavLinkRenderingComponent>(TEXT("EdRenderComp"));
@@ -48,9 +48,9 @@ ANavLinkProxy::ANavLinkProxy(const FObjectInitializer& ObjectInitializer) : Supe
 		static FConstructorStatics ConstructorStatics;
 
 		SpriteComponent->Sprite = ConstructorStatics.SpriteTexture.Get();
-		SpriteComponent->RelativeScale3D = FVector(0.5f, 0.5f, 0.5f);
+		SpriteComponent->SetRelativeScale3D(FVector(0.5f, 0.5f, 0.5f));
 		SpriteComponent->bHiddenInGame = true;
-		SpriteComponent->bVisible = true;
+		SpriteComponent->SetVisibleFlag(true);
 		SpriteComponent->SpriteInfo.Category = ConstructorStatics.ID_Decals;
 		SpriteComponent->SpriteInfo.DisplayName = ConstructorStatics.NAME_Decals;
 		SpriteComponent->SetupAttachment(RootComponent);
@@ -70,7 +70,7 @@ ANavLinkProxy::ANavLinkProxy(const FObjectInitializer& ObjectInitializer) : Supe
 	PointLinks.Add(DefLink);
 
 	SetActorEnableCollision(false);
-	bCanBeDamaged = false;
+	SetCanBeDamaged(false);
 }
 
 #if WITH_EDITOR
diff --git a/Engine/Source/Runtime/AIModule/Private/Perception/PawnSensingComponent.cpp b/Engine/Source/Runtime/AIModule/Private/Perception/PawnSensingComponent.cpp
--- a/Engine/Source/Runtime/AIModule/Private/Perception/PawnSensingComponent.cpp
+++ b/Engine/Source/Runtime/AIModule/Private/Perception/PawnSensingComponent.cpp
@@ -448,7 +448,7 @@ bool UPawnSensingComponent::ShouldCheckVisibilityOf(APawn *Pawn) const
 		return false;
 	}
 
-	if (Pawn->bHidden)
+	if (Pawn->IsHidden())
 	{
 		return false;
 	}
diff --git a/Engine/Source/Runtime/AudioMixer/Private/Components/SynthComponent.cpp b/Engine/Source/Runtime/AudioMixer/Private/Components/SynthComponent.cpp
--- a/Engine/Source/Runtime/AudioMixer/Private/Components/SynthComponent.cpp
+++ b/Engine/Source/Runtime/AudioMixer/Private/Components/SynthComponent.cpp
@@ -153,7 +153,7 @@ void USynthComponent::Activate(bool bReset)
 	if (bReset || ShouldActivate())
 	{
 		Start();
-		if (bIsActive)
+		if (IsActive())
 		{
 			OnComponentActivated.Broadcast(this, bReset);
 		}
@@ -166,7 +166,7 @@ void USynthComponent::Deactivate()
 	{
 		Stop();
 
-		if (!bIsActive)
+		if (!IsActive())
 		{
 			OnComponentDeactivated.Broadcast(this);
 		}
@@ -321,7 +321,7 @@ bool USynthComponent::IsReadyForOwnerToAutoDestroy() const
 #if WITH_EDITOR
 void USynthComponent::PostEditChangeProperty(FPropertyChangedEvent& PropertyChangedEvent)
 {
-	if (bIsActive)
+	if (IsActive())
 	{
 		// If this is an auto destroy component we need to prevent it from being auto-destroyed since we're really just restarting it
 		const bool bWasAutoDestroy = bAutoDestroy;
@@ -397,7 +397,7 @@ int32 USynthComponent::OnGeneratePCMAudio(float* GeneratedPCMData, int32 NumSamp
 void USynthComponent::Start()
 {
 	// Only need to start if we're not already active
-	if (bIsActive)
+	if (IsActive())
 	{
 		return;
 	}
@@ -442,9 +442,9 @@ void USynthComponent::Start()
 		Synth->SoundSubmixObject = SoundSubmix;
 		Synth->SoundSubmixSends = SoundSubmixSends;
 
-		bIsActive = AudioComponent->IsActive();
+		SetActiveFlag(AudioComponent->IsActive());
 
-		if (bIsActive)
+		if (IsActive())
 		{
 			PendingSynthEvents.Enqueue(ESynthEvent::Start);
 		}
@@ -453,7 +453,7 @@ void USynthComponent::Start()
 
 void USynthComponent::Stop()
 {
-	if (bIsActive)
+	if (IsActive())
 	{
 		PendingSynthEvents.Enqueue(ESynthEvent::Stop);
 
@@ -462,7 +462,7 @@ void USynthComponent::Stop()
 			AudioComponent->Stop();
 		}
 
-		bIsActive = false;
+		SetActiveFlag(false);
 	}
 }
 
diff --git a/Engine/Source/Runtime/AugmentedReality/Private/AROriginActor.cpp b/Engine/Source/Runtime/AugmentedReality/Private/AROriginActor.cpp
--- a/Engine/Source/Runtime/AugmentedReality/Private/AROriginActor.cpp
+++ b/Engine/Source/Runtime/AugmentedReality/Private/AROriginActor.cpp
@@ -10,8 +10,8 @@ AAROriginActor::AAROriginActor(const FObjectInitializer& ObjectInitializer)
 	PrimaryActorTick.bCanEverTick = false;
 	bAllowTickBeforeBeginPlay = false;
 	bReplicates = false;
-	bReplicateMovement = false;
-	bCanBeDamaged = false;
+	SetReplicatingMovement(false);
+	SetCanBeDamaged(false);
 }
 
 AAROriginActor* AAROriginActor::GetOriginActor()
diff --git a/Engine/Source/Runtime/CinematicCamera/Private/CameraRig_Crane.cpp b/Engine/Source/Runtime/CinematicCamera/Private/CameraRig_Crane.cpp
--- a/Engine/Source/Runtime/CinematicCamera/Private/CameraRig_Crane.cpp
+++ b/Engine/Source/Runtime/CinematicCamera/Private/CameraRig_Crane.cpp
@@ -27,16 +27,16 @@ ACameraRig_Crane::ACameraRig_Crane(const FObjectInitializer& ObjectInitializer)
 	// create the functional component hierarchy
 	CraneYawControl = CreateDefaultSubobject<USceneComponent>(TEXT("CraneYawControl"));
 	CraneYawControl->SetupAttachment(TransformComponent);
-	CraneYawControl->RelativeLocation = FVector(0.f, 0.f, 70.f);			// pivot height off the ground
-	CraneYawControl->RelativeRotation = FRotator(0.f, CraneYaw, 0.f);
+	CraneYawControl->SetRelativeLocation(FVector(0.f, 0.f, 70.f));			// pivot height off the ground
+	CraneYawControl->SetRelativeRotation(FRotator(0.f, CraneYaw, 0.f));
 
 	CranePitchControl = CreateDefaultSubobject<USceneComponent>(TEXT("CranePitchControl"));
 	CranePitchControl->SetupAttachment(CraneYawControl);
-	CranePitchControl->RelativeRotation = FRotator(CranePitch, 0.f, 0.f);
+	CranePitchControl->SetRelativeRotation(FRotator(CranePitch, 0.f, 0.f));
 
 	CraneCameraMount = CreateDefaultSubobject<USceneComponent>(TEXT("CraneCameraMount"));
 	CraneCameraMount->SetupAttachment(CranePitchControl);
-	CraneCameraMount->RelativeLocation = FVector(CraneArmLength, 0.f, -15.f);			// negative z == underslung mount
+	CraneCameraMount->SetRelativeLocation(FVector(CraneArmLength, 0.f, -15.f));			// negative z == underslung mount
 
 #if WITH_EDITORONLY_DATA
 	// create preview meshes
@@ -66,9 +66,9 @@ ACameraRig_Crane::ACameraRig_Crane(const FObjectInitializer& ObjectInitializer)
 			PreviewMesh_CraneArm->CastShadow = false;
 			
 			PreviewMesh_CraneArm->SetupAttachment(CranePitchControl);		// sibling of the mount
-			PreviewMesh_CraneArm->RelativeRotation = FRotator(0.f, 90.f, 0.f);
-			PreviewMesh_CraneArm->RelativeLocation = FVector(0.f, 0.f, 52.f);
-			PreviewMesh_CraneArm->RelativeScale3D = FVector(0.7f, 0.7f, 0.7f);
+			PreviewMesh_CraneArm->SetRelativeRotation(FRotator(0.f, 90.f, 0.f));
+			PreviewMesh_CraneArm->SetRelativeLocation(FVector(0.f, 0.f, 52.f));
+			PreviewMesh_CraneArm->SetRelativeScale3D(FVector(0.7f, 0.7f, 0.7f));
 		}
 		
 		static ConstructorHelpers::FObjectFinder<UStaticMesh> CraneArmMount(TEXT("/Engine/EditorMeshes/Camera/SM_CraneRig_Mount.SM_CraneRig_Mount"));
@@ -82,9 +82,9 @@ ACameraRig_Crane::ACameraRig_Crane(const FObjectInitializer& ObjectInitializer)
 			PreviewMesh_CraneMount->CastShadow = false;
 
 			PreviewMesh_CraneMount->SetupAttachment(CranePitchControl);
-			PreviewMesh_CraneMount->RelativeRotation = FRotator(0.f, 90.f, 0.f);
-			PreviewMesh_CraneMount->RelativeLocation = FVector(CraneArmLength, 0.f, 0.f);
-			PreviewMesh_CraneMount->RelativeScale3D = FVector(0.7f, 0.7f, 0.7f);
+			PreviewMesh_CraneMount->SetRelativeRotation(FRotator(0.f, 90.f, 0.f));
+			PreviewMesh_CraneMount->SetRelativeLocation(FVector(CraneArmLength, 0.f, 0.f));
+			PreviewMesh_CraneMount->SetRelativeScale3D(FVector(0.7f, 0.7f, 0.7f));
 		}
 		
 		static ConstructorHelpers::FObjectFinder<UStaticMesh> CraneArmCW(TEXT("/Engine/EditorMeshes/Camera/SM_CraneRig_Body.SM_CraneRig_Body"));
@@ -98,8 +98,8 @@ ACameraRig_Crane::ACameraRig_Crane(const FObjectInitializer& ObjectInitializer)
 			PreviewMesh_CraneCounterWeight->CastShadow = false;
 
 			PreviewMesh_CraneCounterWeight->SetupAttachment(CranePitchControl);
-			PreviewMesh_CraneCounterWeight->RelativeRotation = FRotator(0.f, 90.f, 0.f);
-			PreviewMesh_CraneCounterWeight->RelativeScale3D = FVector(0.7f, 0.7f, 0.7f);
+			PreviewMesh_CraneCounterWeight->SetRelativeRotation(FRotator(0.f, 90.f, 0.f));
+			PreviewMesh_CraneCounterWeight->SetRelativeScale3D(FVector(0.7f, 0.7f, 0.7f));
 		}
 
 		UpdatePreviewMeshes();
@@ -118,7 +118,7 @@ void ACameraRig_Crane::UpdatePreviewMeshes()
 		// this tailored to the dimensions of the asset (29cm in scaling direction, pivot at the edge)
 		float const CraneArmYScale = CraneArmLength / CraneArmMesh_DefaultMeshSize;
 
-		FVector NewRelScale3D = PreviewMesh_CraneArm->RelativeScale3D;
+		FVector NewRelScale3D = PreviewMesh_CraneArm->GetRelativeScale3D();
 		NewRelScale3D.Y = CraneArmYScale * 0.7f;
 
 		PreviewMesh_CraneArm->SetRelativeScale3D(NewRelScale3D);
@@ -135,15 +135,15 @@ void ACameraRig_Crane::UpdatePreviewMeshes()
 
 void ACameraRig_Crane::UpdateCraneComponents()
 {
-	FRotator NewYawControlRot = CraneYawControl->RelativeRotation;
+	FRotator NewYawControlRot = CraneYawControl->GetRelativeRotation();
 	NewYawControlRot.Yaw = CraneYaw;
 	CraneYawControl->SetRelativeRotation(FRotator(0.f, CraneYaw, 0.f));
 
-	FRotator NewPitchControlRot = CranePitchControl->RelativeRotation;
+	FRotator NewPitchControlRot = CranePitchControl->GetRelativeRotation();
 	NewPitchControlRot.Pitch = CranePitch;
 	CranePitchControl->SetRelativeRotation(FRotator(CranePitch, 0.f, 0.f));
 
-	FVector NewCameraMountLoc = CraneCameraMount->RelativeLocation;
+	FVector NewCameraMountLoc = CraneCameraMount->GetRelativeLocation();
 	NewCameraMountLoc.X = CraneArmLength;
 
 	// zero the pitch from the camera mount component
@@ -155,7 +155,7 @@ void ACameraRig_Crane::UpdateCraneComponents()
 	}
 	if (!bLockMountYaw)
 	{
-		NewCameraMountWorldRot.Yaw = RootComponent->RelativeRotation.Yaw;
+		NewCameraMountWorldRot.Yaw = RootComponent->GetRelativeRotation().Yaw;
 	}
 	NewCameraMountWorldRot.Roll = 0.f;
 
diff --git a/Engine/Source/Runtime/CinematicCamera/Private/CineCameraComponent.cpp b/Engine/Source/Runtime/CinematicCamera/Private/CineCameraComponent.cpp
--- a/Engine/Source/Runtime/CinematicCamera/Private/CineCameraComponent.cpp
+++ b/Engine/Source/Runtime/CinematicCamera/Private/CineCameraComponent.cpp
@@ -491,8 +491,8 @@ void UCineCameraComponent::CreateDebugFocusPlane()
 			DebugFocusPlaneComponent->CreationMethod = CreationMethod;
 			DebugFocusPlaneComponent->bSelectable = false;
 
-			DebugFocusPlaneComponent->RelativeScale3D = FVector(10000.f, 10000.f, 1.f);
-			DebugFocusPlaneComponent->RelativeRotation = FRotator(90.f, 0.f, 0.f);
+			DebugFocusPlaneComponent->SetRelativeScale3D_Direct(FVector(10000.f, 10000.f, 1.f));
+			DebugFocusPlaneComponent->SetRelativeRotation_Direct(FRotator(90.f, 0.f, 0.f));
 
 			DebugFocusPlaneComponent->RegisterComponentWithWorld(GetWorld());
 
diff --git a/Engine/Source/Runtime/Engine/Classes/GameFramework/ProjectileMovementComponent.h b/Engine/Source/Runtime/Engine/Classes/GameFramework/ProjectileMovementComponent.h
--- a/Engine/Source/Runtime/Engine/Classes/GameFramework/ProjectileMovementComponent.h
+++ b/Engine/Source/Runtime/Engine/Classes/GameFramework/ProjectileMovementComponent.h
@@ -231,7 +231,7 @@ public:
 	UFUNCTION(BlueprintCallable, Category="Game|Components|ProjectileMovement")
 	virtual void StopSimulating(const FHitResult& HitResult);
 
-	bool HasStoppedSimulation() { return (UpdatedComponent == nullptr) || (bIsActive == false); }
+	bool HasStoppedSimulation() { return (UpdatedComponent == nullptr) || (IsActive() == false); }
 
 	/**
 	 * Compute remaining time step given remaining time and current iterations.
diff --git a/Engine/Source/Runtime/Engine/Private/ActorEditor.cpp b/Engine/Source/Runtime/Engine/Private/ActorEditor.cpp
--- a/Engine/Source/Runtime/Engine/Private/ActorEditor.cpp
+++ b/Engine/Source/Runtime/Engine/Private/ActorEditor.cpp
@@ -45,9 +45,9 @@ void AActor::PreEditChange(UProperty* PropertyThatWillChange)
 	}
 }
 
-static FName Name_RelativeLocation = GET_MEMBER_NAME_CHECKED(USceneComponent, RelativeLocation);
-static FName Name_RelativeRotation = GET_MEMBER_NAME_CHECKED(USceneComponent, RelativeRotation);
-static FName Name_RelativeScale3D = GET_MEMBER_NAME_CHECKED(USceneComponent, RelativeScale3D);
+static FName Name_RelativeLocation = USceneComponent::GetRelativeLocationPropertyName();
+static FName Name_RelativeRotation = USceneComponent::GetRelativeRotationPropertyName();
+static FName Name_RelativeScale3D = USceneComponent::GetRelativeScale3DPropertyName();
 
 void AActor::PostEditChangeProperty(FPropertyChangedEvent& PropertyChangedEvent)
 {
@@ -228,7 +228,7 @@ void AActor::DebugShowOneComponentHierarchy( USceneComponent* SceneComp, int32&
 	{
 		FVector Posn = SceneComp->GetComponentTransform().GetLocation();
 		//PosString = FString::Printf( TEXT("{R:%f,%f,%f- W:%f,%f,%f}"), SceneComp->RelativeLocation.X, SceneComp->RelativeLocation.Y, SceneComp->RelativeLocation.Z, Posn.X, Posn.Y, Posn.Z );
-		PosString = FString::Printf( TEXT("{R:%f- W:%f}"), SceneComp->RelativeLocation.Z, Posn.Z );
+		PosString = FString::Printf( TEXT("{R:%f- W:%f}"), SceneComp->GetRelativeLocation().Z, Posn.Z );
 	}
 	else
 	{
@@ -249,7 +249,7 @@ void AActor::DebugShowOneComponentHierarchy( USceneComponent* SceneComp, int32&
 		{
 			FVector Posn = SceneComp->GetComponentTransform().GetLocation();
 			//PosString = FString::Printf( TEXT("{R:%f,%f,%f- W:%f,%f,%f}"), SceneComp->RelativeLocation.X, SceneComp->RelativeLocation.Y, SceneComp->RelativeLocation.Z, Posn.X, Posn.Y, Posn.Z );
-			PosString = FString::Printf( TEXT("{R:%f- W:%f}"), SceneComp->RelativeLocation.Z, Posn.Z );
+			PosString = FString::Printf( TEXT("{R:%f- W:%f}"), SceneComp->GetRelativeLocation().Z, Posn.Z );
 		}
 		else
 		{
@@ -592,7 +592,7 @@ void AActor::EditorApplyRotation(const FRotator& DeltaRotation, bool bAltDown, b
 {
 	if( RootComponent != NULL )
 	{
-		FRotator Rot = RootComponent->GetAttachParent() != NULL ? GetActorRotation() : RootComponent->RelativeRotation;
+		FRotator Rot = RootComponent->GetAttachParent() != NULL ? GetActorRotation() : RootComponent->GetRelativeRotation();
 		FRotator ActorRotWind, ActorRotRem;
 		Rot.GetWindingAndRemainder(ActorRotWind, ActorRotRem);
 		const FQuat ActorQ = ActorRotRem.Quaternion();
@@ -610,7 +610,7 @@ void AActor::EditorApplyRotation(const FRotator& DeltaRotation, bool bAltDown, b
 			NewRelRotation = RootComponent->GetRelativeRotationFromWorld(NewRelRotation);
 			NewActorRotRem = FRotator(NewRelRotation);
 			//now we need to get current relative rotation to find the diff
-			Rot = RootComponent->RelativeRotation;
+			Rot = RootComponent->GetRelativeRotation();
 			Rot.GetWindingAndRemainder(ActorRotWind, ActorRotRem);
 		}
 		else
@@ -635,7 +635,7 @@ void AActor::EditorApplyScale( const FVector& DeltaScale, const FVector* PivotLo
 {
 	if( RootComponent != NULL )
 	{
-		const FVector CurrentScale = GetRootComponent()->RelativeScale3D;
+		const FVector CurrentScale = GetRootComponent()->GetRelativeScale3D();
 
 		// @todo: Remove this hack once we have decided on the scaling method to use.
 		FVector ScaleToApply;
@@ -694,7 +694,7 @@ void AActor::EditorApplyMirror(const FVector& MirrorScale, const FVector& PivotL
 		Loc += PivotLocation;
 		GetRootComponent()->SetRelativeLocation( Loc );
 
-		FVector Scale3D = GetRootComponent()->RelativeScale3D;
+		FVector Scale3D = GetRootComponent()->GetRelativeScale3D();
 		Scale3D.X = -Scale3D.X;
 		GetRootComponent()->SetRelativeScale3D(Scale3D);
 	}
@@ -940,7 +940,7 @@ void AActor::CheckForErrors()
 	}
 
 	UPrimitiveComponent* PrimComp = Cast<UPrimitiveComponent>(RootComponent);
-	if( PrimComp && (PrimComp->Mobility != EComponentMobility::Movable) && PrimComp->BodyInstance.bSimulatePhysics)
+	if (PrimComp && (PrimComp->Mobility != EComponentMobility::Movable) && PrimComp->BodyInstance.bSimulatePhysics)
 	{
 		FFormatNamedArguments Arguments;
 		Arguments.Add(TEXT("ActorName"), FText::FromString(GetPathName()));
@@ -950,14 +950,18 @@ void AActor::CheckForErrors()
 			->AddToken(FMapErrorToken::Create(FMapErrors::StaticPhysNone));
 	}
 
-	if( RootComponent && FMath::IsNearlyZero( GetRootComponent()->RelativeScale3D.X * GetRootComponent()->RelativeScale3D.Y * GetRootComponent()->RelativeScale3D.Z ) )
+	if (RootComponent)
 	{
-		FFormatNamedArguments Arguments;
-		Arguments.Add(TEXT("ActorName"), FText::FromString(GetPathName()));
-		FMessageLog("MapCheck").Error()
-			->AddToken(FUObjectToken::Create(this))
-			->AddToken(FTextToken::Create(FText::Format(LOCTEXT( "MapCheck_Message_InvalidDrawscale", "{ActorName} : Invalid DrawScale/DrawScale3D" ), Arguments) ))
-			->AddToken(FMapErrorToken::Create(FMapErrors::InvalidDrawscale));
+		const FVector LocalRelativeScale3D = RootComponent->GetRelativeScale3D();
+		if (FMath::IsNearlyZero(LocalRelativeScale3D.X * LocalRelativeScale3D.Y * LocalRelativeScale3D.Z))
+		{
+			FFormatNamedArguments Arguments;
+			Arguments.Add(TEXT("ActorName"), FText::FromString(GetPathName()));
+			FMessageLog("MapCheck").Error()
+				->AddToken(FUObjectToken::Create(this))
+				->AddToken(FTextToken::Create(FText::Format(LOCTEXT("MapCheck_Message_InvalidDrawscale", "{ActorName} : Invalid DrawScale/DrawScale3D"), Arguments)))
+				->AddToken(FMapErrorToken::Create(FMapErrors::InvalidDrawscale));
+		}
 	}
 
 	// Route error checking to components.
diff --git a/Engine/Source/Runtime/Engine/Private/AmbientSound.cpp b/Engine/Source/Runtime/Engine/Private/AmbientSound.cpp
--- a/Engine/Source/Runtime/Engine/Private/AmbientSound.cpp
+++ b/Engine/Source/Runtime/Engine/Private/AmbientSound.cpp
@@ -30,8 +30,8 @@ AAmbientSound::AAmbientSound(const FObjectInitializer& ObjectInitializer)
 	RootComponent = AudioComponent;
 
 	bReplicates = false;
-	bHidden = true;
-	bCanBeDamaged = false;
+	SetHidden(true);
+	SetCanBeDamaged(false);
 }
 
 #if WITH_EDITOR
diff --git a/Engine/Source/Runtime/Engine/Private/Animation/AnimNotifyState_Trail.cpp b/Engine/Source/Runtime/Engine/Private/Animation/AnimNotifyState_Trail.cpp
--- a/Engine/Source/Runtime/Engine/Private/Animation/AnimNotifyState_Trail.cpp
+++ b/Engine/Source/Runtime/Engine/Private/Animation/AnimNotifyState_Trail.cpp
@@ -162,7 +162,7 @@ void UAnimNotifyState_Trail::NotifyBegin(class USkeletalMeshComponent * MeshComp
 		NewParticleComp->SecondsBeforeInactive = 0.0f;
 		NewParticleComp->bAutoActivate = false;
 		NewParticleComp->bOverrideLODMethod = false;
-		NewParticleComp->RelativeScale3D = FVector(1.f);
+		NewParticleComp->SetRelativeScale3D_Direct(FVector(1.f));
 		NewParticleComp->bAutoManageAttachment = true; // Let it detach when finished (only happens if not auto-destroying)
 		NewParticleComp->SetAutoAttachParams(MeshComp, NAME_None);
 
diff --git a/Engine/Source/Runtime/Engine/Private/Animation/SkeletalMeshActor.cpp b/Engine/Source/Runtime/Engine/Private/Animation/SkeletalMeshActor.cpp
--- a/Engine/Source/Runtime/Engine/Private/Animation/SkeletalMeshActor.cpp
+++ b/Engine/Source/Runtime/Engine/Private/Animation/SkeletalMeshActor.cpp
@@ -169,7 +169,7 @@ void ASkeletalMeshActor::PostInitializeComponents()
 {
 	Super::PostInitializeComponents();
 	// grab the current mesh for replication
-	if (Role == ROLE_Authority && SkeletalMeshComponent)
+	if (GetLocalRole() == ROLE_Authority && SkeletalMeshComponent)
 	{
 		ReplicatedMesh = SkeletalMeshComponent->SkeletalMesh;
 	}
@@ -291,7 +291,7 @@ void ASkeletalMeshActor::LoadedFromAnotherClass(const FName& OldClassName)
 			SkeletalMeshComponent->bBlendPhysics = true;
 
 			bAlwaysRelevant = true;
-			bReplicateMovement = true;
+			SetReplicatingMovement(true);
 			SetRemoteRoleForBackwardsCompat(ROLE_SimulatedProxy);
 			bReplicates = true;
 		}
diff --git a/Engine/Source/Runtime/Engine/Private/Atmosphere/AtmosphericFog.cpp b/Engine/Source/Runtime/Engine/Private/Atmosphere/AtmosphericFog.cpp
--- a/Engine/Source/Runtime/Engine/Private/Atmosphere/AtmosphericFog.cpp
+++ b/Engine/Source/Runtime/Engine/Private/Atmosphere/AtmosphericFog.cpp
@@ -48,7 +48,7 @@ AAtmosphericFog::AAtmosphericFog(const FObjectInitializer& ObjectInitializer)
 		if (GetSpriteComponent())
 		{
 			GetSpriteComponent()->Sprite = ConstructorStatics.FogTextureObject.Get();
-			GetSpriteComponent()->RelativeScale3D = FVector(0.5f, 0.5f, 0.5f);
+			GetSpriteComponent()->SetRelativeScale3D(FVector(0.5f, 0.5f, 0.5f));
 			GetSpriteComponent()->SpriteInfo.Category = ConstructorStatics.ID_Fog;
 			GetSpriteComponent()->SpriteInfo.DisplayName = ConstructorStatics.NAME_Fog;
 			GetSpriteComponent()->SetupAttachment(AtmosphericFogComponent);
@@ -70,7 +70,7 @@ AAtmosphericFog::AAtmosphericFog(const FObjectInitializer& ObjectInitializer)
 
 
 	PrimaryActorTick.bCanEverTick = true;
-	bHidden = false;
+	SetHidden(false);
 }
 
 #if WITH_EDITOR
diff --git a/Engine/Source/Runtime/Engine/Private/Brush.cpp b/Engine/Source/Runtime/Engine/Private/Brush.cpp
--- a/Engine/Source/Runtime/Engine/Private/Brush.cpp
+++ b/Engine/Source/Runtime/Engine/Private/Brush.cpp
@@ -42,9 +42,9 @@ ABrush::ABrush(const FObjectInitializer& ObjectInitializer)
 
 	RootComponent = BrushComponent;
 	
-	bHidden = true;
+	SetHidden(true);
 	bNotForClientOrServer = false;
-	bCanBeDamaged = false;
+	SetCanBeDamaged(false);
 	bCollideWhenPlacing = true;
 	SpawnCollisionHandlingMethod = ESpawnActorCollisionHandlingMethod::AdjustIfPossibleButDontSpawnIfColliding;
 }
diff --git a/Engine/Source/Runtime/Engine/Private/Camera/CameraComponent.cpp b/Engine/Source/Runtime/Engine/Private/Camera/CameraComponent.cpp
--- a/Engine/Source/Runtime/Engine/Private/Camera/CameraComponent.cpp
+++ b/Engine/Source/Runtime/Engine/Private/Camera/CameraComponent.cpp
@@ -363,11 +363,11 @@ void UCameraComponent::CheckForErrors()
 
 bool UCameraComponent::GetEditorPreviewInfo(float DeltaTime, FMinimalViewInfo& ViewOut)
 {
-	if (bIsActive)
+	if (IsActive())
 	{
 		GetCameraView(DeltaTime, ViewOut);
 	}
-	return bIsActive;
+	return IsActive();
 }
 #endif	// WITH_EDITOR
 
diff --git a/Engine/Source/Runtime/Engine/Private/Character.cpp b/Engine/Source/Runtime/Engine/Private/Character.cpp
--- a/Engine/Source/Runtime/Engine/Private/Character.cpp
+++ b/Engine/Source/Runtime/Engine/Private/Character.cpp
@@ -118,7 +118,7 @@ void ACharacter::PostInitializeComponents()
 	{
 		if (Mesh)
 		{
-			CacheInitialMeshOffset(Mesh->RelativeLocation, Mesh->RelativeRotation);
+			CacheInitialMeshOffset(Mesh->GetRelativeLocation(), Mesh->GetRelativeRotation());
 
 			// force animation tick after movement component updates
 			if (Mesh->PrimaryComponentTick.bCanEverTick && CharacterMovement)
@@ -158,9 +158,11 @@ void ACharacter::CacheInitialMeshOffset(FVector MeshRelativeLocation, FRotator M
 	{
 		logOrEnsureNanError(TEXT("ACharacter::PostInitializeComponents detected NaN in BaseRotationOffset! (%s)"), *BaseRotationOffset.ToString());
 	}
-	if (Mesh->RelativeRotation.ContainsNaN())
+
+	const FRotator LocalRotation = Mesh->GetRelativeRotation();
+	if (LocalRotation.ContainsNaN())
 	{
-		logOrEnsureNanError(TEXT("ACharacter::PostInitializeComponents detected NaN in Mesh->RelativeRotation! (%s)"), *Mesh->RelativeRotation.ToString());
+		logOrEnsureNanError(TEXT("ACharacter::PostInitializeComponents detected NaN in Mesh->RelativeRotation! (%s)"), *LocalRotation.ToString());
 	}
 #endif
 }
@@ -316,7 +318,7 @@ bool ACharacter::IsJumpProvidingForce() const
 	{
 		return true;
 	}
-	else if (bProxyIsJumpForceApplied && (Role==ROLE_SimulatedProxy))
+	else if (bProxyIsJumpForceApplied && (GetLocalRole() == ROLE_SimulatedProxy))
 	{
 		return GetWorld()->TimeSince(ProxyJumpForceStartedTime) <= GetJumpMaxHoldTime();
 	}
@@ -359,7 +361,7 @@ void ACharacter::SetReplicateMovement(bool bInReplicateMovement)
 {
 	Super::SetReplicateMovement(bInReplicateMovement);
 
-	if (CharacterMovement != nullptr && Role == ROLE_Authority)
+	if (CharacterMovement != nullptr && GetLocalRole() == ROLE_Authority)
 	{
 		// Set prediction data time stamp to current time to stop extrapolating
 		// from time bReplicateMovement was turned off to when it was turned on again
@@ -410,8 +412,9 @@ void ACharacter::OnEndCrouch( float HeightAdjust, float ScaledHeightAdjust )
 	const ACharacter* DefaultChar = GetDefault<ACharacter>(GetClass());
 	if (Mesh && DefaultChar->Mesh)
 	{
-		Mesh->RelativeLocation.Z = DefaultChar->Mesh->RelativeLocation.Z;
-		BaseTranslationOffset.Z = Mesh->RelativeLocation.Z;
+		FVector& MeshRelativeLocation = Mesh->GetRelativeLocation_DirectMutable();
+		MeshRelativeLocation.Z = DefaultChar->Mesh->GetRelativeLocation().Z;
+		BaseTranslationOffset.Z = MeshRelativeLocation.Z;
 	}
 	else
 	{
@@ -428,8 +431,9 @@ void ACharacter::OnStartCrouch( float HeightAdjust, float ScaledHeightAdjust )
 	const ACharacter* DefaultChar = GetDefault<ACharacter>(GetClass());
 	if (Mesh && DefaultChar->Mesh)
 	{
-		Mesh->RelativeLocation.Z = DefaultChar->Mesh->RelativeLocation.Z + HeightAdjust;
-		BaseTranslationOffset.Z = Mesh->RelativeLocation.Z;
+		FVector& MeshRelativeLocation = Mesh->GetRelativeLocation_DirectMutable();
+		MeshRelativeLocation.Z = DefaultChar->Mesh->GetRelativeLocation().Z + HeightAdjust;
+		BaseTranslationOffset.Z = MeshRelativeLocation.Z;
 	}
 	else
 	{
@@ -717,10 +721,11 @@ void ACharacter::SetBase( UPrimitiveComponent* NewBaseComponent, const FName InB
 				CharacterMovement->PostPhysicsTickFunction.SetTickFunctionEnable(false);
 			}
 
-			if (Role == ROLE_Authority || Role == ROLE_AutonomousProxy)
+			const ENetRole LocalRole = GetLocalRole();
+			if (LocalRole == ROLE_Authority || LocalRole == ROLE_AutonomousProxy)
 			{
 				BasedMovement.bServerHasBaseComponent = (BasedMovement.MovementBase != nullptr); // Also set on proxies for nicer debugging.
-				UE_LOG(LogCharacter, Verbose, TEXT("Setting base on %s for '%s' to '%s'"), Role == ROLE_Authority ? TEXT("Server") : TEXT("AutoProxy"), *GetName(), *GetFullNameSafe(NewBaseComponent));
+				UE_LOG(LogCharacter, Verbose, TEXT("Setting base on %s for '%s' to '%s'"), LocalRole == ROLE_Authority ? TEXT("Server") : TEXT("AutoProxy"), *GetName(), *GetFullNameSafe(NewBaseComponent));
 			}
 			else
 			{
@@ -816,7 +821,7 @@ void ACharacter::PossessedBy(AController* NewController)
 	Super::PossessedBy(NewController);
 
 	// If we are controlled remotely, set animation timing to be driven by client's network updates. So timing and events remain in sync.
-	if (Mesh && bReplicateMovement && (GetRemoteRole() == ROLE_AutonomousProxy && GetNetConnection() != nullptr))
+	if (Mesh && IsReplicatingMovement() && (GetRemoteRole() == ROLE_AutonomousProxy && GetNetConnection() != nullptr))
 	{
 		Mesh->bOnlyAllowAutonomousTickPose = true;
 	}
@@ -1058,7 +1063,7 @@ void ACharacter::PreNetReceive()
 
 void ACharacter::PostNetReceive()
 {
-	if (Role == ROLE_SimulatedProxy)
+	if (GetLocalRole() == ROLE_SimulatedProxy)
 	{
 		CharacterMovement->bNetworkMovementModeChanged |= (SavedMovementMode != ReplicatedMovementMode);
 		CharacterMovement->bNetworkUpdateReceived |= CharacterMovement->bNetworkMovementModeChanged || CharacterMovement->bJustTeleported;
@@ -1069,7 +1074,7 @@ void ACharacter::PostNetReceive()
 
 void ACharacter::OnRep_ReplicatedBasedMovement()
 {	
-	if (Role != ROLE_SimulatedProxy)
+	if (GetLocalRole() != ROLE_SimulatedProxy)
 	{
 		return;
 	}
@@ -1162,7 +1167,7 @@ void ACharacter::OnRep_RootMotion()
 		return;
 	}
 
-	if (Role == ROLE_SimulatedProxy)
+	if (GetLocalRole() == ROLE_SimulatedProxy)
 	{
 
 		UE_LOG(LogRootMotion, Log,  TEXT("ACharacter::OnRep_RootMotion"));
@@ -1370,18 +1375,19 @@ void ACharacter::OnUpdateSimulatedPosition(const FVector& OldLocation, const FQu
 
 void ACharacter::PostNetReceiveLocationAndRotation()
 {
-	if( Role == ROLE_SimulatedProxy )
+	if(GetLocalRole() == ROLE_SimulatedProxy)
 	{
 		// Don't change transform if using relative position (it should be nearly the same anyway, or base may be slightly out of sync)
 		if (!ReplicatedBasedMovement.HasRelativeLocation())
 		{
+			const FRepMovement& ConstRepMovement = GetReplicatedMovement();
 			const FVector OldLocation = GetActorLocation();
-			const FVector NewLocation = FRepMovement::RebaseOntoLocalOrigin(ReplicatedMovement.Location, this);
+			const FVector NewLocation = FRepMovement::RebaseOntoLocalOrigin(ConstRepMovement.Location, this);
 			const FQuat OldRotation = GetActorQuat();
 
 			CharacterMovement->bNetworkSmoothingComplete = false;
 			CharacterMovement->bJustTeleported |= (OldLocation != NewLocation);
-			CharacterMovement->SmoothCorrection(OldLocation, OldRotation, NewLocation, ReplicatedMovement.Rotation.Quaternion());
+			CharacterMovement->SmoothCorrection(OldLocation, OldRotation, NewLocation, ConstRepMovement.Rotation.Quaternion());
 			OnUpdateSimulatedPosition(OldLocation, OldRotation);
 		}
 		CharacterMovement->bNetworkUpdateReceived = true;
@@ -1464,7 +1470,7 @@ void ACharacter::PreReplicationForReplay(IRepChangedPropertyTracker & ChangedPro
 	{
 		// On client replays, our view pitch will be set to 0 as by default we do not replicate
 		// pitch for owners, just for simulated. So instead push our rotation into the sampler
-		if (World->IsRecordingClientReplay() && Controller != nullptr && Role == ROLE_AutonomousProxy && GetNetMode() == NM_Client)
+		if (World->IsRecordingClientReplay() && Controller != nullptr && GetLocalRole() == ROLE_AutonomousProxy && GetNetMode() == NM_Client)
 		{
 			SetRemoteViewPitch(Controller->GetControlRotation().Pitch);
 		}
diff --git a/Engine/Source/Runtime/Engine/Private/CheatManager.cpp b/Engine/Source/Runtime/Engine/Private/CheatManager.cpp
--- a/Engine/Source/Runtime/Engine/Private/CheatManager.cpp
+++ b/Engine/Source/Runtime/Engine/Private/CheatManager.cpp
@@ -181,16 +181,16 @@ void UCheatManager::Ghost()
 void UCheatManager::God()
 {
 	APawn* Pawn = GetOuterAPlayerController()->GetPawn();
-	if ( Pawn != NULL )
+	if (Pawn != nullptr)
 	{
-		if ( Pawn->bCanBeDamaged )
+		if (Pawn->CanBeDamaged())
 		{
-			Pawn->bCanBeDamaged = false;
+			Pawn->SetCanBeDamaged(false);
 			GetOuterAPlayerController()->ClientMessage(TEXT("God mode on"));
 		}
 		else
 		{
-			Pawn->bCanBeDamaged = true;
+			Pawn->SetCanBeDamaged(true);
 			GetOuterAPlayerController()->ClientMessage(TEXT("God Mode off"));
 		}
 	}
@@ -369,7 +369,7 @@ void UCheatManager::Summon( const FString& ClassName )
 
 					SpawnLoc += 72.f * SpawnRot.Vector() + FVector(0.f, 0.f, 1.f) * 15.f;
 					FActorSpawnParameters SpawnInfo;
-					SpawnInfo.Instigator = MyPlayerController->Instigator;
+					SpawnInfo.Instigator = MyPlayerController->GetInstigator();
 					AActor* Actor = GetWorld()->SpawnActor(NewClass, &SpawnLoc, &SpawnRot, SpawnInfo );
 					if ( Actor )
 					{
@@ -576,7 +576,7 @@ void UCheatManager::EnableDebugCamera()
 		{
 			// spawn if necessary
 			FActorSpawnParameters SpawnInfo;
-			SpawnInfo.Instigator = PC->Instigator;
+			SpawnInfo.Instigator = PC->GetInstigator();
 			DebugCameraControllerRef = GetWorld()->SpawnActor<ADebugCameraController>(DebugCameraControllerClass, SpawnInfo );
 		}
 		if (DebugCameraControllerRef)
diff --git a/Engine/Source/Runtime/Engine/Private/ComponentInstanceDataCache.cpp b/Engine/Source/Runtime/Engine/Private/ComponentInstanceDataCache.cpp
--- a/Engine/Source/Runtime/Engine/Private/ComponentInstanceDataCache.cpp
+++ b/Engine/Source/Runtime/Engine/Private/ComponentInstanceDataCache.cpp
@@ -38,9 +38,9 @@ public:
 				// If this is the owning Actor's root scene component, don't include relative transform properties. This is handled elsewhere.
 				if (Component == ComponentOwner->GetRootComponent())
 				{
-					PropertiesToSkip.Add(ComponentClass->FindPropertyByName(GET_MEMBER_NAME_CHECKED(USceneComponent, RelativeLocation)));
-					PropertiesToSkip.Add(ComponentClass->FindPropertyByName(GET_MEMBER_NAME_CHECKED(USceneComponent, RelativeRotation)));
-					PropertiesToSkip.Add(ComponentClass->FindPropertyByName(GET_MEMBER_NAME_CHECKED(USceneComponent, RelativeScale3D)));
+					PropertiesToSkip.Add(ComponentClass->FindPropertyByName(USceneComponent::GetRelativeLocationPropertyName()));
+					PropertiesToSkip.Add(ComponentClass->FindPropertyByName(USceneComponent::GetRelativeRotationPropertyName()));
+					PropertiesToSkip.Add(ComponentClass->FindPropertyByName(USceneComponent::GetRelativeScale3DPropertyName()));
 				}
 			}
 
diff --git a/Engine/Source/Runtime/Engine/Private/Components/AudioComponent.cpp b/Engine/Source/Runtime/Engine/Private/Components/AudioComponent.cpp
--- a/Engine/Source/Runtime/Engine/Private/Components/AudioComponent.cpp
+++ b/Engine/Source/Runtime/Engine/Private/Components/AudioComponent.cpp
@@ -92,7 +92,7 @@ void UAudioComponent::BeginDestroy()
 {
 	Super::BeginDestroy();
 
-	if (bIsActive && Sound && Sound->IsLooping())
+	if (IsActive() && Sound && Sound->IsLooping())
 	{
 		UE_LOG(LogAudio, Verbose, TEXT("Audio Component is being destroyed prior to stopping looping sound '%s' directly."), *Sound->GetFullName());
 		Stop();
@@ -210,9 +210,9 @@ void UAudioComponent::OnRegister()
 			}
 		}
 
-		SavedAutoAttachRelativeLocation = RelativeLocation;
-		SavedAutoAttachRelativeRotation = RelativeRotation;
-		SavedAutoAttachRelativeScale3D = RelativeScale3D;
+		SavedAutoAttachRelativeLocation = GetRelativeLocation();
+		SavedAutoAttachRelativeRotation = GetRelativeRotation();
+		SavedAutoAttachRelativeScale3D = GetRelativeScale3D();
 	}
 
 	Super::OnRegister();
@@ -276,7 +276,7 @@ void UAudioComponent::OnUpdateTransform(EUpdateTransformFlags UpdateTransformFla
 
 	if (FAudioDevice* AudioDevice = GetAudioDevice())
 	{
-		if (bIsActive)
+		if (IsActive())
 		{
 			DECLARE_CYCLE_STAT(TEXT("FAudioThreadTask.UpdateAudioComponentTransform"), STAT_AudioUpdateComponentTransform, STATGROUP_AudioThreadCommands);
 
@@ -315,9 +315,9 @@ void UAudioComponent::CancelAutoAttachment(bool bDetachFromParent, const UWorld*
 		if (bDidAutoAttach)
 		{
 			// Restore relative transform from before attachment. Actual transform will be updated as part of DetachFromParent().
-			RelativeLocation = SavedAutoAttachRelativeLocation;
-			RelativeRotation = SavedAutoAttachRelativeRotation;
-			RelativeScale3D = SavedAutoAttachRelativeScale3D;
+			SetRelativeLocation_Direct(SavedAutoAttachRelativeLocation);
+			SetRelativeRotation_Direct(SavedAutoAttachRelativeRotation);
+			SetRelativeScale3D_Direct(SavedAutoAttachRelativeScale3D);
 			bDidAutoAttach = false;
 		}
 
@@ -366,7 +366,7 @@ void UAudioComponent::PlayInternal(const float StartTime, const float FadeInDura
 	// Reset our fading out flag in case this is a reused audio component and we are replaying after previously fading out
 	bIsFadingOut = false;
 
-	if (bIsActive)
+	if (IsActive())
 	{
 		// If this is an auto destroy component we need to prevent it from being auto-destroyed since we're really just restarting it
 		bool bCurrentAutoDestroy = bAutoDestroy;
@@ -404,9 +404,9 @@ void UAudioComponent::PlayInternal(const float StartTime, const float FadeInDura
 					{
 						bDidAutoAttach = bWasAutoAttached;
 						CancelAutoAttachment(true, World);
-						SavedAutoAttachRelativeLocation = RelativeLocation;
-						SavedAutoAttachRelativeRotation = RelativeRotation;
-						SavedAutoAttachRelativeScale3D = RelativeScale3D;
+						SavedAutoAttachRelativeLocation = GetRelativeLocation();
+						SavedAutoAttachRelativeRotation = GetRelativeRotation();
+						SavedAutoAttachRelativeScale3D = GetRelativeScale3D();
 						AttachToComponent(NewParent, FAttachmentTransformRules(AutoAttachLocationRule, AutoAttachRotationRule, AutoAttachScaleRule, false), AutoAttachSocketName);
 					}
 
@@ -524,7 +524,7 @@ void UAudioComponent::PlayInternal(const float StartTime, const float FadeInDura
 
 			// In editor, the audio thread is not run separate from the game thread, and can result in calling PlaybackComplete prior
 			// to bIsActive being set. Therefore, we assign to the current state of ActiveCount as opposed to just setting to true.
-			bIsActive = ActiveCount > 0;
+			SetActiveFlag(ActiveCount > 0);
 		}
 	}
 }
@@ -571,7 +571,7 @@ void UAudioComponent::AdjustVolume(float AdjustVolumeDuration, float AdjustVolum
 
 void UAudioComponent::AdjustVolumeInternal(float AdjustVolumeDuration, float AdjustVolumeLevel, bool bIsFadeOut, const EAudioFaderCurve FadeCurve)
 {
-	if (!bIsActive)
+	if (!IsActive())
 	{
 		return;
 	}
@@ -647,7 +647,7 @@ void UAudioComponent::AdjustVolumeInternal(float AdjustVolumeDuration, float Adj
 
 void UAudioComponent::Stop()
 {
-	if (!bIsActive)
+	if (!IsActive())
 	{
 		return;
 	}
@@ -659,7 +659,7 @@ void UAudioComponent::Stop()
 	}
 
 	// Set this to immediately be inactive
-	bIsActive = false;
+	SetActiveFlag(false);
 
 	UE_LOG(LogAudio, Verbose, TEXT("%g: Stopping AudioComponent : '%s' with Sound: '%s'"),
 		GetWorld() ? GetWorld()->GetAudioTimeSeconds() : 0.0f, *GetFullName(),
@@ -677,7 +677,7 @@ void UAudioComponent::StopDelayed(float DelayTime)
 		return;
 	}
 
-	if (!bIsActive)
+	if (!IsActive())
 	{
 		return;
 	}
@@ -737,7 +737,7 @@ void UAudioComponent::SetPaused(bool bPause)
 	{
 		bIsPaused = bPause;
 
-		if (bIsActive)
+		if (IsActive())
 		{
 			UE_LOG(LogAudio, Verbose, TEXT("%g: Pausing AudioComponent : '%s' with Sound: '%s'"), GetWorld() ? GetWorld()->GetAudioTimeSeconds() : 0.0f, *GetFullName(), Sound ? *Sound->GetName() : TEXT("nullptr"));
 
@@ -781,7 +781,7 @@ void UAudioComponent::PlaybackCompleted(bool bFailedToStart)
 	}
 
 	// Mark inactive before calling destroy to avoid recursion
-	bIsActive = false;
+	SetActiveFlag(false);
 
 	const UWorld* MyWorld = GetWorld();
 	if (!bFailedToStart && MyWorld != nullptr && (OnAudioFinished.IsBound() || OnAudioFinishedNative.IsBound()))
@@ -807,13 +807,13 @@ void UAudioComponent::PlaybackCompleted(bool bFailedToStart)
 
 bool UAudioComponent::IsPlaying() const
 {
-	return bIsActive;
+	return IsActive();
 }
 
 EAudioComponentPlayState UAudioComponent::GetPlayState() const
 {
 	UWorld* World = GetWorld();
-	if (!bIsActive || !World)
+	if (!IsActive() || !World)
 	{
 		return EAudioComponentPlayState::Stopped;
 	}
@@ -862,7 +862,7 @@ void UAudioComponent::UpdateSpriteTexture()
 #if WITH_EDITOR
 void UAudioComponent::PostEditChangeProperty(FPropertyChangedEvent& PropertyChangedEvent)
 {
-	if (bIsActive)
+	if (IsActive())
 	{
 		// If this is an auto destroy component we need to prevent it from being auto-destroyed since we're really just restarting it
 		const bool bWasAutoDestroy = bAutoDestroy;
@@ -938,7 +938,7 @@ void UAudioComponent::Activate(bool bReset)
 	if (bReset || ShouldActivate() == true)
 	{
 		Play();
-		if (bIsActive)
+		if (IsActive())
 		{
 			OnComponentActivated.Broadcast(this, bReset);
 		}
@@ -951,7 +951,7 @@ void UAudioComponent::Deactivate()
 	{
 		Stop();
 
-		if (!bIsActive)
+		if (!IsActive())
 		{
 			OnComponentDeactivated.Broadcast(this);
 		}
@@ -984,7 +984,7 @@ void UAudioComponent::SetFloatParameter( const FName InName, const float InFloat
 		}
 
 		// If we're active we need to push this value to the ActiveSound
-		if (bIsActive)
+		if (IsActive())
 		{
 			if (FAudioDevice* AudioDevice = GetAudioDevice())
 			{
@@ -1029,7 +1029,7 @@ void UAudioComponent::SetWaveParameter( const FName InName, USoundWave* InWave )
 		}
 
 		// If we're active we need to push this value to the ActiveSound
-		if (bIsActive)
+		if (IsActive())
 		{
 			if (FAudioDevice* AudioDevice = GetAudioDevice())
 			{
@@ -1074,7 +1074,7 @@ void UAudioComponent::SetBoolParameter( const FName InName, const bool InBool )
 		}
 
 		// If we're active we need to push this value to the ActiveSound
-		if (bIsActive)
+		if (IsActive())
 		{
 			if (FAudioDevice* AudioDevice = GetAudioDevice())
 			{
@@ -1120,7 +1120,7 @@ void UAudioComponent::SetIntParameter( const FName InName, const int32 InInt )
 		}
 
 		// If we're active we need to push this value to the ActiveSound
-		if (bIsActive)
+		if (IsActive())
 		{
 			if (FAudioDevice* AudioDevice = GetAudioDevice())
 			{
@@ -1162,7 +1162,7 @@ void UAudioComponent::SetSoundParameter(const FAudioComponentParam& Param)
 			const int32 NewParamIndex = InstanceParameters.Add(Param);
 		}
 
-		if (bIsActive)
+		if (IsActive())
 		{
 			if (FAudioDevice* AudioDevice = GetAudioDevice())
 			{
@@ -1187,7 +1187,7 @@ void UAudioComponent::SetVolumeMultiplier(const float NewVolumeMultiplier)
 	VolumeMultiplier = NewVolumeMultiplier;
 	VolumeModulationMin = VolumeModulationMax = 1.f;
 
-	if (bIsActive)
+	if (IsActive())
 	{
 		if (FAudioDevice* AudioDevice = GetAudioDevice())
 		{
@@ -1211,7 +1211,7 @@ void UAudioComponent::SetPitchMultiplier(const float NewPitchMultiplier)
 	PitchMultiplier = NewPitchMultiplier;
 	PitchModulationMin = PitchModulationMax = 1.f;
 
-	if (bIsActive)
+	if (IsActive())
 	{
 		if (FAudioDevice* AudioDevice = GetAudioDevice())
 		{
@@ -1234,7 +1234,7 @@ void UAudioComponent::SetUISound(const bool bInIsUISound)
 {
 	bIsUISound = bInIsUISound;
 
-	if (bIsActive)
+	if (IsActive())
 	{
 		if (FAudioDevice* AudioDevice = GetAudioDevice())
 		{
@@ -1258,7 +1258,7 @@ void UAudioComponent::AdjustAttenuation(const FSoundAttenuationSettings& InAtten
 	bOverrideAttenuation = true;
 	AttenuationOverrides = InAttenuationSettings;
 
-	if (bIsActive)
+	if (IsActive())
 	{
 		if (FAudioDevice* AudioDevice = GetAudioDevice())
 		{
diff --git a/Engine/Source/Runtime/Engine/Private/Components/BillboardComponent.cpp b/Engine/Source/Runtime/Engine/Private/Components/BillboardComponent.cpp
--- a/Engine/Source/Runtime/Engine/Private/Components/BillboardComponent.cpp
+++ b/Engine/Source/Runtime/Engine/Private/Components/BillboardComponent.cpp
@@ -284,7 +284,7 @@ UBillboardComponent::UBillboardComponent(const FObjectInitializer& ObjectInitial
 #endif
 
 	SetCollisionProfileName(UCollisionProfile::NoCollision_ProfileName);
-	bAbsoluteScale = true;
+	SetUsingAbsoluteScale(true);
 
 	bIsScreenSizeScaled = false;
 	ScreenSize = BillboardConstants::DefaultScreenSize;
diff --git a/Engine/Source/Runtime/Engine/Private/Components/BrushComponent.cpp b/Engine/Source/Runtime/Engine/Private/Components/BrushComponent.cpp
--- a/Engine/Source/Runtime/Engine/Private/Components/BrushComponent.cpp
+++ b/Engine/Source/Runtime/Engine/Private/Components/BrushComponent.cpp
@@ -687,7 +687,9 @@ void UBrushComponent::RequestUpdateBrushCollision()
 {
 	if (BrushBodySetup)
 	{
-		const bool bIsMirrored = (RelativeScale3D.X * RelativeScale3D.Y * RelativeScale3D.Z) < 0.0f;
+
+		const FVector LocalRelativeScale3D = GetRelativeScale3D();
+		const bool bIsMirrored = (LocalRelativeScale3D.X * LocalRelativeScale3D.Y * LocalRelativeScale3D.Z) < 0.0f;
 		if ((BrushBodySetup->bGenerateNonMirroredCollision && bIsMirrored) || (BrushBodySetup->bGenerateMirroredCollision && !bIsMirrored))
 		{
 			// Brushes only maintain one convex mesh as they can't be transformed at runtime.
@@ -753,7 +755,8 @@ bool UBrushComponent::HasInvertedPolys() const
 	// Determine if a brush looks as if it has had its sense inverted
 	// (due to the old behavior of inverting the poly winding and normal when performing a Mirror operation).
 
-	const bool bIsMirrored = (RelativeScale3D.X * RelativeScale3D.Y * RelativeScale3D.Z < 0.0f);
+	const FVector LocalRelativeScale3D = GetRelativeScale3D();
+	const bool bIsMirrored = (LocalRelativeScale3D.X * LocalRelativeScale3D.Y * LocalRelativeScale3D.Z < 0.0f);
 	// Only attempt to fix up brushes with negative scale
 	if (bIsMirrored)
 	{
