Include the number of blueprints compiled when displaying detailed compile results

diff --git a/Engine/Source/Editor/Kismet/Private/BlueprintCompilationManager.cpp b/Engine/Source/Editor/Kismet/Private/BlueprintCompilationManager.cpp
--- a/Engine/Source/Editor/Kismet/Private/BlueprintCompilationManager.cpp
+++ b/Engine/Source/Editor/Kismet/Private/BlueprintCompilationManager.cpp
@@ -241,6 +241,16 @@ void FBlueprintCompilationManagerImpl::CompileSynchronouslyImpl(const FBPCompile
 	FlushCompilationQueueImpl(bSuppressBroadcastCompiled, &CompiledBlueprints, &SkeletonCompiledBlueprints, nullptr);
 	FlushReinstancingQueueImpl();
 	
+	if( Request.ClientResultsLog && Request.ClientResultsLog->bLogDetailedResults)
+	{
+		Request.ClientResultsLog->Note(
+			*FText::Format(
+				LOCTEXT("TotalCompiledBPs", "Compiled {0} total blueprints"),
+				CompiledBlueprints.Num()
+			).ToString()
+		);
+	}
+
 	if (FBlueprintEditorUtils::IsLevelScriptBlueprint(Request.BPToCompile) && !bRegenerateSkeletonOnly)
 	{
 		// When the Blueprint is recompiled, then update the bound events for level scripting
