Handle the edge cases of the (poor idea) of destroying the root component in the user construction script

diff --git a/Engine/Source/Editor/Kismet/Private/SSCSEditor.cpp b/Engine/Source/Editor/Kismet/Private/SSCSEditor.cpp
--- a/Engine/Source/Editor/Kismet/Private/SSCSEditor.cpp
+++ b/Engine/Source/Editor/Kismet/Private/SSCSEditor.cpp
@@ -5223,9 +5223,9 @@ UActorComponent* SSCSEditor::AddNewComponent( UClass* NewComponentClass, UObject
 		if (ComponentTemplate)
 		{
 			// Create a duplicate of the provided template
-			UActorComponent* NewInstanceComponent = FComponentEditorUtils::DuplicateComponent(ComponentTemplate);
-			FSCSEditorTreeNodePtrType ParentNodePtr = FindParentForNewComponent(NewInstanceComponent);
-			NewComponent = AddNewNodeForInstancedComponent(MoveTemp(AddTransaction), NewInstanceComponent, ParentNodePtr, nullptr, bSetFocusToNewItem);
+			NewComponent = FComponentEditorUtils::DuplicateComponent(ComponentTemplate);
+			FSCSEditorTreeNodePtrType ParentNodePtr = FindParentForNewComponent(NewComponent);
+			AddNewNodeForInstancedComponent(MoveTemp(AddTransaction), NewComponent, ParentNodePtr, nullptr, bSetFocusToNewItem);
 		}
 		else if (AActor* ActorInstance = GetActorContext())
 		{
@@ -5295,7 +5295,12 @@ UActorComponent* SSCSEditor::AddNewComponent( UClass* NewComponentClass, UObject
 			// Rerun construction scripts
 			ActorInstance->RerunConstructionScripts();
 
-			NewComponent = AddNewNodeForInstancedComponent(MoveTemp(AddTransaction), NewInstanceComponent, ParentNodePtr, Asset, bSetFocusToNewItem);
+			// If the running the construction script destroyed the new node, don't create an entry for it
+			if (!NewInstanceComponent->IsPendingKill())
+			{
+				AddNewNodeForInstancedComponent(MoveTemp(AddTransaction), NewInstanceComponent, ParentNodePtr, Asset, bSetFocusToNewItem);
+				NewComponent = NewInstanceComponent;
+			}
 		}
 	}
 
@@ -5411,7 +5416,7 @@ FSCSEditorTreeNodePtrType SSCSEditor::FindParentForNewComponent(UActorComponent*
 	}
 
 	return TargetParentNode;
-	}
+}
 
 FSCSEditorTreeNodePtrType SSCSEditor::FindParentForNewNode(USCS_Node* NewNode) const
 {
@@ -5463,7 +5468,7 @@ UActorComponent* SSCSEditor::AddNewNode(TUniquePtr<FScopedTransaction> InOngoing
 	return NewNode->ComponentTemplate;
 }
 
-UActorComponent* SSCSEditor::AddNewNodeForInstancedComponent(TUniquePtr<FScopedTransaction> InOngoingCreateTransaction, UActorComponent* NewInstanceComponent, FSCSEditorTreeNodePtrType InParentNodePtr, UObject* Asset, bool bSetFocusToNewItem)
+void SSCSEditor::AddNewNodeForInstancedComponent(TUniquePtr<FScopedTransaction> InOngoingCreateTransaction, UActorComponent* NewInstanceComponent, FSCSEditorTreeNodePtrType InParentNodePtr, UObject* Asset, bool bSetFocusToNewItem)
 {
 	check(NewInstanceComponent != nullptr);
 
@@ -5480,8 +5485,6 @@ UActorComponent* SSCSEditor::AddNewNodeForInstancedComponent(TUniquePtr<FScopedT
 	}
 
 	UpdateTree(false);
-
-	return NewInstanceComponent;
 }
 
 bool SSCSEditor::IsComponentSelected(const UPrimitiveComponent* PrimComponent) const
diff --git a/Engine/Source/Editor/Kismet/Public/SSCSEditor.h b/Engine/Source/Editor/Kismet/Public/SSCSEditor.h
--- a/Engine/Source/Editor/Kismet/Public/SSCSEditor.h
+++ b/Engine/Source/Editor/Kismet/Public/SSCSEditor.h
@@ -851,9 +851,8 @@ public:
 		@param NewInstanceComponent	(In) The component being added to the actor instance
 		@param InParentNodePtr (In) The node this component will be added to
 		@param Asset (In) Optional asset to assign to the component
-		@param bSetFocusToNewItem (In) Select the new item and activate the inline rename widget (default is true)
-		@return The reference of the newly created ActorComponent */
-	UActorComponent* AddNewNodeForInstancedComponent(TUniquePtr<FScopedTransaction> OngoingCreateTransaction, UActorComponent* NewInstanceComponent, FSCSEditorTreeNodePtrType InParentNodePtr, UObject* Asset, bool bSetFocusToNewItem = true);
+		@param bSetFocusToNewItem (In) Select the new item and activate the inline rename widget (default is true) */
+	void AddNewNodeForInstancedComponent(TUniquePtr<FScopedTransaction> OngoingCreateTransaction, UActorComponent* NewInstanceComponent, FSCSEditorTreeNodePtrType InParentNodePtr, UObject* Asset, bool bSetFocusToNewItem = true);
 	
 	/** Returns true if the specified component is currently selected */
 	bool IsComponentSelected(const UPrimitiveComponent* PrimComponent) const;
diff --git a/Engine/Source/Runtime/Engine/Private/SimpleConstructionScript.cpp b/Engine/Source/Runtime/Engine/Private/SimpleConstructionScript.cpp
--- a/Engine/Source/Runtime/Engine/Private/SimpleConstructionScript.cpp
+++ b/Engine/Source/Runtime/Engine/Private/SimpleConstructionScript.cpp
@@ -635,7 +635,8 @@ void USimpleConstructionScript::ExecuteScriptOnActor(AActor* Actor, const TInlin
 
 		for (USCS_Node* RootNode : RootNodes)
 		{
-			if(RootNode != nullptr)
+			// If the node is a default scene root and the actor already has a root component, skip it
+			if (RootNode && ((RootNode != DefaultSceneRootNode) || (RootComponent == nullptr)))
 			{
 				// If the root node specifies that it has a parent
 				USceneComponent* ParentComponent = nullptr;
