Rework dependent blueprint compilation to avoid false positives. This speeds up blueprint compilation for blueprints with large #'s of dependents

diff --git a/Engine/Source/Editor/Kismet/Private/BlueprintCompilationManager.cpp b/Engine/Source/Editor/Kismet/Private/BlueprintCompilationManager.cpp
--- a/Engine/Source/Editor/Kismet/Private/BlueprintCompilationManager.cpp
+++ b/Engine/Source/Editor/Kismet/Private/BlueprintCompilationManager.cpp
@@ -745,9 +745,9 @@ void FBlueprintCompilationManagerImpl::FlushCompilationQueueImpl(bool bSuppressB
 
 			// if any function signatures have changed in this skeleton class we will need to recompile all dependencies, but if not
 			// then we can avoid dependency recompilation:
-			TSet<UBlueprint*> BlueprintsWithSignatureChanges;
 			const UBlueprintEditorProjectSettings* EditorProjectSettings = GetDefault<UBlueprintEditorProjectSettings>();
 			bool bSkipUnneededDependencyCompilation = !EditorProjectSettings->bForceAllDependenciesToRecompile;
+			TSet<UObject*> OldFunctionsWithSignatureChanges;
 
 			for (FCompilerData& CompilerData : CurrentlyCompilingBPs)
 			{
@@ -770,21 +770,25 @@ void FBlueprintCompilationManagerImpl::FlushCompilationQueueImpl(bool bSuppressB
 						{
 							for (TFieldIterator<UFunction> FuncIt(AuthoritativeClass, EFieldIteratorFlags::ExcludeSuper); FuncIt; ++FuncIt)
 							{
+								UFunction* OldFunction = *FuncIt;
+
+								if(!OldFunction->HasAnyFunctionFlags(EFunctionFlags::FUNC_BlueprintCallable))
+								{
+									continue;
+								}
+
 								// We assume that if the func is FUNC_BlueprintCallable that it will be present in the Skeleton class.
 								// If it is not in the skeleton class we will always think that this blueprints public interface has 
 								// changed. Not a huge deal, but will mean we recompile dependencies more often than necessary.
-								if(FuncIt->HasAnyFunctionFlags(EFunctionFlags::FUNC_BlueprintCallable))
+								UFunction* NewFunction = BP->SkeletonGeneratedClass->FindFunctionByName((OldFunction)->GetFName());
+								if(	NewFunction == nullptr || 
+									!NewFunction->IsSignatureCompatibleWith(OldFunction) || 
+									// If a function changes its net flags, callers may now need to do a full EX_FinalFunction/EX_VirtualFunction 
+									// instead of a EX_LocalFinalFunction/EX_LocalVirtualFunction:
+									NewFunction->HasAnyFunctionFlags(FUNC_NetFuncFlags) != OldFunction->HasAnyFunctionFlags(FUNC_NetFuncFlags))
 								{
-									UFunction* NewFunction = BP->SkeletonGeneratedClass->FindFunctionByName((*FuncIt)->GetFName());
-									if(	NewFunction == nullptr || 
-										!NewFunction->IsSignatureCompatibleWith(*FuncIt) || 
-										// If a function changes its net flags, callers may now need to do a full EX_FinalFunction/EX_VirtualFunction 
-										// instead of a EX_LocalFinalFunction/EX_LocalVirtualFunction:
-										NewFunction->HasAnyFunctionFlags(FUNC_NetFuncFlags) != FuncIt->HasAnyFunctionFlags(FUNC_NetFuncFlags))
-									{
-										BlueprintsWithSignatureChanges.Add(BP);
-										break;
-									}
+									OldFunctionsWithSignatureChanges.Add(OldFunction);
+									break;
 								}
 							}
 						}
@@ -847,61 +851,52 @@ void FBlueprintCompilationManagerImpl::FlushCompilationQueueImpl(bool bSuppressB
 
 			// Skip further compilation for blueprints that are being bytecode compiled as a dependency of something that has
 			// not had a change in its function parameters:
-			auto DependenciesAreCompiled = [&BlueprintsWithSignatureChanges](FCompilerData& Data)
+			if(bSkipUnneededDependencyCompilation)
 			{
-				if(Data.ShouldSkipIfDependenciesAreUnchanged())
+				const auto HasNoReferencesToChangedFunctions = [&OldFunctionsWithSignatureChanges](FCompilerData& Data)
 				{
-					DECLARE_SCOPE_HIERARCHICAL_COUNTER(SkipIfDependenciesAreUnchanged)
+					if(!Data.ShouldSkipIfDependenciesAreUnchanged())
+					{
+						return false;
+					}
 
-					// if our parent is still being compiled, then we still need to be compiled:
-					UClass* Iter = Data.BP->ParentClass;
-					while(Iter)
+					// look for references to a function with a signature change
+					// in the old class, if it has none, we can skip bytecode recompile:
+					bool bHasNoReferencesToChangedFunctions = true;
+					UBlueprintGeneratedClass* BPGC = Cast<UBlueprintGeneratedClass>(Data.BP->GeneratedClass);
+					if(BPGC)
 					{
-						if(UBlueprint* BP = Cast<UBlueprint>(Iter->ClassGeneratedBy))
+						for(UFunction* CalledFunction : BPGC->CalledFunctions)
 						{
-							if(BP->bBeingCompiled)
+							if(OldFunctionsWithSignatureChanges.Contains(CalledFunction))
 							{
-								return false;
+								bHasNoReferencesToChangedFunctions = false;
+								break;
 							}
 						}
-						Iter = Iter->GetSuperClass();
 					}
 
-					// otherwise if we're dependent on a blueprint that had a function signature change, we still need to be compiled:
-					ensure(Data.BP->bCachedDependenciesUpToDate);
-					ensure(Data.BP->CachedDependencies.Num() > 0); // why are we bytecode compiling a blueprint with no dependencies?
-					for(const TWeakObjectPtr<UBlueprint>& Dependency : Data.BP->CachedDependencies)
+					if(bHasNoReferencesToChangedFunctions)
 					{
-						if (UBlueprint* DependencyBP = Dependency.Get())
+						// This BP is not actually going to be compiled, clean it up:
+						Data.BP->bBeingCompiled = false;
+						Data.BP->CurrentMessageLog = nullptr;
+						if(UPackage* Package = Data.BP->GetOutermost())
 						{
-							if(DependencyBP->bBeingCompiled && BlueprintsWithSignatureChanges.Contains(DependencyBP))
-							{
-								return false;
-							}
+							Package->SetDirtyFlag(Data.bPackageWasDirty);
 						}
+						if(Data.ResultsLog)
+						{
+							Data.ResultsLog->EndEvent();
+						}
+						Data.BP->bQueuedForCompilation = false;
 					}
-					
-					Data.BP->bBeingCompiled = false;
-					Data.BP->CurrentMessageLog = nullptr;
-					if(UPackage* Package = Data.BP->GetOutermost())
-					{
-						Package->SetDirtyFlag(Data.bPackageWasDirty);
-					}
-					if(Data.ResultsLog)
-					{
-						Data.ResultsLog->EndEvent();
-					}
-					Data.BP->bQueuedForCompilation = false;
-					return true;
-				}
-			
-				return false;
-			};
 
-			if(bSkipUnneededDependencyCompilation)
-			{
+					return bHasNoReferencesToChangedFunctions;
+				};
+
 				// Order very much matters, but we could RemoveAllSwap and re-sort:
-				CurrentlyCompilingBPs.RemoveAll(DependenciesAreCompiled);
+				CurrentlyCompilingBPs.RemoveAll(HasNoReferencesToChangedFunctions);
 			}
 		}
 
diff --git a/Engine/Source/Editor/KismetCompiler/Private/KismetCompilerVMBackend.cpp b/Engine/Source/Editor/KismetCompiler/Private/KismetCompilerVMBackend.cpp
--- a/Engine/Source/Editor/KismetCompiler/Private/KismetCompilerVMBackend.cpp
+++ b/Engine/Source/Editor/KismetCompiler/Private/KismetCompilerVMBackend.cpp
@@ -1124,6 +1124,8 @@ public:
 		UFunction* FunctionToCall = Statement.FunctionToCall;
 		check(FunctionToCall);
 
+		ClassBeingBuilt->CalledFunctions.Emplace(FunctionToCall);
+
 		if (FunctionToCall->HasAllFunctionFlags(FUNC_Native))
 		{
 			// Array output parameters are cleared, in case the native function doesn't clear them before filling.
@@ -1295,6 +1297,8 @@ public:
 		check(NULL != Statement.FunctionContext);
 		check(FunctionToCall->HasAnyFunctionFlags(FUNC_Delegate));
 
+		ClassBeingBuilt->CalledFunctions.Emplace(FunctionToCall);
+
 		// The function to call doesn't have a native index
 		Writer << EX_CallMulticastDelegate;
 		Writer << FunctionToCall;
@@ -2129,6 +2133,11 @@ void FKismetCompilerVMBackend::GenerateCodeFromClass(UClass* SourceClass, TIndir
 			ConstructFunction(Function, bIsUbergraph, bGenerateStubsOnly);
 		}
 	}
+
+	// Remove duplicates from CalledFunctions:
+	UBlueprintGeneratedClass* ClassBeingBuilt = CastChecked<UBlueprintGeneratedClass>(SourceClass);
+	TSet<UFunction*> Unique(ClassBeingBuilt->CalledFunctions);
+	ClassBeingBuilt->CalledFunctions = Unique.Array();
 }
 
 void FKismetCompilerVMBackend::ConstructFunction(FKismetFunctionContext& FunctionContext, bool bIsUbergraph, bool bGenerateStubOnly)
diff --git a/Engine/Source/Runtime/Engine/Classes/Engine/BlueprintGeneratedClass.h b/Engine/Source/Runtime/Engine/Classes/Engine/BlueprintGeneratedClass.h
--- a/Engine/Source/Runtime/Engine/Classes/Engine/BlueprintGeneratedClass.h
+++ b/Engine/Source/Runtime/Engine/Classes/Engine/BlueprintGeneratedClass.h
@@ -699,6 +699,9 @@ public:
 	/** Property guid map */
 	UPROPERTY()
 	TMap<FName,FGuid> PropertyGuids;
+
+	UPROPERTY(Transient)
+	TArray<UFunction*> CalledFunctions;
 #endif //WITH_EDITORONLY_DATA
 
 	// Mapping of changed properties & data to apply when instancing components in a cooked build (one entry per named AddComponent node template for fast lookup at runtime).
diff --git a/Engine/Source/Runtime/Engine/Private/BlueprintGeneratedClass.cpp b/Engine/Source/Runtime/Engine/Private/BlueprintGeneratedClass.cpp
--- a/Engine/Source/Runtime/Engine/Private/BlueprintGeneratedClass.cpp
+++ b/Engine/Source/Runtime/Engine/Private/BlueprintGeneratedClass.cpp
@@ -1592,6 +1592,7 @@ void UBlueprintGeneratedClass::PurgeClass(bool bRecompilingOnLoad)
 #if UE_BLUEPRINT_EVENTGRAPH_FASTCALLS
 	FastCallPairs_DEPRECATED.Empty();
 #endif
+	CalledFunctions.Empty();
 #endif //WITH_EDITOR
 }
 
