CompileAllBlueprints commandlet now reinstances, ensuring that it does not later crash during GC

diff --git a/Engine/Source/Editor/Kismet/Private/BlueprintCompilationManager.cpp b/Engine/Source/Editor/Kismet/Private/BlueprintCompilationManager.cpp
--- a/Engine/Source/Editor/Kismet/Private/BlueprintCompilationManager.cpp
+++ b/Engine/Source/Editor/Kismet/Private/BlueprintCompilationManager.cpp
@@ -1026,6 +1026,11 @@ void FBlueprintCompilationManagerImpl::FlushCompilationQueueImpl(bool bSuppressB
 					)
 				);
 
+				if(CompilerData.Compiler.IsValid())
+				{
+					CompilerData.Compiler->OldClass = Cast<UBlueprintGeneratedClass>(CompilerData.Reinstancer->DuplicatedClass);
+				}
+
 				if(BP->GeneratedClass)
 				{
 					BP->GeneratedClass->ClassFlags |= CLASS_LayoutChanging;
@@ -1707,14 +1712,29 @@ void FBlueprintCompilationManagerImpl::ReinstanceBatch(TArray<FReinstancingJob>&
 			continue;
 		}
 
-		InOutOldToNewClassMap.Add(OldClass, ReinstancingJob.OldToNew.Value);
+		UClass* NewClass = ReinstancingJob.OldToNew.Value;
+
+		InOutOldToNewClassMap.Add(OldClass, NewClass);
 
 		if(!HasChildren(OldClass))
 		{
 			continue;
 		}
 
-		bool bParentLayoutChanged = !FStructUtils::TheSameLayout(OldClass, ReinstancingJob.OldToNew.Value);
+		bool bParentLayoutChanged = !FStructUtils::TheSameLayout(OldClass, NewClass);
+
+		if(!bParentLayoutChanged)
+		{
+			// make sure uber graph didn't change, if present:
+			UBlueprintGeneratedClass* OldParent = Cast<UBlueprintGeneratedClass>(OldClass);
+			UBlueprintGeneratedClass* NewBPGC = Cast<UBlueprintGeneratedClass>(NewClass);
+
+			if(OldParent && NewBPGC && OldParent->UberGraphFunction)
+			{
+				bParentLayoutChanged = !FStructUtils::TheSameLayout(OldParent->UberGraphFunction, NewBPGC->UberGraphFunction);
+			}
+		}
+
 		if(bParentLayoutChanged)
 		{
 			// we need *all* derived types:
diff --git a/Engine/Source/Editor/KismetCompiler/Private/KismetCompiler.cpp b/Engine/Source/Editor/KismetCompiler/Private/KismetCompiler.cpp
--- a/Engine/Source/Editor/KismetCompiler/Private/KismetCompiler.cpp
+++ b/Engine/Source/Editor/KismetCompiler/Private/KismetCompiler.cpp
@@ -130,6 +130,7 @@ FKismetCompilerContext::FKismetCompilerContext(UBlueprint* SourceSketch, FCompil
 	, CompileOptions(InCompilerOptions)
 	, Blueprint(SourceSketch)
 	, NewClass(NULL)
+	, OldClass(nullptr)
 	, ConsolidatedEventGraph(NULL)
 	, UbergraphContext(NULL)
 	, bIsFullCompile(false)
@@ -2210,6 +2211,17 @@ void FKismetCompilerContext::FinishCompilingFunction(FKismetFunctionContext& Con
 	if( NewClass->UberGraphFunction == Context.Function )
 	{
 		NewClass->UberGraphFunctionKey = IncrementUberGraphSerialNumber();
+
+		// if the old class uber graph function matches, just reuse that ID, this check means
+		// that if child types aren't reinstanced we can still validate their uber graph:
+		if(NewClass->UberGraphFunction && OldClass)
+		{
+			bool bSameLayout = FStructUtils::TheSameLayout(OldClass->UberGraphFunction, NewClass->UberGraphFunction);
+			if(bSameLayout)
+			{
+				NewClass->UberGraphFunctionKey = OldClass->UberGraphFunctionKey;
+			}
+		}
 	}
 #endif//VALIDATE_UBER_GRAPH_PERSISTENT_FRAME
 }
diff --git a/Engine/Source/Editor/KismetCompiler/Public/KismetCompiler.h b/Engine/Source/Editor/KismetCompiler/Public/KismetCompiler.h
--- a/Engine/Source/Editor/KismetCompiler/Public/KismetCompiler.h
+++ b/Engine/Source/Editor/KismetCompiler/Public/KismetCompiler.h
@@ -93,6 +93,7 @@ protected:
 public:
 	UBlueprint* Blueprint;
 	UBlueprintGeneratedClass* NewClass;
+	UBlueprintGeneratedClass* OldClass;
 
 	// The ubergraph; valid from roughly the start of CreateAndProcessEventGraph
 	UEdGraph* ConsolidatedEventGraph;
diff --git a/Engine/Source/Editor/UnrealEd/Private/Commandlets/CompileAllBlueprintsCommandlet.cpp b/Engine/Source/Editor/UnrealEd/Private/Commandlets/CompileAllBlueprintsCommandlet.cpp
--- a/Engine/Source/Editor/UnrealEd/Private/Commandlets/CompileAllBlueprintsCommandlet.cpp
+++ b/Engine/Source/Editor/UnrealEd/Private/Commandlets/CompileAllBlueprintsCommandlet.cpp
@@ -302,6 +302,8 @@ void UCompileAllBlueprintsCommandlet::CompileBlueprint(UBlueprint* Blueprint)
 		{
 			MessageLog.bAnnotateMentionedNodes = true;
 		}
+		MessageLog.SetSourcePath(Blueprint->GetPathName());
+		MessageLog.BeginEvent(TEXT("Compile"));
 	
 		FKismetCompilerOptions CompileOptions;
 		//Defaults to use for CompileOptions
@@ -315,8 +317,10 @@ void UCompileAllBlueprintsCommandlet::CompileBlueprint(UBlueprint* Blueprint)
 		{
 			CompileOptions.CompileType = EKismetCompileType::SkeletonOnly;
 		}
-		
-		KismetBlueprintCompilerModule->CompileBlueprint(Blueprint, CompileOptions, MessageLog);
+
+		FKismetEditorUtilities::CompileBlueprint(Blueprint, EBlueprintCompileOptions::SkipGarbageCollection, &MessageLog);
+
+		MessageLog.EndEvent();
 
 		if ((MessageLog.NumErrors + MessageLog.NumWarnings) > 0)
 		{
