Don't allow functions to be converted in to events in function libraries
#rnx

diff --git a/Engine/Source/Editor/Kismet/Private/BlueprintEditor.cpp b/Engine/Source/Editor/Kismet/Private/BlueprintEditor.cpp
--- a/Engine/Source/Editor/Kismet/Private/BlueprintEditor.cpp
+++ b/Engine/Source/Editor/Kismet/Private/BlueprintEditor.cpp
@@ -5332,14 +5332,19 @@ void FBlueprintEditor::MoveNodesToAveragePos(TSet<UEdGraphNode*>& AverageNodes,
 
 bool FBlueprintEditor::CanConvertFunctionToEvent() const
 {
-	UEdGraphNode* const SelectedNode = GetSingleSelectedNode();
-	if (SelectedNode)
+	if (UBlueprint* Blueprint = GetBlueprintObj())
 	{
-		if (UK2Node_FunctionEntry* const SelectedCallFunctionNode = Cast<UK2Node_FunctionEntry>(SelectedNode))
+		if (BPTYPE_FunctionLibrary != Blueprint->BlueprintType)
 		{
-			if (SelectedCallFunctionNode->FindSignatureFunction())
+			if (UEdGraphNode* const SelectedNode = GetSingleSelectedNode())
 			{
-				return true;
+				if (UK2Node_FunctionEntry* const SelectedCallFunctionNode = Cast<UK2Node_FunctionEntry>(SelectedNode))
+				{
+					if (SelectedCallFunctionNode->FindSignatureFunction())
+					{
+						return true;
+					}
+				}
 			}
 		}
 	}
