Move subgraphs if a composite node is moved via ConvertFunctionToEvent/ConvertEventToFunction

diff --git a/Engine/Source/Editor/Kismet/Private/BlueprintEditor.cpp b/Engine/Source/Editor/Kismet/Private/BlueprintEditor.cpp
--- a/Engine/Source/Editor/Kismet/Private/BlueprintEditor.cpp
+++ b/Engine/Source/Editor/Kismet/Private/BlueprintEditor.cpp
@@ -7853,6 +7853,12 @@ void FBlueprintEditor::MoveNodesToGraph(TArray<UEdGraphNode*>& SourceNodes, UEdG
 		
 		DestinationGraph->AddNode(Node, /* bFromUI */ false, /* bSelectNewNode */ false);
 		
+		if(UK2Node_Composite* Composite = Cast<UK2Node_Composite>(Node))
+		{
+			OriginalGraph->SubGraphs.Remove(Composite->BoundGraph);
+			DestinationGraph->SubGraphs.Add(Composite->BoundGraph);
+		}
+
 		// Want to test exactly against tunnel, we shouldn't collapse embedded collapsed
 		// nodes or macros, only the tunnels in/out of the collapsed graph
 		if (Node->GetClass() == UK2Node_Tunnel::StaticClass())
