OnExpandNodes now clears the selection set, preventing stale selection entries

diff --git a/Engine/Source/Editor/Kismet/Private/BlueprintEditor.cpp b/Engine/Source/Editor/Kismet/Private/BlueprintEditor.cpp
--- a/Engine/Source/Editor/Kismet/Private/BlueprintEditor.cpp
+++ b/Engine/Source/Editor/Kismet/Private/BlueprintEditor.cpp
@@ -5213,6 +5213,12 @@ void FBlueprintEditor::OnExpandNodes()
 
 	// Expand selected nodes into the focused graph context.
 	const FGraphPanelSelectionSet SelectedNodes = GetSelectedNodes();
+	
+	if(FocusedGraphEd)
+	{
+		FocusedGraphEd->ClearSelectionSet();
+	}
+
 	for (FGraphPanelSelectionSet::TConstIterator NodeIt(SelectedNodes); NodeIt; ++NodeIt)
 	{
 		ExpandedNodes.Empty();
