Fix crash when deleting function graphs, collapsed graphs, etc

diff --git a/Engine/Source/Editor/Kismet/Private/BlueprintDetailsCustomization.cpp b/Engine/Source/Editor/Kismet/Private/BlueprintDetailsCustomization.cpp
--- a/Engine/Source/Editor/Kismet/Private/BlueprintDetailsCustomization.cpp
+++ b/Engine/Source/Editor/Kismet/Private/BlueprintDetailsCustomization.cpp
@@ -2891,7 +2891,7 @@ void FBlueprintGraphArgumentLayout::GenerateChildContent( IDetailChildrenBuilder
 		}
 
 		bool bMacroGraph = false;
-		if (TargetNode)
+		if (TargetNode && TargetNode->HasValidBlueprint())
 		{
 			if (const UBlueprint* Blueprint = TargetNode->GetBlueprint())
 			{
