Do Delegate and return node sanitization after we have validated node compatability in ConvertEventToFunction

diff --git a/Engine/Source/Editor/Kismet/Private/BlueprintEditor.cpp b/Engine/Source/Editor/Kismet/Private/BlueprintEditor.cpp
--- a/Engine/Source/Editor/Kismet/Private/BlueprintEditor.cpp
+++ b/Engine/Source/Editor/Kismet/Private/BlueprintEditor.cpp
@@ -5649,14 +5649,6 @@ void FBlueprintEditor::ConvertEventToFunction(UK2Node_Event* SelectedEventNode)
 
 		const UEdGraphSchema_K2* Schema = GetDefault<UEdGraphSchema_K2>();
 
-		// Break any connections to the delegate pin, because that won't exist 
-		{
-			UEdGraphPin* DelegatePin = SelectedEventNode->FindPin(UK2Node_Event::DelegateOutputName);
-			if (DelegatePin)
-			{
-				DelegatePin->BreakAllPinLinks();
-			}
-		}
 		// Refresh this node before getting the pin names to ensure that we get the most up to date ones
 		Schema->ReconstructNode(*SelectedEventNode);
 
@@ -5694,17 +5686,6 @@ void FBlueprintEditor::ConvertEventToFunction(UK2Node_Event* SelectedEventNode)
 			{
 				FBlueprintEditorUtils::AddFunctionGraph<UFunction>(NodeBP, NewGraph, /*bIsUserCreated=*/ true, FunctionSig);
 			}
-
-			// Remove any return nodes that may have been added @see UE-78785
-			TArray<UK2Node_FunctionResult*> FunctionReturnNodes;
-			NewGraph->GetNodesOfClass(FunctionReturnNodes);
-			for (UK2Node_FunctionResult* Node : FunctionReturnNodes)
-			{
-				if (Node)
-				{
-					Node->DestroyNode();
-				}
-			}
 			
 			// Get the new entry node
 			TArray<UK2Node_FunctionEntry*> EntryNodes;
@@ -5754,7 +5735,27 @@ void FBlueprintEditor::ConvertEventToFunction(UK2Node_Event* SelectedEventNode)
 				NewGraph->MarkPendingKill();
 				return;
 			}
+
+			// Remove any return nodes that may have been added @see UE-78785
+			TArray<UK2Node_FunctionResult*> FunctionReturnNodes;
+			NewGraph->GetNodesOfClass(FunctionReturnNodes);
+			for (UK2Node_FunctionResult* Node : FunctionReturnNodes)
+			{
+				if (Node)
+				{
+					Node->DestroyNode();
+				}
+			}
 			
+			// Break any connections to the delegate pin, because that won't exist 
+			{
+				UEdGraphPin* DelegatePin = SelectedEventNode->FindPin(UK2Node_Event::DelegateOutputName);
+				if (DelegatePin)
+				{
+					DelegatePin->BreakAllPinLinks();
+				}
+			}
+
 			// Move the nodes to the new graph
 			UEdGraphNode* Entry = nullptr;
 			UEdGraphNode* Result = nullptr;
@@ -5989,6 +5990,11 @@ void FBlueprintEditor::DeleteSelectedNodes()
 
 	SetUISelectionState(NAME_None);
 
+	if(FocusedGraphEd)
+	{
+		FocusedGraphEd->ClearSelectionSet();
+	}
+
 	// this closes all the document that is outered by this node
 	// this is used by AnimBP statemachines/states that can create subgraph
 	auto CloseAllDocumentsTab = [this](const UEdGraphNode* InNode)
