Fix crash that occurs after undoing creation of a function

diff --git a/Engine/Source/Editor/Kismet/Private/BlueprintEditor.cpp b/Engine/Source/Editor/Kismet/Private/BlueprintEditor.cpp
--- a/Engine/Source/Editor/Kismet/Private/BlueprintEditor.cpp
+++ b/Engine/Source/Editor/Kismet/Private/BlueprintEditor.cpp
@@ -9117,9 +9117,15 @@ bool FBlueprintEditor::IsEditingAnimGraph() const
 
 UEdGraph* FBlueprintEditor::GetFocusedGraph() const
 {
-	if(FocusedGraphEdPtr.IsValid())
+	if (FocusedGraphEdPtr.IsValid())
 	{
-		return FocusedGraphEdPtr.Pin()->GetCurrentGraph();
+		if (UEdGraph* Graph = FocusedGraphEdPtr.Pin()->GetCurrentGraph())
+		{
+			if (!Graph->IsPendingKill())
+			{
+				return Graph;
+			}
+		}
 	}
 	return nullptr;
 }
