Refactor call to CompileBlueprint into compilation manager

diff --git a/Engine/Source/Developer/BlueprintNativeCodeGen/Private/BlueprintNativeCodeGenUtils.cpp b/Engine/Source/Developer/BlueprintNativeCodeGen/Private/BlueprintNativeCodeGenUtils.cpp
--- a/Engine/Source/Developer/BlueprintNativeCodeGen/Private/BlueprintNativeCodeGenUtils.cpp
+++ b/Engine/Source/Developer/BlueprintNativeCodeGen/Private/BlueprintNativeCodeGenUtils.cpp
@@ -1,6 +1,7 @@
 // Copyright 1998-2019 Epic Games, Inc. All Rights Reserved.
 
 #include "BlueprintNativeCodeGenUtils.h"
+#include "BlueprintCompilationManager.h"
 #include "Engine/Blueprint.h"
 #include "HAL/FileManager.h"
 #include "Misc/Paths.h"
@@ -453,19 +454,9 @@ void FBlueprintNativeCodeGenUtils::GenerateCppCode(UObject* Obj, TSharedPtr<FStr
 		CodeGenBackend.NativizationSummary() = NativizationSummary;
 
 		{
-			TSharedPtr<FBlueprintCompileReinstancer> Reinstancer = FBlueprintCompileReinstancer::Create(DuplicateBP->GeneratedClass);
+			FBlueprintCompilationManager::CompileSynchronouslyToCpp( DuplicateBP, OutCppSource, OutHeaderSource, NativizationOptions );
+			
 			IKismetCompilerInterface& Compiler = FModuleManager::LoadModuleChecked<IKismetCompilerInterface>(KISMET_COMPILER_MODULENAME);
-			TGuardValue<bool> GuardTemplateNameFlag(GCompilingBlueprint, true);
-			FCompilerResultsLog Results;
-
-			FKismetCompilerOptions CompileOptions;
-			CompileOptions.CompileType = EKismetCompileType::Cpp;
-			CompileOptions.OutCppSourceCode = OutCppSource;
-			CompileOptions.OutHeaderSourceCode = OutHeaderSource;
-			CompileOptions.NativizationOptions = NativizationOptions;
-
-			Compiler.CompileBlueprint(DuplicateBP, CompileOptions, Results);
-
 			Compiler.RemoveBlueprintGeneratedClasses(DuplicateBP);
 		}
 
diff --git a/Engine/Source/Editor/Kismet/Private/BlueprintCompilationManager.cpp b/Engine/Source/Editor/Kismet/Private/BlueprintCompilationManager.cpp
--- a/Engine/Source/Editor/Kismet/Private/BlueprintCompilationManager.cpp
+++ b/Engine/Source/Editor/Kismet/Private/BlueprintCompilationManager.cpp
@@ -87,6 +87,25 @@ struct FReinstancingJob;
 struct FSkeletonFixupData;
 struct FCompilerData;
 
+struct FBPCompilationManagerCPPResults
+{
+	const FCompilerNativizationOptions* CppOptions = nullptr;
+
+	TSharedPtr<FString> CppSource;
+	TSharedPtr<FString> HeaderSource;
+};
+
+struct FBPCompileRequestInternal
+{
+	FBPCompileRequestInternal(FBPCompileRequest InUserData)
+		: UserData(InUserData)
+	{
+	}
+
+	FBPCompileRequest UserData;
+	FBPCompilationManagerCPPResults CPPResults;
+};
+
 struct FBlueprintCompilationManagerImpl : public FGCObject
 {
 	FBlueprintCompilationManagerImpl();
@@ -98,8 +117,8 @@ struct FBlueprintCompilationManagerImpl : public FGCObject
 
 	void RegisterCompilerExtension(TSubclassOf<UBlueprint> BlueprintType, UBlueprintCompilerExtension* Extension);
 
-	void QueueForCompilation(const FBPCompileRequest& CompileJob);
-	void CompileSynchronouslyImpl(const FBPCompileRequest& Request);
+	void QueueForCompilation(const FBPCompileRequestInternal& CompileJob);
+	void CompileSynchronouslyImpl(const FBPCompileRequestInternal& Request);
 	void FlushCompilationQueueImpl(bool bSuppressBroadcastCompiled, TArray<UBlueprint*>* BlueprintsCompiled, TArray<UBlueprint*>* BlueprintsCompiledOrSkeletonCompiled, FUObjectSerializeContext* InLoadContext);
 	void ProcessExtensions(const TArray<FCompilerData>& InCurrentlyCompilingBPs);
 	void FlushReinstancingQueueImpl();
@@ -131,7 +150,7 @@ struct FBlueprintCompilationManagerImpl : public FGCObject
 	TMap<UClass*, TArray<UBlueprintCompilerExtension*> > CompilerExtensions;
 
 	// Queued requests to be processed in the next FlushCompilationQueueImpl call:
-	TArray<FBPCompileRequest> QueuedRequests;
+	TArray<FBPCompileRequestInternal> QueuedRequests;
 	
 	// Data stored for reinstancing, which finishes much later than compilation,
 	// populated by FlushCompilationQueueImpl, cleared by FlushReinstancingQueueImpl:
@@ -171,9 +190,9 @@ void FBlueprintCompilationManagerImpl::AddReferencedObjects(FReferenceCollector&
 		Collector.AddReferencedObjects(Extensions.Value);
 	}
 
-	for( FBPCompileRequest& Job : QueuedRequests )
+	for( FBPCompileRequestInternal& Job : QueuedRequests )
 	{
-		Collector.AddReferencedObject(Job.BPToCompile);
+		Collector.AddReferencedObject(Job.UserData.BPToCompile);
 	}
 
 	Collector.AddReferencedObjects(ClassesToReinstance);
@@ -190,38 +209,38 @@ void FBlueprintCompilationManagerImpl::RegisterCompilerExtension(TSubclassOf<UBl
 	CompilerExtensions.FindOrAdd(BlueprintType).Emplace(Extension);
 }
 
-void FBlueprintCompilationManagerImpl::QueueForCompilation(const FBPCompileRequest& CompileJob)
+void FBlueprintCompilationManagerImpl::QueueForCompilation(const FBPCompileRequestInternal& CompileJob)
 {
-	if(!CompileJob.BPToCompile->bQueuedForCompilation)
+	if(!CompileJob.UserData.BPToCompile->bQueuedForCompilation)
 	{
-		CompileJob.BPToCompile->bQueuedForCompilation = true;
+		CompileJob.UserData.BPToCompile->bQueuedForCompilation = true;
 		QueuedRequests.Add(CompileJob);
 	}
 }
 
-void FBlueprintCompilationManagerImpl::CompileSynchronouslyImpl(const FBPCompileRequest& Request)
+void FBlueprintCompilationManagerImpl::CompileSynchronouslyImpl(const FBPCompileRequestInternal& Request)
 {
 	DECLARE_SCOPE_HIERARCHICAL_COUNTER_FUNC()
 
-	Request.BPToCompile->bQueuedForCompilation = true;
+	Request.UserData.BPToCompile->bQueuedForCompilation = true;
 
-	const bool bIsRegeneratingOnLoad		= (Request.CompileOptions & EBlueprintCompileOptions::IsRegeneratingOnLoad		) != EBlueprintCompileOptions::None;
-	const bool bRegenerateSkeletonOnly		= (Request.CompileOptions & EBlueprintCompileOptions::RegenerateSkeletonOnly	) != EBlueprintCompileOptions::None;
-	const bool bSkipGarbageCollection		= (Request.CompileOptions & EBlueprintCompileOptions::SkipGarbageCollection		) != EBlueprintCompileOptions::None
+	const bool bIsRegeneratingOnLoad		= (Request.UserData.CompileOptions & EBlueprintCompileOptions::IsRegeneratingOnLoad		) != EBlueprintCompileOptions::None;
+	const bool bRegenerateSkeletonOnly		= (Request.UserData.CompileOptions & EBlueprintCompileOptions::RegenerateSkeletonOnly	) != EBlueprintCompileOptions::None;
+	const bool bSkipGarbageCollection		= (Request.UserData.CompileOptions & EBlueprintCompileOptions::SkipGarbageCollection	) != EBlueprintCompileOptions::None
 		|| bRegenerateSkeletonOnly;
-	const bool bBatchCompile				= (Request.CompileOptions & EBlueprintCompileOptions::BatchCompile				) != EBlueprintCompileOptions::None;
-	const bool bSkipReinstancing			= (Request.CompileOptions & EBlueprintCompileOptions::SkipReinstancing			) != EBlueprintCompileOptions::None;
-	const bool bSkipSaving					= (Request.CompileOptions & EBlueprintCompileOptions::SkipSave					) != EBlueprintCompileOptions::None;
+	const bool bBatchCompile				= (Request.UserData.CompileOptions & EBlueprintCompileOptions::BatchCompile				) != EBlueprintCompileOptions::None;
+	const bool bSkipReinstancing			= (Request.UserData.CompileOptions & EBlueprintCompileOptions::SkipReinstancing			) != EBlueprintCompileOptions::None;
+	const bool bSkipSaving					= (Request.UserData.CompileOptions & EBlueprintCompileOptions::SkipSave					) != EBlueprintCompileOptions::None;
 
 	ensure(!bIsRegeneratingOnLoad); // unexpected code path, compile on load handled with different function call
 	ensure(!bSkipReinstancing); // This is an internal option, should not go through CompileSynchronouslyImpl
 	ensure(QueuedRequests.Num() == 0);
 
 	// Wipe the PreCompile log, any generated messages are now irrelevant
-	Request.BPToCompile->PreCompileLog.Reset();
+	Request.UserData.BPToCompile->PreCompileLog.Reset();
 
 	// Reset the flag, so if the user tries to use PIE it will warn them if the BP did not compile
-	Request.BPToCompile->bDisplayCompilePIEWarning = true;
+	Request.UserData.BPToCompile->bDisplayCompilePIEWarning = true;
 	
 	// Do not want to run this code without the editor present nor when running commandlets.
 	// We do not want to regenerate a search Guid during loads, nothing has changed in the Blueprint
@@ -229,10 +248,11 @@ void FBlueprintCompilationManagerImpl::CompileSynchronouslyImpl(const FBPCompile
 	// We would like to regenerated it when a skeleton changes, but it is too expensive:
 	if (GEditor && GIsEditor && !bIsRegeneratingOnLoad && !bRegenerateSkeletonOnly)
 	{
-		FFindInBlueprintSearchManager::Get().AddOrUpdateBlueprintSearchMetadata(Request.BPToCompile);
+		FFindInBlueprintSearchManager::Get().AddOrUpdateBlueprintSearchMetadata(Request.UserData.BPToCompile);
 	}
 
 	QueuedRequests.Add(Request);
+
 	// We suppress normal compilation broadcasts because the old code path 
 	// did this after GC and we want to match the old behavior:
 	const bool bSuppressBroadcastCompiled = true;
@@ -241,9 +261,9 @@ void FBlueprintCompilationManagerImpl::CompileSynchronouslyImpl(const FBPCompile
 	FlushCompilationQueueImpl(bSuppressBroadcastCompiled, &CompiledBlueprints, &SkeletonCompiledBlueprints, nullptr);
 	FlushReinstancingQueueImpl();
 	
-	if( Request.ClientResultsLog && Request.ClientResultsLog->bLogDetailedResults)
+	if( Request.UserData.ClientResultsLog && Request.UserData.ClientResultsLog->bLogDetailedResults)
 	{
-		Request.ClientResultsLog->Note(
+		Request.UserData.ClientResultsLog->Note(
 			*FText::Format(
 				LOCTEXT("TotalCompiledBPs", "Compiled {0} total blueprints"),
 				CompiledBlueprints.Num()
@@ -251,10 +271,10 @@ void FBlueprintCompilationManagerImpl::CompileSynchronouslyImpl(const FBPCompile
 		);
 	}
 
-	if (FBlueprintEditorUtils::IsLevelScriptBlueprint(Request.BPToCompile) && !bRegenerateSkeletonOnly)
+	if (FBlueprintEditorUtils::IsLevelScriptBlueprint(Request.UserData.BPToCompile) && !bRegenerateSkeletonOnly)
 	{
 		// When the Blueprint is recompiled, then update the bound events for level scripting
-		ULevelScriptBlueprint* LevelScriptBP = CastChecked<ULevelScriptBlueprint>(Request.BPToCompile);
+		ULevelScriptBlueprint* LevelScriptBP = CastChecked<ULevelScriptBlueprint>(Request.UserData.BPToCompile);
 
 		// ULevel::OnLevelScriptBlueprintChanged needs to be run after the CDO has
 		// been updated as it respawns the actor:
@@ -278,7 +298,7 @@ void FBlueprintCompilationManagerImpl::CompileSynchronouslyImpl(const FBPCompile
 		GEditor->BroadcastBlueprintReinstanced();
 	}
 	
-	ensure(Request.BPToCompile->bQueuedForCompilation == false);
+	ensure(Request.UserData.BPToCompile->bQueuedForCompilation == false);
 
 	if(!bSkipGarbageCollection)
 	{
@@ -356,7 +376,13 @@ struct FSkeletonFixupData
 
 struct FCompilerData
 {
-	explicit FCompilerData(UBlueprint* InBP, ECompilationManagerJobType InJobType, FCompilerResultsLog* InResultsLogOverride, EBlueprintCompileOptions UserOptions, bool bBytecodeOnly)
+	explicit FCompilerData(
+		UBlueprint* InBP, 
+		ECompilationManagerJobType InJobType, 
+		FCompilerResultsLog* InResultsLogOverride, 
+		EBlueprintCompileOptions UserOptions, 
+		const FBPCompilationManagerCPPResults& CPPResults,
+		bool bBytecodeOnly)
 	{
 		check(InBP);
 		BP = InBP;
@@ -384,6 +410,14 @@ struct FCompilerData
 		InternalOptions.bSkipFiBSearchMetaUpdate = (UserOptions & EBlueprintCompileOptions::SkipFiBSearchMetaUpdate) != EBlueprintCompileOptions::None;
 		InternalOptions.CompileType = bBytecodeOnly ? EKismetCompileType::BytecodeOnly : EKismetCompileType::Full;
 
+		if(!bBytecodeOnly && CPPResults.CppOptions)
+		{
+			InternalOptions.OutCppSourceCode = CPPResults.CppSource;
+			InternalOptions.OutHeaderSourceCode = CPPResults.HeaderSource;
+			InternalOptions.NativizationOptions = *CPPResults.CppOptions;
+			InternalOptions.CompileType = EKismetCompileType::Cpp;
+		}
+
 		Compiler = FKismetCompilerContext::GetCompilerForBP(BP, *ActiveResultsLog, InternalOptions);
 	}
 
@@ -477,9 +511,9 @@ void FBlueprintCompilationManagerImpl::FlushCompilationQueueImpl(bool bSuppressB
 		TArray<UBlueprint*> BlueprintsToRecompile;
 
 		// First add any dependents of macro libraries that are being compiled:
-		for(const FBPCompileRequest& CompileJob : QueuedRequests)
+		for(const FBPCompileRequestInternal& CompileJob : QueuedRequests)
 		{
-			if ((CompileJob.CompileOptions & 
+			if ((CompileJob.UserData.CompileOptions & 
 				(	EBlueprintCompileOptions::RegenerateSkeletonOnly|
 					EBlueprintCompileOptions::IsRegeneratingOnLoad)
 				) != EBlueprintCompileOptions::None)
@@ -487,21 +521,22 @@ void FBlueprintCompilationManagerImpl::FlushCompilationQueueImpl(bool bSuppressB
 				continue;
 			}
 			
-			if(CompileJob.BPToCompile->BlueprintType == BPTYPE_MacroLibrary)
+			if(CompileJob.UserData.BPToCompile->BlueprintType == BPTYPE_MacroLibrary)
 			{
 				TArray<UBlueprint*> DependentBlueprints;
-				FBlueprintEditorUtils::GetDependentBlueprints(CompileJob.BPToCompile, DependentBlueprints);
+				FBlueprintEditorUtils::GetDependentBlueprints(CompileJob.UserData.BPToCompile, DependentBlueprints);
 				for(UBlueprint* DependentBlueprint : DependentBlueprints)
 				{
 					if(!IsQueuedForCompilation(DependentBlueprint))
 					{
 						DependentBlueprint->bQueuedForCompilation = true;
-						CurrentlyCompilingBPs.Add(
+						CurrentlyCompilingBPs.Emplace(
 							FCompilerData(
 								DependentBlueprint, 
 								ECompilationManagerJobType::Normal, 
 								nullptr, 
 								EBlueprintCompileOptions::None, 
+								FBPCompilationManagerCPPResults(),
 								false // full compile
 							)
 						);
@@ -512,9 +547,9 @@ void FBlueprintCompilationManagerImpl::FlushCompilationQueueImpl(bool bSuppressB
 		}
 
 		// then make sure any normal blueprints have their bytecode dependents recompiled, this is in case a function signature changes:
-		for(const FBPCompileRequest& CompileJob : QueuedRequests)
+		for(const FBPCompileRequestInternal& CompileJob : QueuedRequests)
 		{
-			if ((CompileJob.CompileOptions & EBlueprintCompileOptions::RegenerateSkeletonOnly) != EBlueprintCompileOptions::None)
+			if ((CompileJob.UserData.CompileOptions & EBlueprintCompileOptions::RegenerateSkeletonOnly) != EBlueprintCompileOptions::None)
 			{
 				continue;
 			}
@@ -524,7 +559,7 @@ void FBlueprintCompilationManagerImpl::FlushCompilationQueueImpl(bool bSuppressB
 			// added or removed. Below (Stage VIII) we skip further compilation for blueprints 
 			// that are being bytecode compiled, but their dependencies have not changed:
 			TArray<UBlueprint*> DependentBlueprints;
-			FBlueprintEditorUtils::GetDependentBlueprints(CompileJob.BPToCompile, DependentBlueprints);
+			FBlueprintEditorUtils::GetDependentBlueprints(CompileJob.UserData.BPToCompile, DependentBlueprints);
 			for(UBlueprint* DependentBlueprint : DependentBlueprints)
 			{
 				if(!IsQueuedForCompilation(DependentBlueprint))
@@ -533,12 +568,13 @@ void FBlueprintCompilationManagerImpl::FlushCompilationQueueImpl(bool bSuppressB
 					// Because we're adding this as a bytecode only blueprint compile we don't need to 
 					// recursively recompile dependencies. The assumption is that a bytecode only compile
 					// will not change the class layout. @todo: add an ensure to detect class layout changes
-					CurrentlyCompilingBPs.Add(
+					CurrentlyCompilingBPs.Emplace(
 						FCompilerData(
 							DependentBlueprint, 
 							ECompilationManagerJobType::Normal, 
 							nullptr, 
 							EBlueprintCompileOptions::None, 
+							FBPCompilationManagerCPPResults(),
 							true
 						)
 					);
@@ -551,8 +587,8 @@ void FBlueprintCompilationManagerImpl::FlushCompilationQueueImpl(bool bSuppressB
 		for(int32 I = 0; I < QueuedRequests.Num(); ++I)
 		{
 			bool bSkipCompile = false;
-			FBPCompileRequest& QueuedJob = QueuedRequests[I];
-			UBlueprint* QueuedBP = QueuedJob.BPToCompile;
+			FBPCompileRequestInternal& QueuedJob = QueuedRequests[I];
+			UBlueprint* QueuedBP = QueuedJob.UserData.BPToCompile;
 
 			ensure(!QueuedBP->GeneratedClass ||
 				!QueuedBP->GeneratedClass->ClassDefaultObject ||
@@ -591,7 +627,16 @@ void FBlueprintCompilationManagerImpl::FlushCompilationQueueImpl(bool bSuppressB
 
 			if(bSkipCompile)
 			{
-				CurrentlyCompilingBPs.Add(FCompilerData(QueuedBP, ECompilationManagerJobType::SkeletonOnly, QueuedJob.ClientResultsLog, QueuedJob.CompileOptions, false));
+				CurrentlyCompilingBPs.Emplace(
+					FCompilerData(
+						QueuedBP, 
+						ECompilationManagerJobType::SkeletonOnly, 
+						QueuedJob.UserData.ClientResultsLog, 
+						QueuedJob.UserData.CompileOptions,
+						FBPCompilationManagerCPPResults(),
+						false
+					)
+				);
 				if (QueuedBP->GeneratedClass != nullptr)
 				{
 					// set bIsRegeneratingOnLoad so that we don't reset loaders:
@@ -614,12 +659,22 @@ void FBlueprintCompilationManagerImpl::FlushCompilationQueueImpl(bool bSuppressB
 			else
 			{
 				ECompilationManagerJobType JobType = ECompilationManagerJobType::Normal;
-				if ((QueuedJob.CompileOptions & EBlueprintCompileOptions::RegenerateSkeletonOnly) != EBlueprintCompileOptions::None)
+				if ((QueuedJob.UserData.CompileOptions & EBlueprintCompileOptions::RegenerateSkeletonOnly) != EBlueprintCompileOptions::None)
 				{
 					JobType = ECompilationManagerJobType::SkeletonOnly;
 				}
 
-				CurrentlyCompilingBPs.Add(FCompilerData(QueuedBP, JobType, QueuedJob.ClientResultsLog, QueuedJob.CompileOptions, false));
+				CurrentlyCompilingBPs.Emplace(
+					FCompilerData(
+						QueuedBP, 
+						JobType, 
+						QueuedJob.UserData.ClientResultsLog, 
+						QueuedJob.UserData.CompileOptions, 
+						QueuedJob.CPPResults, 
+						false
+					)
+				);
+
 				BlueprintsToRecompile.Add(QueuedBP);
 			}
 		}
@@ -642,7 +697,16 @@ void FBlueprintCompilationManagerImpl::FlushCompilationQueueImpl(bool bSuppressB
 						{
 							ChildBlueprint->bQueuedForCompilation = true;
 							ensure(ChildBlueprint->bHasBeenRegenerated);
-							CurrentlyCompilingBPs.Add(FCompilerData(ChildBlueprint, ECompilationManagerJobType::RelinkOnly, nullptr, EBlueprintCompileOptions::None, false));
+							CurrentlyCompilingBPs.Emplace(
+								FCompilerData(
+									ChildBlueprint, 
+									ECompilationManagerJobType::RelinkOnly, 
+									nullptr, 
+									EBlueprintCompileOptions::None, 
+									FBPCompilationManagerCPPResults(),
+									false
+								)
+							);
 						}
 					}
 				}
@@ -2848,6 +2912,26 @@ void FBlueprintCompilationManager::CompileSynchronously(const FBPCompileRequest&
 	}
 }
 
+void FBlueprintCompilationManager::CompileSynchronouslyToCpp(UBlueprint* BP, TSharedPtr<FString> OutHeaderSource, TSharedPtr<FString> OutCppSource, const FCompilerNativizationOptions& NativizationOptions)
+{
+	if(BPCMImpl)
+	{
+		FBPCompileRequestInternal Request(
+			FBPCompileRequest( 
+				BP, 
+				EBlueprintCompileOptions::SkipGarbageCollection|EBlueprintCompileOptions::SkipSave,
+				nullptr
+			)
+		);
+
+		Request.CPPResults.CppOptions = &NativizationOptions;
+		Request.CPPResults.HeaderSource = OutHeaderSource;
+		Request.CPPResults.CppSource = OutCppSource;
+
+		BPCMImpl->CompileSynchronouslyImpl(Request);
+	}
+}
+
 void FBlueprintCompilationManager::NotifyBlueprintLoaded(UBlueprint* BPLoaded)
 {
 	// Blueprints can be loaded before editor modules are on line:
diff --git a/Engine/Source/Editor/Kismet/Public/BlueprintCompilationManager.h b/Engine/Source/Editor/Kismet/Public/BlueprintCompilationManager.h
--- a/Engine/Source/Editor/Kismet/Public/BlueprintCompilationManager.h
+++ b/Engine/Source/Editor/Kismet/Public/BlueprintCompilationManager.h
@@ -55,6 +55,11 @@ struct KISMET_API FBlueprintCompilationManager
 	 * if there are several blueprints that require compilation (e.g. typical case on PIE):
 	 */
 	static void CompileSynchronously(const FBPCompileRequest& Request);
+
+	/** 
+	 * Write c++ represenetation of the blueprint to OutHeaderSource and OutCppSource
+	 */
+	static void CompileSynchronouslyToCpp(UBlueprint* BP, TSharedPtr<FString> OutHeaderSource, TSharedPtr<FString> OutCppSource, const FCompilerNativizationOptions& NativizationOptions);
 	
 	/**
 	 * Adds a newly loaded blueprint to the compilation queue
