Retain parameter default values when converting between functions and custom events
Move BlueprintEditor pin helper functions into EdGraphUtilities

diff --git a/Engine/Source/Editor/Kismet/Private/BlueprintEditor.cpp b/Engine/Source/Editor/Kismet/Private/BlueprintEditor.cpp
--- a/Engine/Source/Editor/Kismet/Private/BlueprintEditor.cpp
+++ b/Engine/Source/Editor/Kismet/Private/BlueprintEditor.cpp
@@ -5453,7 +5453,7 @@ void FBlueprintEditor::ConvertFunctionToEvent(UK2Node_FunctionEntry* SelectedCal
 
 		// Keep track of the old connections from the entry node
 		TMap<FString, TSet<UEdGraphPin*>> PinConnections;
-		GetPinConnectionMap(SelectedCallFunctionNode, PinConnections);
+		FEdGraphUtilities::GetPinConnectionMap(SelectedCallFunctionNode, PinConnections);
 
 		FName EventName = Func->GetFName();
 		UClass* const OverrideFuncClass = CastChecked<UClass>(Func->GetOuter())->GetAuthoritativeClass();
@@ -5501,7 +5501,7 @@ void FBlueprintEditor::ConvertFunctionToEvent(UK2Node_FunctionEntry* SelectedCal
 			NewEventNode->CustomFunctionName = EventName;
 			NewEventNode->bOverrideFunction = false;
 
-			// Add every type of user pin that we need tot he new event node
+			// Add every type of user pin that we need to the new event node
 			for (TSharedPtr<FUserPinInfo> Pin : SelectedCallFunctionNode->UserDefinedPins)
 			{
 				NewEventNode->CreateUserDefinedPin(Pin->PinName, Pin->PinType, Pin->DesiredPinDirection);
@@ -5533,7 +5533,8 @@ void FBlueprintEditor::ConvertFunctionToEvent(UK2Node_FunctionEntry* SelectedCal
 		if (NewEventNode)
 		{
 			// Link the nodes from the original function entry node to the new event node
-			ReconnectPinMap(NewEventNode, PinConnections);
+			FEdGraphUtilities::ReconnectPinMap(NewEventNode, PinConnections);
+			FEdGraphUtilities::CopyPinDefaults(SelectedCallFunctionNode, NewEventNode);
 			FKismetEditorUtilities::BringKismetToFocusAttentionOnObject(NewEventNode, false);
 		}
 
@@ -5544,39 +5545,6 @@ void FBlueprintEditor::ConvertFunctionToEvent(UK2Node_FunctionEntry* SelectedCal
 	}
 }
 
-void FBlueprintEditor::GetPinConnectionMap(UEdGraphNode* Node, TMap<FString, TSet<UEdGraphPin*>>& OutPinConnections) const
-{
-	check(Node);
-	
-	for (UEdGraphPin* Pin : Node->Pins)
-	{
-		if (Pin != nullptr)
-		{
-			FString PinName = Pin->GetName();
-
-			// If this is the first time seeing this pin name, add a new set
-			if (!OutPinConnections.Contains(PinName))
-			{
-				OutPinConnections.Add(PinName, TSet<UEdGraphPin*>());
-			}
-			else
-			{
-				// There are no pins connected to this
-				continue;
-			}
-
-			for (UEdGraphPin* ConnectedPin : Pin->LinkedTo)
-			{
-				if (ConnectedPin != nullptr)
-				{
-					// Add to the array of nodes at this pin name
-					OutPinConnections[PinName].Add(ConnectedPin);
-				}
-			}
-		}
-	}
-}
-
 bool FBlueprintEditor::CanConvertEventToFunction() const
 {
 	// Only allow when only a _single_ event node is selected, not allowed on the anim graph
@@ -5654,7 +5622,7 @@ void FBlueprintEditor::ConvertEventToFunction(UK2Node_Event* SelectedEventNode)
 
 		// Keep track of the old connections from the event node
 		TMap<FString, TSet<UEdGraphPin*>> PinConnections;
-		GetPinConnectionMap(SelectedEventNode, PinConnections);
+		FEdGraphUtilities::GetPinConnectionMap(SelectedEventNode, PinConnections);
 
 		TArray<UEdGraphNode*> CollapsableNodes = GetAllConnectedNodes(SelectedEventNode);
 
@@ -5763,9 +5731,11 @@ void FBlueprintEditor::ConvertEventToFunction(UK2Node_Event* SelectedEventNode)
 			MoveNodesToGraph(CollapsableNodes, NewGraph, ExpandedNodes, &Entry, &Result);
 			
 			// Link the new nodes accordingly
-			ReconnectPinMap(NewEntryNode, PinConnections);
 			if (NewEntryNode)
 			{
+				FEdGraphUtilities::ReconnectPinMap(NewEntryNode, PinConnections);
+				FEdGraphUtilities::CopyPinDefaults(SelectedEventNode, NewEntryNode);
+
 				MoveNodesToAveragePos(ExpandedNodes, FVector2D(NewEntryNode->NodePosX + 500.0f, NewEntryNode->NodePosY));
 			}
 			else
@@ -5800,28 +5770,19 @@ void FBlueprintEditor::ConvertEventToFunction(UK2Node_Event* SelectedEventNode)
 			{
 				IBlueprintNodeBinder::FBindingSet Bindings;
 				UEdGraphNode* OutFunctionNode = UBlueprintFunctionNodeSpawner::Create(NewFunction)->Invoke(SourceGraph, Bindings, NewFuncCallSpawn);
+				if (NewEntryNode)
+				{
+					FEdGraphUtilities::CopyPinDefaults(NewEntryNode, OutFunctionNode);
+				}
+				FKismetEditorUtilities::BringKismetToFocusAttentionOnObject(OutFunctionNode, false);
 			}
-		}
-
-		FBlueprintEditorUtils::MarkBlueprintAsStructurallyModified(NodeBP);
-	}
-}
-
-void FBlueprintEditor::ReconnectPinMap(UEdGraphNode* Node, const TMap<FString, TSet<UEdGraphPin*>>& PinConnections)
-{
-	check(Node);
-
-	for (UEdGraphPin* const NewPin : Node->Pins)
-	{
-		const FString& NewPinName = NewPin->GetName();
-		if (PinConnections.Contains(NewPinName))
-		{
-			// Connect the new pins here
-			for (UEdGraphPin* OldPin : PinConnections[NewPinName])
+			else if(NewEntryNode)
 			{
-				NewPin->MakeLinkTo(OldPin);
+				FKismetEditorUtilities::BringKismetToFocusAttentionOnObject(NewEntryNode, false);
 			}
 		}
+
+		FBlueprintEditorUtils::MarkBlueprintAsStructurallyModified(NodeBP);
 	}
 }
 
diff --git a/Engine/Source/Editor/Kismet/Public/BlueprintEditor.h b/Engine/Source/Editor/Kismet/Public/BlueprintEditor.h
--- a/Engine/Source/Editor/Kismet/Public/BlueprintEditor.h
+++ b/Engine/Source/Editor/Kismet/Public/BlueprintEditor.h
@@ -837,22 +837,6 @@ protected:
 	*/
 	void ConvertFunctionToEvent(UK2Node_FunctionEntry* SelectedCallFunctionNode);
 
-	/**
-	* Output a map of pin names to a set of connections given a function entry pin
-	* 
-	* @param Node				The entry pin to gather connections from
-	* @param OutPinConnections	Output map of connection data
-	*/
-	void GetPinConnectionMap(UEdGraphNode* Node, TMap<FString, TSet<UEdGraphPin*>>& OutPinConnections) const;
-
-	/**
-	* Reconnect the pin map to the given node
-	* 
-	* @param Node				The node to connect pins to
-	* @param PinConnections		Map of pin connections to set
-	*/
-	void ReconnectPinMap(UEdGraphNode* Node, const TMap<FString, TSet<UEdGraphPin*>>& PinConnections);
-
 	void OnConvertEventToFunction();
 	void ConvertEventToFunction(UK2Node_Event* SelectedEventNode);
 	bool CanConvertEventToFunction() const;
diff --git a/Engine/Source/Editor/UnrealEd/Private/EdGraphUtilities.cpp b/Engine/Source/Editor/UnrealEd/Private/EdGraphUtilities.cpp
--- a/Engine/Source/Editor/UnrealEd/Private/EdGraphUtilities.cpp
+++ b/Engine/Source/Editor/UnrealEd/Private/EdGraphUtilities.cpp
@@ -628,6 +628,87 @@ void FEdGraphUtilities::UnregisterVisualPinConnectionFactory(TSharedPtr<FGraphPa
     VisualPinConnectionFactories.Remove(OldFactory);
 }
 
+
+void FEdGraphUtilities::CopyPinDefaults(const UEdGraphNode* NodeFrom, UEdGraphNode* NodeTo)
+{
+	check(NodeFrom && NodeTo);
+
+	UK2Node_EditablePinBase* WithUserDefinedPins = Cast<UK2Node_EditablePinBase>(NodeTo);
+
+	for (UEdGraphPin* ToPin : NodeTo->Pins)
+	{
+		if (UEdGraphPin * FromPin = NodeFrom->FindPin(ToPin->GetName()))
+		{
+			ToPin->DefaultValue = FromPin->DefaultValue;
+			ToPin->DefaultObject = FromPin->DefaultObject;
+			ToPin->DefaultTextValue = FromPin->DefaultTextValue;
+
+			if (WithUserDefinedPins)
+			{
+				// Copy defaults to user defined pins
+				WithUserDefinedPins->PinDefaultValueChanged(FromPin);
+			}
+			else
+			{
+				// No user defined pins implies we're copying to a function call node.
+				// Ensure that underlying changes to the entry node default are reflected to unmodified call nodes.
+				ToPin->AutogeneratedDefaultValue = FromPin->DefaultValue;
+			}
+		}
+	}
+}
+
+void FEdGraphUtilities::GetPinConnectionMap(const UEdGraphNode* Node, TMap<FString, TSet<UEdGraphPin*>>& OutPinConnections)
+{
+	check(Node);
+
+	for (UEdGraphPin* Pin : Node->Pins)
+	{
+		if (Pin != nullptr)
+		{
+			FString PinName = Pin->GetName();
+
+			// If this is the first time seeing this pin name, add a new set
+			if (!OutPinConnections.Contains(PinName))
+			{
+				OutPinConnections.Add(PinName, TSet<UEdGraphPin*>());
+			}
+			else
+			{
+				// There are no pins connected to this
+				continue;
+			}
+
+			for (UEdGraphPin* ConnectedPin : Pin->LinkedTo)
+			{
+				if (ConnectedPin != nullptr)
+				{
+					// Add to the array of nodes at this pin name
+					OutPinConnections[PinName].Add(ConnectedPin);
+				}
+			}
+		}
+	}
+}
+
+void FEdGraphUtilities::ReconnectPinMap(UEdGraphNode* Node, const TMap<FString, TSet<UEdGraphPin*>>& PinConnections)
+{
+	check(Node);
+
+	for (UEdGraphPin* const NewPin : Node->Pins)
+	{
+		const FString& NewPinName = NewPin->GetName();
+		if (PinConnections.Contains(NewPinName))
+		{
+			// Connect the new pins here
+			for (UEdGraphPin* OldPin : PinConnections[NewPinName])
+			{
+				NewPin->MakeLinkTo(OldPin);
+			}
+		}
+	}
+}
+
 void FEdGraphUtilities::FNodeVisitor::TraverseNodes(UEdGraphNode* Node)
 {
 	VisitedNodes.Add(Node);
diff --git a/Engine/Source/Editor/UnrealEd/Public/EdGraphUtilities.h b/Engine/Source/Editor/UnrealEd/Public/EdGraphUtilities.h
--- a/Engine/Source/Editor/UnrealEd/Public/EdGraphUtilities.h
+++ b/Engine/Source/Editor/UnrealEd/Public/EdGraphUtilities.h
@@ -149,9 +149,28 @@ public:
 	/** returns the first pin marked as an MapParam, usually returning nullptr */
 	static UEdGraphPin* FindMapParamPin(const UFunction* Function, const UEdGraphNode* Node);
 
-	/** returns the first pin referred to in a , and | deliniated list of pin names */
+	/** returns the first pin referred to in a , and | delineated list of pin names */
 	static UEdGraphPin* FindPinFromMetaData(const UFunction* Function, const UEdGraphNode* Node, FName MetaData );
 
+	/**
+	* Output a map of pin names to a set of connections given a function entry pin
+	* 
+	* @param Node				The entry pin to gather connections from
+	* @param OutPinConnections	Output map of connection data
+	*/
+	static void GetPinConnectionMap(const UEdGraphNode* Node, TMap<FString, TSet<UEdGraphPin*>>& OutPinConnections);
+
+	/**
+	* Reconnect the pin map to the given node
+	* 
+	* @param Node				The node to connect pins to
+	* @param PinConnections		Map of pin connections to set
+	*/
+	static void ReconnectPinMap(UEdGraphNode* Node, const TMap<FString, TSet<UEdGraphPin*>>& PinConnections);
+
+	/** Copy pin default values between nodes for any pins with matching names */
+	static void CopyPinDefaults(const UEdGraphNode* NodeFrom, UEdGraphNode* NodeTo);
+
 private:
 	static TArray< TSharedPtr<FGraphPanelNodeFactory> > VisualNodeFactories;
 	static TArray< TSharedPtr<FGraphPanelPinFactory> > VisualPinFactories;
