Allow changing component class override in blueprints to be undone, however, since this causes crashes that we don't have time to fully investigate now, also disable the component class override feature.
#rnx

diff --git a/Engine/Source/Editor/Kismet/Private/BlueprintDetailsCustomization.cpp b/Engine/Source/Editor/Kismet/Private/BlueprintDetailsCustomization.cpp
--- a/Engine/Source/Editor/Kismet/Private/BlueprintDetailsCustomization.cpp
+++ b/Engine/Source/Editor/Kismet/Private/BlueprintDetailsCustomization.cpp
@@ -5782,7 +5782,8 @@ void FBlueprintComponentDetails::CustomizeDetails(IDetailLayoutBuilder& DetailLa
 			]
 		];
 
-		UActorComponent* ComponentTemplate = (CachedNodePtr->IsNative() ? CachedNodePtr->GetComponentTemplate() : nullptr);
+		// Remove UI to specify overriden component class until undo/redo can be made to work with crashing
+		/*UActorComponent* ComponentTemplate = (CachedNodePtr->IsNative() ? CachedNodePtr->GetComponentTemplate() : nullptr);
 		if (ComponentTemplate)
 		{
 			UClass* BaseClass = ComponentTemplate->GetClass();
@@ -5814,7 +5815,7 @@ void FBlueprintComponentDetails::CustomizeDetails(IDetailLayoutBuilder& DetailLa
 				.AllowNone(false)
 				.SelectedClass(this, &FBlueprintComponentDetails::GetSelectedEntryClass)
 				.OnSetClass(this, &FBlueprintComponentDetails::HandleNewEntryClassSelected)];
-		}
+		}*/
 
 		IDetailCategoryBuilder& SocketsCategory = DetailLayout.EditCategory("Sockets", LOCTEXT("BlueprintComponentDetailsCategory", "Sockets"), ECategoryPriority::Important);
 
@@ -6112,11 +6113,15 @@ void FBlueprintComponentDetails::HandleNewEntryClassSelected(const UClass* NewEn
 
 		if (UActorComponent* ComponentTemplate = CachedNodePtr->GetComponentTemplate())
 		{
+			const FScopedTransaction Transaction(LOCTEXT("SetComponentClassOverride", "Set Component Class Override"));
+
 			const FName ComponentTemplateName = ComponentTemplate->GetFName();
 
 			UBlueprint* BlueprintObj = GetBlueprintObj();
 			check(BlueprintObj);
 
+			BlueprintObj->Modify();
+
 			if (FBPComponentClassOverride* Override = BlueprintObj->ComponentClassOverrides.FindByKey(ComponentTemplateName))
 			{
 				bool bRemoveEntry = false;
