Don't allow properties of components that are not editable when inherited to be edited through the data only blueprints window or the actor's instance properties panel

diff --git a/Engine/Source/Editor/Kismet/Private/BlueprintEditorTabFactories.cpp b/Engine/Source/Editor/Kismet/Private/BlueprintEditorTabFactories.cpp
--- a/Engine/Source/Editor/Kismet/Private/BlueprintEditorTabFactories.cpp
+++ b/Engine/Source/Editor/Kismet/Private/BlueprintEditorTabFactories.cpp
@@ -2,6 +2,7 @@
 
 #include "BlueprintEditorTabFactories.h"
 #include "Widgets/Text/STextBlock.h"
+#include "Widgets/Images/SImage.h"
 #include "Editor.h"
 #include "Widgets/Layout/SWrapBox.h"
 #include "Widgets/Notifications/SErrorText.h"
@@ -153,6 +154,7 @@ TSharedRef<SWidget> FDefaultsEditorSummoner::CreateTabBody(const FWorkflowTabSpa
 	TSharedPtr<FBlueprintEditor> BlueprintEditorPtr = StaticCastSharedPtr<FBlueprintEditor>(HostingApp.Pin());
 
 	TSharedRef<SWidget> Message = CreateOptionalDataOnlyMessage();
+	TSharedRef<SWidget> EditableWarning = CreateOptionalEditableWarning();
 
 	return SNew(SVerticalBox)
 
@@ -163,6 +165,13 @@ TSharedRef<SWidget> FDefaultsEditorSummoner::CreateTabBody(const FWorkflowTabSpa
 			Message
 		]
 
+		+SVerticalBox::Slot()
+		.AutoHeight()
+		.Padding(FMargin(0,0,0,1))
+		[
+			EditableWarning
+		]
+
 		+ SVerticalBox::Slot()
 		.FillHeight(1.0f)
 		[
@@ -170,6 +179,60 @@ TSharedRef<SWidget> FDefaultsEditorSummoner::CreateTabBody(const FWorkflowTabSpa
 		];
 }
 
+TSharedRef<SWidget> FDefaultsEditorSummoner::CreateOptionalEditableWarning() const
+{
+	TSharedPtr<FBlueprintEditor> BlueprintEditorPtr = StaticCastSharedPtr<FBlueprintEditor>(HostingApp.Pin());
+
+	bool bHasUneditableBlueprintComponent = false;
+	if (UBlueprint* Blueprint = BlueprintEditorPtr->GetBlueprintObj())
+	{
+		if (Blueprint->GeneratedClass && FBlueprintEditorUtils::IsDataOnlyBlueprint(Blueprint))
+		{
+			if (AActor* Actor = Cast<AActor>(Blueprint->GeneratedClass->GetDefaultObject()))
+			{
+				for (UActorComponent* Component : Actor->GetComponents())
+				{
+					if (Component && !Component->IsCreatedByConstructionScript())
+					{
+						UActorComponent* ComponentArchetype = Cast<UActorComponent>(Component->GetArchetype());
+						if (ComponentArchetype && !ComponentArchetype->bEditableWhenInherited)
+						{
+							bHasUneditableBlueprintComponent = true;
+							break;
+						}
+					}
+				}
+			}
+		}
+	}
+
+	TSharedRef<SWidget> Message = SNullWidget::NullWidget;
+	if (bHasUneditableBlueprintComponent)
+	{
+		Message = SNew(SHorizontalBox)
+			+ SHorizontalBox::Slot()
+			.AutoWidth()
+			.HAlign(HAlign_Center)
+			.VAlign(VAlign_Center)
+			.Padding(2)
+			[
+				SNew(SImage)
+				.Image(FEditorStyle::Get().GetBrush("Icons.Warning"))
+			]
+
+			+ SHorizontalBox::Slot()
+			.VAlign(VAlign_Center)
+			.Padding(2)
+			[
+				SNew(STextBlock)
+				.Font(FEditorStyle::GetFontStyle("BoldFont"))
+				.Text(LOCTEXT("BlueprintUneditableInheritedComponentWarning", "Some properties are not editable due to belonging to a Component flagged as not editable when inherited."))
+			];
+	}
+
+	return Message;
+}
+
 TSharedRef<SWidget> FDefaultsEditorSummoner::CreateOptionalDataOnlyMessage() const
 {
 	TSharedPtr<FBlueprintEditor> BlueprintEditorPtr = StaticCastSharedPtr<FBlueprintEditor>(HostingApp.Pin());
diff --git a/Engine/Source/Editor/Kismet/Private/BlueprintEditorTabFactories.h b/Engine/Source/Editor/Kismet/Private/BlueprintEditorTabFactories.h
--- a/Engine/Source/Editor/Kismet/Private/BlueprintEditorTabFactories.h
+++ b/Engine/Source/Editor/Kismet/Private/BlueprintEditorTabFactories.h
@@ -83,6 +83,7 @@ public:
 
 private:
 	TSharedRef<SWidget> CreateOptionalDataOnlyMessage() const;
+	TSharedRef<SWidget> CreateOptionalEditableWarning() const;
 
 	void OnChangeBlueprintToNotDataOnly();
 };
diff --git a/Engine/Source/Runtime/Engine/Classes/Components/ActorComponent.h b/Engine/Source/Runtime/Engine/Classes/Components/ActorComponent.h
--- a/Engine/Source/Runtime/Engine/Classes/Components/ActorComponent.h
+++ b/Engine/Source/Runtime/Engine/Classes/Components/ActorComponent.h
@@ -792,6 +792,7 @@ public:
 	virtual void Serialize(FArchive& Ar) override;
 #if WITH_EDITOR
 	virtual bool Modify( bool bAlwaysMarkDirty = true ) override;
+	virtual bool CanEditChange(const UProperty* InProperty) const override;
 	virtual void PreEditChange(UProperty* PropertyThatWillChange) override;
 	virtual void PostEditChangeProperty(FPropertyChangedEvent& PropertyChangedEvent) override;
 	virtual void PostEditChangeChainProperty( FPropertyChangedChainEvent& PropertyChangedEvent ) override;
diff --git a/Engine/Source/Runtime/Engine/Private/Components/ActorComponent.cpp b/Engine/Source/Runtime/Engine/Private/Components/ActorComponent.cpp
--- a/Engine/Source/Runtime/Engine/Private/Components/ActorComponent.cpp
+++ b/Engine/Source/Runtime/Engine/Private/Components/ActorComponent.cpp
@@ -1762,6 +1762,21 @@ void UActorComponent::OnRep_IsActive()
 	SetComponentTickEnabled(IsActive());
 }
 
+#if WITH_EDITOR
+bool UActorComponent::CanEditChange(const UProperty* InProperty) const
+{
+	if (Super::CanEditChange(InProperty))
+	{
+		UActorComponent* ComponentArchetype = Cast<UActorComponent>(GetArchetype());
+		if (ComponentArchetype == nullptr || ComponentArchetype->bEditableWhenInherited)
+		{
+			return true;
+		}
+	}
+	return false;
+}
+#endif
+
 bool UActorComponent::IsEditableWhenInherited() const
 {
 	bool bCanEdit = bEditableWhenInherited;
