FiB: Defer the serialization cost of unloaded BP asset search data versioning info from editor startup.

diff --git a/Engine/Source/Editor/Kismet/Private/FindInBlueprintManager.cpp b/Engine/Source/Editor/Kismet/Private/FindInBlueprintManager.cpp
--- a/Engine/Source/Editor/Kismet/Private/FindInBlueprintManager.cpp
+++ b/Engine/Source/Editor/Kismet/Private/FindInBlueprintManager.cpp
@@ -115,7 +115,9 @@ inline UObject* GetAssetObject(UObject* InObject)
 // FSearchDataVersionInfo
 FSearchDataVersionInfo FSearchDataVersionInfo::Current =
 {
+	/** FiBDataVersion */
 	EFiBVersion::FIB_VER_LATEST,
+	/** EditorObjectVersion */
 	FEditorObjectVersion::LatestVersion
 };
 
@@ -366,6 +368,56 @@ namespace FiBSerializationHelpers
 		SizeOfDataAr << ReturnValue;
 		return ReturnValue;
 	}
+
+	/** Helper function to validate search data version info */
+	bool ValidateSearchDataVersionInfo(const FSearchDataVersionInfo& InVersionInfo)
+	{
+		return InVersionInfo.FiBDataVersion != EFiBVersion::FIB_VER_NONE && InVersionInfo.EditorObjectVersion >= 0;
+	}
+
+	/** Helper function to validate and/or deserialize version info if necessary */
+	bool ValidateSearchDataVersionInfo(const FString& InAssetPath, const FString& InFiBData, FSearchDataVersionInfo& InOutVersionInfo)
+	{
+		if (InOutVersionInfo.FiBDataVersion == EFiBVersion::FIB_VER_NONE)
+		{
+			// Deserialize the FiB data version
+			const int32 FiBDataLength = InFiBData.Len();
+			if (ensureMsgf(FiBDataLength > 0, TEXT("Versioned search data was zero length!")))
+			{
+				FBufferReader ReaderStream((void*)*InFiBData, FiBDataLength * sizeof(TCHAR), false);
+				InOutVersionInfo.FiBDataVersion = FiBSerializationHelpers::Deserialize<int32>(ReaderStream);
+			}
+		}
+
+		if (InOutVersionInfo.EditorObjectVersion < 0)
+		{
+			// Determine the editor object version that the asset package was last serialized with
+			FString PackageFilename;
+			const FString PackageName = FPackageName::ObjectPathToPackageName(InAssetPath);
+			if (ensureMsgf(FPackageName::DoesPackageExist(PackageName, nullptr, &PackageFilename), TEXT("FiB: Failed to map package to filename.")))
+			{
+				// Open a new file archive for reading
+				FArchive* PackageFile = IFileManager::Get().CreateFileReader(*PackageFilename);
+				if (ensureMsgf(PackageFile != nullptr, TEXT("FiB: Unable to open package to read file summary.")))
+				{
+					// Read the package file summary
+					FPackageFileSummary PackageFileSummary;
+					*PackageFile << PackageFileSummary;
+
+					// Close the file
+					delete PackageFile;
+
+					// If an editor object version exists in the package file summary, record it
+					if (const FCustomVersion* const EditorObjectVersion = PackageFileSummary.GetCustomVersionContainer().GetVersion(FEditorObjectVersion::GUID))
+					{
+						InOutVersionInfo.EditorObjectVersion = EditorObjectVersion->Version;
+					}
+				}
+			}
+		}
+
+		return ValidateSearchDataVersionInfo(InOutVersionInfo);
+	}
 }
 
 namespace BlueprintSearchMetaDataHelpers
@@ -1565,7 +1617,7 @@ void FFindInBlueprintSearchManager::OnAssetAdded(const FAssetData& InAssetData)
 			}
 			else
 			{
-				ExtractUnloadedFiBData(InAssetData, FiBVersionedSearchData, true);
+				ExtractUnloadedFiBData(InAssetData, FiBVersionedSearchData, EFiBVersion::FIB_VER_NONE);
 			}
 		}
 		else
@@ -1574,7 +1626,7 @@ void FFindInBlueprintSearchManager::OnAssetAdded(const FAssetData& InAssetData)
 			FAssetDataTagMapSharedView::FFindTagResult ResultLegacy = InAssetData.TagsAndValues.FindTag("FiB");
 			if (ResultLegacy.IsSet())
 			{
-				ExtractUnloadedFiBData(InAssetData, ResultLegacy.GetValue(), false);
+				ExtractUnloadedFiBData(InAssetData, ResultLegacy.GetValue(), EFiBVersion::FIB_VER_BASE);
 			}
 			// The asset has no FiB data, keep track of it so we can inform the user
 			else
@@ -1586,7 +1638,7 @@ void FFindInBlueprintSearchManager::OnAssetAdded(const FAssetData& InAssetData)
 	}
 }
 
-void FFindInBlueprintSearchManager::ExtractUnloadedFiBData(const FAssetData& InAssetData, const FString& InFiBData, bool bIsVersioned)
+void FFindInBlueprintSearchManager::ExtractUnloadedFiBData(const FAssetData& InAssetData, const FString& InFiBData, EFiBVersion InFiBDataVersion)
 {
 	if (SearchMap.Contains(InAssetData.ObjectPath))
 	{
@@ -1647,36 +1699,8 @@ void FFindInBlueprintSearchManager::ExtractUnloadedFiBData(const FAssetData& InA
 	NewSearchData.bMarkedForDeletion = false;
 	NewSearchData.Value = *InFiBData;
 
-	// Deserialize the version if available
-	if (bIsVersioned)
-	{
-		checkf(NewSearchData.Value.Len(), TEXT("Versioned search data was zero length!"));
-		FBufferReader ReaderStream((void*)*NewSearchData.Value, NewSearchData.Value.Len() * sizeof(TCHAR), false);
-		NewSearchData.VersionInfo.FiBDataVersion = FiBSerializationHelpers::Deserialize<int32>(ReaderStream);
-	}
-
-	// Determine the editor object version that the asset package was last serialized with
-	FString PackageFilename;
-	if (ensureMsgf(FPackageName::DoesPackageExist(InAssetData.PackageName.ToString(), nullptr, &PackageFilename), TEXT("FiB: Failed to map package to filename.")))
-	{
-		// Open a new file archive for reading
-		FArchive* PackageFile = IFileManager::Get().CreateFileReader(*PackageFilename);
-		if (ensureMsgf(PackageFile != nullptr, TEXT("FiB: Unable to open package to read file summary.")))
-		{
-			// Read the package file summary
-			FPackageFileSummary PackageFileSummary;
-			*PackageFile << PackageFileSummary;
-
-			// Close the file
-			delete PackageFile;
-
-			// If an editor object version exists in the package file summary, record it
-			if (const FCustomVersion* const EditorObjectVersion = PackageFileSummary.GetCustomVersionContainer().GetVersion(FEditorObjectVersion::GUID))
-			{
-				NewSearchData.VersionInfo.EditorObjectVersion = EditorObjectVersion->Version;
-			}
-		}
-	}
+	// This will be set to 'None' if the data is versioned. Deserialization of the actual version from the tag value is deferred until later.
+	NewSearchData.VersionInfo.FiBDataVersion = InFiBDataVersion;
 
 	// Since the asset was not loaded, pull out the searchable data stored in the asset
 	AddSearchDataToDatabase(MoveTemp(NewSearchData));
@@ -1961,8 +1985,12 @@ bool FFindInBlueprintSearchManager::ContinueSearchQuery(const FStreamSearch* InS
 				// If there is FiB data, parse it into an ImaginaryBlueprint
 				if (SearchArray[SearchIdx].Value.Len() > 0)
 				{
-					SearchArray[SearchIdx].ImaginaryBlueprint = MakeShareable(new FImaginaryBlueprint(FPaths::GetBaseFilename(SearchArray[SearchIdx].AssetPath.ToString()), SearchArray[SearchIdx].AssetPath.ToString(), SearchArray[SearchIdx].ParentClass, SearchArray[SearchIdx].Interfaces, SearchArray[SearchIdx].Value, SearchArray[SearchIdx].VersionInfo));
-					SearchArray[SearchIdx].Value.Empty();
+					const FString AssetPath = SearchArray[SearchIdx].AssetPath.ToString();
+					if (FiBSerializationHelpers::ValidateSearchDataVersionInfo(AssetPath, SearchArray[SearchIdx].Value, SearchArray[SearchIdx].VersionInfo))
+					{
+						SearchArray[SearchIdx].ImaginaryBlueprint = MakeShareable(new FImaginaryBlueprint(FPaths::GetBaseFilename(AssetPath), AssetPath, SearchArray[SearchIdx].ParentClass, SearchArray[SearchIdx].Interfaces, SearchArray[SearchIdx].Value, SearchArray[SearchIdx].VersionInfo));
+						SearchArray[SearchIdx].Value.Empty();
+					}
 				}
 
  				OutSearchData = SearchArray[SearchIdx++];
@@ -2487,6 +2515,9 @@ TSharedPtr< FJsonObject > FFindInBlueprintSearchManager::ConvertJsonStringToObje
 	int32 SizeOfData;
 	FBufferReader ReaderStream((void*)*InJsonString, InJsonString.Len() * sizeof(TCHAR), false);
 
+	// It's unexpected to not have valid search data versioning info at this point
+	checkf(FiBSerializationHelpers::ValidateSearchDataVersionInfo(InVersionInfo), TEXT("FiB: Invalid JSON stream data version."));
+
 	// If the stream is versioned, read past the version info
 	if (InVersionInfo.FiBDataVersion > EFiBVersion::FIB_VER_BASE)
 	{
diff --git a/Engine/Source/Editor/Kismet/Public/FindInBlueprintManager.h b/Engine/Source/Editor/Kismet/Public/FindInBlueprintManager.h
--- a/Engine/Source/Editor/Kismet/Public/FindInBlueprintManager.h
+++ b/Engine/Source/Editor/Kismet/Public/FindInBlueprintManager.h
@@ -105,6 +105,8 @@ struct KISMET_API FFindInBlueprintSearchTags
 /** FiB data versioning */
 enum EFiBVersion
 {
+	FIB_VER_NONE = -1, // Unknown version (not set)
+
 	FIB_VER_BASE = 0, // All Blueprints prior to versioning will automatically be assumed to be at 0 if they have FiB data collected
 	FIB_VER_VARIABLE_REFERENCE, // Variable references (FMemberReference) is collected in FiB
 	FIB_VER_INTERFACE_GRAPHS, // Implemented Interface Graphs is collected in FiB
@@ -118,7 +120,7 @@ enum EFiBVersion
 struct FSearchDataVersionInfo
 {
 	/** FiB asset registry tag data version */
-	int32 FiBDataVersion = EFiBVersion::FIB_VER_BASE;
+	int32 FiBDataVersion = EFiBVersion::FIB_VER_NONE;
 
 	/** Editor object version used to serialize values in the JSON string lookup table */
 	int32 EditorObjectVersion = -1;
@@ -593,7 +595,7 @@ private:
 	void RemoveBlueprintByPath(FName InPath);
 
 	/** Begins the process of extracting unloaded FiB data */
-	void ExtractUnloadedFiBData(const FAssetData& InAssetData, const FString& InFiBData, bool bIsVersioned);
+	void ExtractUnloadedFiBData(const FAssetData& InAssetData, const FString& InFiBData, EFiBVersion InFiBDataVersion);
 
 	/** Determines the global find results tab label */
 	FText GetGlobalFindResultsTabLabel(int32 TabIdx);
