Animation interface graphs now display sub graphs in the My Blueprint tab

diff --git a/Engine/Source/Editor/Kismet/Private/SMyBlueprint.cpp b/Engine/Source/Editor/Kismet/Private/SMyBlueprint.cpp
--- a/Engine/Source/Editor/Kismet/Private/SMyBlueprint.cpp
+++ b/Engine/Source/Editor/Kismet/Private/SMyBlueprint.cpp
@@ -1407,6 +1407,11 @@ void SMyBlueprint::CollectAllActions(FGraphActionListBuilderBase& OutAllActions)
 						OverridableFunctionActions.Add(NewFuncAction);
 						OverridableFunctionNames.Add(FunctionName);
 					}
+
+					if(bIsAnimFunction && NewFuncAction->EdGraph)
+					{
+						GetChildGraphs(NewFuncAction->EdGraph, NewFuncAction->GetSectionID(), SortList, FunctionCategory);
+					}
 				}
 			}
 		}
