Fix Copy/Select All in global FIB window

diff --git a/Engine/Source/Editor/Kismet/Private/FindInBlueprintManager.cpp b/Engine/Source/Editor/Kismet/Private/FindInBlueprintManager.cpp
--- a/Engine/Source/Editor/Kismet/Private/FindInBlueprintManager.cpp
+++ b/Engine/Source/Editor/Kismet/Private/FindInBlueprintManager.cpp
@@ -2598,7 +2598,6 @@ TSharedRef<SDockTab> FFindInBlueprintSearchManager::SpawnGlobalFindResultsTab(co
 		.ToolTipText(LOCTEXT("GlobalFindResultsTabTooltip", "Search for a string in all Blueprint assets."));
 
 	TSharedRef<SFindInBlueprints> FindResults = SNew(SFindInBlueprints)
-		.bIsSearchWindow(false)
 		.ContainingTab(NewTab);
 
 	GlobalFindResults.Add(FindResults);
diff --git a/Engine/Source/Editor/Kismet/Private/FindInBlueprints.cpp b/Engine/Source/Editor/Kismet/Private/FindInBlueprints.cpp
--- a/Engine/Source/Editor/Kismet/Private/FindInBlueprints.cpp
+++ b/Engine/Source/Editor/Kismet/Private/FindInBlueprints.cpp
@@ -1570,4 +1570,18 @@ void SFindInBlueprints::CloseHostTab()
 	}
 }
 
+FReply SFindInBlueprints::OnKeyDown(const FGeometry& MyGeometry, const FKeyEvent& InKeyEvent)
+{
+	// BlueprintEditor's IToolkit code will handle shortcuts itself - but we can just use 
+	// simple slate handlers when we're standalone:
+	if(!BlueprintEditorPtr.IsValid() && CommandList.IsValid())
+	{
+		if( CommandList->ProcessCommandBindings(InKeyEvent))
+		{
+			return FReply::Handled();
+		}
+	}
+	return SCompoundWidget::OnKeyDown(MyGeometry, InKeyEvent);
+}
+
 #undef LOCTEXT_NAMESPACE
diff --git a/Engine/Source/Editor/Kismet/Public/FindInBlueprints.h b/Engine/Source/Editor/Kismet/Public/FindInBlueprints.h
--- a/Engine/Source/Editor/Kismet/Public/FindInBlueprints.h
+++ b/Engine/Source/Editor/Kismet/Public/FindInBlueprints.h
@@ -258,6 +258,9 @@ public:
 		return bIsLocked;
 	}
 
+	/** SWidget overrides */
+	virtual FReply OnKeyDown(const FGeometry& MyGeometry, const FKeyEvent& InKeyEvent) override;
+
 private:
 	/** Processes results of the ongoing async stream search */
 	EActiveTimerReturnType UpdateSearchResults( double InCurrentTime, float InDeltaTime );
