Fix for BP Diff menu - adding a background for the scrollbar

diff --git a/Engine/Source/Editor/Kismet/Private/SBlueprintRevisionMenu.cpp b/Engine/Source/Editor/Kismet/Private/SBlueprintRevisionMenu.cpp
--- a/Engine/Source/Editor/Kismet/Private/SBlueprintRevisionMenu.cpp
+++ b/Engine/Source/Editor/Kismet/Private/SBlueprintRevisionMenu.cpp
@@ -242,7 +242,12 @@ void SBlueprintRevisionMenu::OnSourceControlQueryComplete(const FSourceControlOp
 	MenuBuilder.EndSection();
 	MenuBox->AddSlot() 
 	[
-		MenuBuilder.MakeWidget(nullptr, 500)
+		SNew(SBorder)
+		.BorderImage(FEditorStyle::GetBrush("Menu.Background"))
+		.Content()
+		[
+			MenuBuilder.MakeWidget(nullptr, 500)
+		]
 	];
 
 	SourceControlQueryOp.Reset();
