Don't open a new event graph view in the outgoing function graph's tab after converting a function to an event.

#rnx

diff --git a/Engine/Source/Editor/Kismet/Private/BlueprintEditor.cpp b/Engine/Source/Editor/Kismet/Private/BlueprintEditor.cpp
--- a/Engine/Source/Editor/Kismet/Private/BlueprintEditor.cpp
+++ b/Engine/Source/Editor/Kismet/Private/BlueprintEditor.cpp
@@ -5529,19 +5529,25 @@ void FBlueprintEditor::ConvertFunctionToEvent(UK2Node_FunctionEntry* SelectedCal
 		{
 			Result->DestroyNode();
 		}
+
 		// Connect any pins that need to be set from the old function
 		if (NewEventNode)
 		{
 			// Link the nodes from the original function entry node to the new event node
 			FEdGraphUtilities::ReconnectPinMap(NewEventNode, PinConnections);
 			FEdGraphUtilities::CopyPinDefaults(SelectedCallFunctionNode, NewEventNode);
-			FKismetEditorUtilities::BringKismetToFocusAttentionOnObject(NewEventNode, false);
 		}
 
 		// Remove the old function graph
 		FBlueprintEditorUtils::RemoveGraph(NodeBP, FunctionGraph, EGraphRemoveFlags::Recompile);
 		FunctionGraph->MarkPendingKill();
 		FBlueprintEditorUtils::MarkBlueprintAsStructurallyModified(NodeBP);
+
+		// Do this AFTER removing the function graph so that it's not opened into the existing function graph document tab.
+		if (NewEventNode)
+		{
+			FKismetEditorUtilities::BringKismetToFocusAttentionOnObject(NewEventNode, false);
+		}
 	}
 }
 
