Destroy components immediately when deleted from the SCS, in much the same way that renaming affects them immediately.

diff --git a/Engine/Source/Editor/Kismet/Private/SSCSEditor.cpp b/Engine/Source/Editor/Kismet/Private/SSCSEditor.cpp
--- a/Engine/Source/Editor/Kismet/Private/SSCSEditor.cpp
+++ b/Engine/Source/Editor/Kismet/Private/SSCSEditor.cpp
@@ -5942,6 +5942,22 @@ void SSCSEditor::RemoveComponentNode(FSCSEditorTreeNodePtrType InNodePtr)
 				SCS_Node->ComponentTemplate->Modify();
 				SCS_Node->ComponentTemplate->Rename(*RemovedName, /*NewOuter =*/nullptr, REN_DontCreateRedirectors);
 
+				TArray<UObject*> ArchetypeInstances;
+				auto DestroyArchetypeInstances = [&ArchetypeInstances, &RemovedName](UActorComponent* ComponentTemplate)
+				{
+					ComponentTemplate->GetArchetypeInstances(ArchetypeInstances);
+					for (UObject* ArchetypeInstance : ArchetypeInstances)
+					{
+						if (!ArchetypeInstance->HasAllFlags(RF_ArchetypeObject | RF_InheritableComponentTemplate))
+						{
+							CastChecked<UActorComponent>(ArchetypeInstance)->DestroyComponent();
+							ArchetypeInstance->Rename(*RemovedName, nullptr, REN_DontCreateRedirectors);
+						}
+					}
+				};
+
+				DestroyArchetypeInstances(SCS_Node->ComponentTemplate);
+				
 				if (Blueprint)
 				{
 					// Children need to have their inherited component template instance of the component renamed out of the way as well
@@ -5956,6 +5972,8 @@ void SSCSEditor::RemoveComponentNode(FSCSEditorTreeNodePtrType InNodePtr)
 						{
 							Component->Modify();
 							Component->Rename(*RemovedName, /*NewOuter =*/nullptr, REN_DontCreateRedirectors);
+
+							DestroyArchetypeInstances(Component);
 						}
 					}
 				}
