Fix nodes refreshing when using 'step over' function in the blueprint editor
#rnx

diff --git a/Engine/Source/Editor/Kismet/Private/BlueprintEditor.cpp b/Engine/Source/Editor/Kismet/Private/BlueprintEditor.cpp
--- a/Engine/Source/Editor/Kismet/Private/BlueprintEditor.cpp
+++ b/Engine/Source/Editor/Kismet/Private/BlueprintEditor.cpp
@@ -8652,8 +8652,17 @@ TSharedPtr<SGraphEditor> FBlueprintEditor::OpenGraphAndBringToFront(UEdGraph* Gr
 	// First, switch back to standard mode
 	SetCurrentMode(FBlueprintEditorApplicationModes::StandardBlueprintEditorMode);
 
-	// Next, try to make sure there is a copy open
-	TSharedPtr<SDockTab> TabWithGraph = OpenDocument(Graph, FDocumentTracker::NavigatingCurrentDocument);
+		
+	TSharedPtr<SDockTab> TabWithGraph;
+	TSharedPtr<SGraphEditor> FocusedGraphEd = FocusedGraphEdPtr.Pin();
+	if (FocusedGraphEd.IsValid() && FocusedGraphEd->GetCurrentGraph() == Graph)
+	{
+		TabWithGraph = OpenDocument(Graph, FDocumentTracker::CreateHistoryEvent);
+	}
+	else
+	{
+		 TabWithGraph = OpenDocument(Graph, FDocumentTracker::NavigatingCurrentDocument);
+	}
 
 	// We know that the contents of the opened tabs will be a graph editor.
 	TSharedRef<SGraphEditor> NewGraphEditor = StaticCastSharedRef<SGraphEditor>(TabWithGraph->GetContent());
