Fixed an issue where a tab could remain open for a function after the function is deleted if there were several copies of the same graph open at once

diff --git a/Engine/Source/Editor/Kismet/Private/WorkflowOrientedApp/WorkflowTabManager.cpp b/Engine/Source/Editor/Kismet/Private/WorkflowOrientedApp/WorkflowTabManager.cpp
--- a/Engine/Source/Editor/Kismet/Private/WorkflowOrientedApp/WorkflowTabManager.cpp
+++ b/Engine/Source/Editor/Kismet/Private/WorkflowOrientedApp/WorkflowTabManager.cpp
@@ -645,7 +645,7 @@ TSharedPtr<SDockTab> FDocumentTracker::OpenNewTab(TSharedPtr<FTabPayload> InPayl
 
 void FDocumentTracker::CloseTab(TSharedPtr<FTabPayload> Payload)
 {
-	FTabList& List = GetSpawnedList();
+	FTabList List = GetSpawnedList();
 	for (auto ListIt = List.CreateIterator(); ListIt; ++ListIt)
 	{
 		if ((*ListIt)->PayloadMatches(Payload))
