Clarified separator for UserDefinedStructure

diff --git a/Engine/Source/Editor/Kismet/Private/UserDefinedStructureEditor.cpp b/Engine/Source/Editor/Kismet/Private/UserDefinedStructureEditor.cpp
--- a/Engine/Source/Editor/Kismet/Private/UserDefinedStructureEditor.cpp
+++ b/Engine/Source/Editor/Kismet/Private/UserDefinedStructureEditor.cpp
@@ -383,9 +383,10 @@ TSharedRef<SDockTab> FUserDefinedStructureEditor::SpawnStructureTab(const FSpawn
 
 	TSharedRef<SSplitter> Splitter = SNew(SSplitter)
 		.Orientation(Orient_Vertical)
-		.PhysicalSplitterHandleSize(10.0f)
-		.ResizeMode(ESplitterResizeMode::FixedPosition);
-
+		.PhysicalSplitterHandleSize(2.0f)
+		.HitDetectionSplitterHandleSize(8.f)
+		.ResizeMode(ESplitterResizeMode::FixedPosition)
+		.Style(FEditorStyle::Get(), "ContentBrowser.Splitter");
 	{
 		// Create a property view
 		FPropertyEditorModule& EditModule = FModuleManager::Get().GetModuleChecked<FPropertyEditorModule>("PropertyEditor");
@@ -398,7 +399,11 @@ TSharedRef<SDockTab> FUserDefinedStructureEditor::SpawnStructureTab(const FSpawn
 		Splitter->AddSlot()
 		.Value(0.25f)
 		[
-			PropertyView.ToSharedRef()
+			SNew(SBox)
+			.Padding(FMargin(0.0f, 0.0f, 0.0f, 3.0f))
+			[
+				PropertyView.ToSharedRef()
+			]
 		];
 	}
 
@@ -414,7 +419,11 @@ TSharedRef<SDockTab> FUserDefinedStructureEditor::SpawnStructureTab(const FSpawn
 		{
 			Splitter->AddSlot()
 			[
-				DefaultValueWidget.ToSharedRef()
+				SNew(SBox)
+				.Padding(FMargin(0.0f, 3.0f, 0.0f, 0.0f))
+				[
+					DefaultValueWidget.ToSharedRef()
+				]
 			];
 		}
 	}
