Fix for user construction script's local variables not showing up in the details panel.
[Cherry picking CL# 9706161 by JonathanL from Dev-Framework]

diff --git a/Engine/Source/Editor/Kismet/Private/BlueprintDetailsCustomization.cpp b/Engine/Source/Editor/Kismet/Private/BlueprintDetailsCustomization.cpp
--- a/Engine/Source/Editor/Kismet/Private/BlueprintDetailsCustomization.cpp
+++ b/Engine/Source/Editor/Kismet/Private/BlueprintDetailsCustomization.cpp
@@ -848,8 +848,20 @@ void FBlueprintVarActionDetails::CustomizeDetails( IDetailLayoutBuilder& DetailL
 				
 				UK2Node_FunctionEntry* FuncEntry = EntryNodes[0];
 
-				// This uses the cached struct data inside the node, which is updated on change
-				TSharedPtr<FStructOnScope> StructData = FuncEntry->GetFunctionVariableCache();
+				TSharedPtr<FStructOnScope> StructData = MakeShareable(new FStructOnScope((UFunction*)StructScope));
+
+				for (const FBPVariableDescription& LocalVar : FuncEntry->LocalVariables)
+				{
+					if (LocalVar.VarName == VariableProperty->GetFName())
+					{
+						// Only set the default value if there is one
+						if (!LocalVar.DefaultValue.IsEmpty())
+						{
+							FBlueprintEditorUtils::PropertyValueFromString(VariableProperty, LocalVar.DefaultValue, StructData->GetStructMemory());
+						}
+						break;
+					}
+				}
 
 				if(BlueprintEditor.IsValid())
 				{
