Guard against a rare crash (regression) in the Blueprint SCS editor seemingly caused by an invalid tree node hierarchy.

#rnx

diff --git a/Engine/Source/Editor/Kismet/Private/SSCSEditor.cpp b/Engine/Source/Editor/Kismet/Private/SSCSEditor.cpp
--- a/Engine/Source/Editor/Kismet/Private/SSCSEditor.cpp
+++ b/Engine/Source/Editor/Kismet/Private/SSCSEditor.cpp
@@ -6313,26 +6313,33 @@ bool SSCSEditor::CanRenameComponent() const
 	return IsEditingAllowed() && SCSTreeWidget->GetSelectedItems().Num() == 1 && SCSTreeWidget->GetSelectedItems()[0]->CanRename();
 }
 
-void SSCSEditor::GetCollapsedNodes(const FSCSEditorTreeNodePtrType& InNodePtr, TSet<FSCSEditorTreeNodePtrType>& OutCollapsedNodes) const
+void SSCSEditor::DepthFirstTraversal(const FSCSEditorTreeNodePtrType& InNodePtr, TSet<FSCSEditorTreeNodePtrType>& OutVisitedNodes, const TFunctionRef<void(const FSCSEditorTreeNodePtrType&)> InFunction) const
 {
-	if(InNodePtr.IsValid())
+	if (InNodePtr.IsValid()
+		&& ensureMsgf(!OutVisitedNodes.Contains(InNodePtr), TEXT("Already visited node: %s (Parent: %s)"), *InNodePtr->GetDisplayString(), InNodePtr->GetParent().IsValid() ? *InNodePtr->GetParent()->GetDisplayString() : TEXT("NULL")))
 	{
-		const TArray<FSCSEditorTreeNodePtrType>& Children = InNodePtr->GetChildren();
-		if(Children.Num() > 0)
-		{
-			if(!SCSTreeWidget->IsItemExpanded(InNodePtr))
-			{
-				OutCollapsedNodes.Add(InNodePtr);
-			}
+		InFunction(InNodePtr);
+		OutVisitedNodes.Add(InNodePtr);
 
-			for(int32 i = 0; i < Children.Num(); ++i)
-			{
-				GetCollapsedNodes(Children[i], OutCollapsedNodes);
-			}
+		for (const FSCSEditorTreeNodePtrType& Child : InNodePtr->GetChildren())
+		{
+			DepthFirstTraversal(Child, OutVisitedNodes, InFunction);
 		}
 	}
 }
 
+void SSCSEditor::GetCollapsedNodes(const FSCSEditorTreeNodePtrType& InNodePtr, TSet<FSCSEditorTreeNodePtrType>& OutCollapsedNodes) const
+{
+	TSet<FSCSEditorTreeNodePtrType> VisitedNodes;
+	DepthFirstTraversal(InNodePtr, VisitedNodes, [SCSTreeWidget = this->SCSTreeWidget, &OutCollapsedNodes](const FSCSEditorTreeNodePtrType& InNodePtr)
+	{
+		if (InNodePtr->GetChildren().Num() > 0 && !SCSTreeWidget->IsItemExpanded(InNodePtr))
+		{
+			OutCollapsedNodes.Add(InNodePtr);
+		}
+	});
+}
+
 EVisibility SSCSEditor::GetPromoteToBlueprintButtonVisibility() const
 {
 	EVisibility ButtonVisibility = EVisibility::Collapsed;
diff --git a/Engine/Source/Editor/Kismet/Public/SSCSEditor.h b/Engine/Source/Editor/Kismet/Public/SSCSEditor.h
--- a/Engine/Source/Editor/Kismet/Public/SSCSEditor.h
+++ b/Engine/Source/Editor/Kismet/Public/SSCSEditor.h
@@ -1093,6 +1093,9 @@ protected:
 	/** Callback when a component item is double clicked. */
 	void HandleItemDoubleClicked(FSCSEditorTreeNodePtrType InItem);
 
+	/** Recursively visits the given node + its children and invokes the given function for each. */
+	void DepthFirstTraversal(const FSCSEditorTreeNodePtrType& InNodePtr, TSet<FSCSEditorTreeNodePtrType>& OutVisitedNodes, const TFunctionRef<void(const FSCSEditorTreeNodePtrType&)> InFunction) const;
+
 	/** Returns the set of expandable nodes that are currently collapsed in the UI */
 	void GetCollapsedNodes(const FSCSEditorTreeNodePtrType& InNodePtr, TSet<FSCSEditorTreeNodePtrType>& OutCollapsedNodes) const;
 
