Control Rig: Ensure to always refresh the details inspector with bForceRefresh

diff --git a/Engine/Plugins/Experimental/ControlRig/Source/ControlRigEditor/Private/Editor/ControlRigEditor.cpp b/Engine/Plugins/Experimental/ControlRig/Source/ControlRigEditor/Private/Editor/ControlRigEditor.cpp
--- a/Engine/Plugins/Experimental/ControlRig/Source/ControlRigEditor/Private/Editor/ControlRigEditor.cpp
+++ b/Engine/Plugins/Experimental/ControlRig/Source/ControlRigEditor/Private/Editor/ControlRigEditor.cpp
@@ -573,6 +573,8 @@ UBlueprint* FControlRigEditor::GetBlueprintObj() const
 void FControlRigEditor::SetDetailObjects(const TArray<UObject*>& InObjects)
 {
 	RigElementInDetailPanel = FRigElementKey();
+	SKismetInspector::FShowDetailsOptions Options;
+	Options.bForceRefresh = true;
 	Inspector->ShowDetailsForObjects(InObjects);
 }
 
diff --git a/Engine/Source/Editor/Kismet/Private/BlueprintEditor.cpp b/Engine/Source/Editor/Kismet/Private/BlueprintEditor.cpp
--- a/Engine/Source/Editor/Kismet/Private/BlueprintEditor.cpp
+++ b/Engine/Source/Editor/Kismet/Private/BlueprintEditor.cpp
@@ -3269,8 +3269,9 @@ void FBlueprintEditor::OnSelectedNodesChangedImpl(const FGraphPanelSelectionSet&
 		SetUISelectionState(FBlueprintEditor::SelectionState_Graph);
 	}
 
-	Inspector->ShowDetailsForObjects(NewSelection.Array());
-
+	SKismetInspector::FShowDetailsOptions DetailsOptions;
+	DetailsOptions.bForceRefresh = true;
+	Inspector->ShowDetailsForObjects(NewSelection.Array(), DetailsOptions);
 
 	bSelectRegularNode = false;
 	for (FGraphPanelSelectionSet::TConstIterator It(NewSelection); It; ++It)
