Add profiling hooks to Blueprint search for A/B testing and consolidated some code.

diff --git a/Engine/Config/BaseEditor.ini b/Engine/Config/BaseEditor.ini
--- a/Engine/Config/BaseEditor.ini
+++ b/Engine/Config/BaseEditor.ini
@@ -431,6 +431,7 @@ ForceUnits=EUnit::Newtons
 
 [BlueprintSearchSettings]
 bDisableDeferredIndexing=false
+bEnableCsvStatsProfiling=false
 
 [BlueprintNativizationSettings]
 bUseStaticVariablesInWrappers=true
diff --git a/Engine/Source/Editor/Kismet/Private/FiBSearchInstance.cpp b/Engine/Source/Editor/Kismet/Private/FiBSearchInstance.cpp
--- a/Engine/Source/Editor/Kismet/Private/FiBSearchInstance.cpp
+++ b/Engine/Source/Editor/Kismet/Private/FiBSearchInstance.cpp
@@ -257,6 +257,8 @@ void FFiBSearchInstance::BuildFunctionTargetsByName(FImaginaryFiBDataSharedPtr I
 
 bool FFiBSearchInstance::OnFilterFunction(const FTextFilterString& A, ESearchQueryFilter InSearchQueryFilter)
 {
+	CSV_SCOPED_TIMING_STAT(FindInBlueprint, OnFilterFunction);
+
 	if (CurrentSearchable.IsValid())
 	{
 		TSharedPtr< FFiBSearchInstance > SubSearchInstance(new FFiBSearchInstance);
@@ -299,6 +301,8 @@ bool FFiBSearchInstance::OnFilterFunction(const FTextFilterString& A, ESearchQue
 
 bool FFiBSearchInstance::OnFilterDefaultFunction(const FTextFilterString& InFunctionName, const FTextFilterString& InFunctionParams)
 {
+	CSV_SCOPED_TIMING_STAT(FindInBlueprint, OnFilterDefaultFunction);
+
 	if (CurrentSearchable.IsValid())
 	{
 		TSharedPtr< FFiBSearchInstance > SubSearchInstance(new FFiBSearchInstance);
@@ -336,6 +340,8 @@ bool FFiBSearchInstance::OnFilterDefaultFunction(const FTextFilterString& InFunc
 
 bool FFindInBlueprintExpressionEvaluator::EvaluateCompiledExpression(const ExpressionParser::CompileResultType& InCompiledResult, const ITextFilterExpressionContext& InContext, FText* OutErrorText) const
 {
+	CSV_SCOPED_TIMING_STAT(FindInBlueprint, EvaluateCompiledExpression);
+
 	using namespace TextFilterExpressionParser;
 	using TextFilterExpressionParser::FTextToken;
 
diff --git a/Engine/Source/Editor/Kismet/Private/FindInBlueprintManager.cpp b/Engine/Source/Editor/Kismet/Private/FindInBlueprintManager.cpp
--- a/Engine/Source/Editor/Kismet/Private/FindInBlueprintManager.cpp
+++ b/Engine/Source/Editor/Kismet/Private/FindInBlueprintManager.cpp
@@ -53,6 +53,9 @@
 #include "Framework/Notifications/NotificationManager.h"
 #include "Widgets/Notifications/SNotificationList.h"
 
+DEFINE_LOG_CATEGORY(LogFindInBlueprint);
+CSV_DEFINE_CATEGORY(FindInBlueprint, false);
+
 #define LOCTEXT_NAMESPACE "FindInBlueprintManager"
 
 FFindInBlueprintSearchManager* FFindInBlueprintSearchManager::Instance = NULL;
@@ -123,28 +126,20 @@ FSearchDataVersionInfo FSearchDataVersionInfo::Current =
 
 ////////////////////////////////////
 // FStreamSearch
-FStreamSearch::FStreamSearch(const FString& InSearchValue)
-	: SearchValue(InSearchValue)
-	, bThreadCompleted(false)
-	, StopTaskCounter(0)
-	, MinimiumVersionRequirement(EFiBVersion::FIB_VER_LATEST)
-	, BlueprintCountBelowVersion(0)
-	, ImaginaryDataFilter(ESearchQueryFilter::AllFilter)
-{
-	// Add on a Guid to the thread name to ensure the thread is uniquely named.
-	Thread = FRunnableThread::Create( this, *FString::Printf(TEXT("FStreamSearch%s"), *FGuid::NewGuid().ToString()), 0, TPri_BelowNormal );
-}
 
-FStreamSearch::FStreamSearch(const FString& InSearchValue, ESearchQueryFilter InImaginaryDataFilter, EFiBVersion InMinimiumVersionRequirement)
+FStreamSearch::FStreamSearch(const FString& InSearchValue, const FStreamSearchOptions& InSearchOptions)
 	: SearchValue(InSearchValue)
-	, bThreadCompleted(false)
+	, SearchOptions(InSearchOptions)
 	, StopTaskCounter(0)
-	, MinimiumVersionRequirement(InMinimiumVersionRequirement)
 	, BlueprintCountBelowVersion(0)
-	, ImaginaryDataFilter(InImaginaryDataFilter)
+	, bThreadCompleted(false)
 {
-	// Add on a Guid to the thread name to ensure the thread is uniquely named.
-	Thread = FRunnableThread::Create( this, *FString::Printf(TEXT("FStreamSearch%s"), *FGuid::NewGuid().ToString()), 0, TPri_BelowNormal );
+	// Unique identifier for this search, used to generate a unique label for profiling and debugging.
+	static int32 GlobalSearchCounter = 0;
+	SearchId = GlobalSearchCounter++;
+
+	// Create a uniquely-named thread to ensure disambiguation in the thread view and assist with debugging multiple instances running in parallel.
+	Thread = FRunnableThread::Create( this, *FString::Printf(TEXT("FStreamSearch_%d"), SearchId), 0, TPri_BelowNormal );
 }
 
 bool FStreamSearch::Init()
@@ -154,6 +149,9 @@ bool FStreamSearch::Init()
 
 uint32 FStreamSearch::Run()
 {
+	const double StartTime = FPlatformTime::Seconds();
+	CSV_EVENT(FindInBlueprint, TEXT("FStreamSearch_%d START"), SearchId);
+
 	FFindInBlueprintSearchManager::Get().BeginSearchQuery(this);
 
 	TFunction<void(const FSearchResult&)> OnResultReady = [this](const FSearchResult& Result) {
@@ -168,17 +166,17 @@ uint32 FStreamSearch::Run()
 		if (QueryResult.ImaginaryBlueprint.IsValid())
 		{
 			// If the Blueprint is below the version, add it to a list. The search will still proceed on this Blueprint
-			if (QueryResult.VersionInfo.FiBDataVersion < MinimiumVersionRequirement)
+			if (QueryResult.VersionInfo.FiBDataVersion < SearchOptions.MinimiumVersionRequirement)
 			{
 				++BlueprintCountBelowVersion;
 			}
 
 			TSharedPtr< FFiBSearchInstance > SearchInstance(new FFiBSearchInstance);
 			FSearchResult SearchResult;
-			if (ImaginaryDataFilter != ESearchQueryFilter::AllFilter)
+			if (SearchOptions.ImaginaryDataFilter != ESearchQueryFilter::AllFilter)
 			{
 				SearchInstance->MakeSearchQuery(*SearchValue, QueryResult.ImaginaryBlueprint);
-				SearchInstance->CreateFilteredResultsListFromTree(ImaginaryDataFilter, FilteredImaginaryResults);
+				SearchInstance->CreateFilteredResultsListFromTree(SearchOptions.ImaginaryDataFilter, FilteredImaginaryResults);
 				SearchResult = SearchInstance->GetSearchResults(QueryResult.ImaginaryBlueprint);
 			}
 			else
@@ -202,6 +200,9 @@ uint32 FStreamSearch::Run()
 
 	bThreadCompleted = true;
 
+	CSV_EVENT(FindInBlueprint, TEXT("FStreamSearch_%d END"), SearchId);
+	UE_LOG(LogFindInBlueprint, Log, TEXT("Search completed in %0.2f seconds."), FPlatformTime::Seconds() - StartTime);
+
 	return 0;
 }
 
@@ -378,6 +379,8 @@ namespace FiBSerializationHelpers
 	/** Helper function to validate and/or deserialize version info if necessary */
 	bool ValidateSearchDataVersionInfo(const FString& InAssetPath, const FString& InFiBData, FSearchDataVersionInfo& InOutVersionInfo)
 	{
+		CSV_SCOPED_TIMING_STAT(FindInBlueprint, ValidateSearchDataVersionInfo);
+
 		if (InOutVersionInfo.FiBDataVersion == EFiBVersion::FIB_VER_NONE)
 		{
 			// Deserialize the FiB data version
@@ -1485,11 +1488,12 @@ FFindInBlueprintSearchManager& FFindInBlueprintSearchManager::Get()
 }
 
 FFindInBlueprintSearchManager::FFindInBlueprintSearchManager()
-	: bEnableGatheringData(true)
+	: AssetRegistryModule(nullptr)
+	, CachingObject(nullptr)
+	, bEnableGatheringData(true)
 	, bDisableDeferredIndexing(false)
+	, bEnableCSVStatsProfiling(false)
 	, bIsPausing(false)
-	, AssetRegistryModule(nullptr)
-	, CachingObject(nullptr)
 {
 	for (int32 TabIdx = 0; TabIdx < UE_ARRAY_COUNT(GlobalFindResultsTabIDs); TabIdx++)
 	{
@@ -1505,6 +1509,7 @@ FFindInBlueprintSearchManager::~FFindInBlueprintSearchManager()
 		AssetRegistryModule->Get().OnAssetAdded().RemoveAll(this);
 		AssetRegistryModule->Get().OnAssetRemoved().RemoveAll(this);
 		AssetRegistryModule->Get().OnAssetRenamed().RemoveAll(this);
+		AssetRegistryModule->Get().OnFilesLoaded().RemoveAll(this);
 	}
 	FKismetEditorUtilities::OnBlueprintUnloaded.RemoveAll(this);
 	FCoreUObjectDelegates::GetPreGarbageCollectDelegate().RemoveAll(this);
@@ -1525,6 +1530,19 @@ void FFindInBlueprintSearchManager::Initialize()
 {
 	// Init configuration
 	GConfig->GetBool(TEXT("BlueprintSearchSettings"), TEXT("bDisableDeferredIndexing"), bDisableDeferredIndexing, GEditorIni);
+	GConfig->GetBool(TEXT("BlueprintSearchSettings"), TEXT("bEnableCsvStatsProfiling"), bEnableCSVStatsProfiling, GEditorIni);
+
+	// If profiling has been enabled, turn on the stat category and begin a capture.
+	if (bEnableCSVStatsProfiling)
+	{
+		FCsvProfiler::Get()->EnableCategoryByString(TEXT("FindInBlueprint"));
+		if (!FCsvProfiler::Get()->IsCapturing())
+		{
+			const FString CaptureFolder = FPaths::ProfilingDir() + TEXT("CSV/FindInBlueprint");
+			FCsvProfiler::Get()->BeginCapture(-1, CaptureFolder);
+		}
+	}
+
 
 	// Must ensure we do not attempt to load the AssetRegistry Module while saving a package, however, if it is loaded already we can safely obtain it
 	if (!GIsSavingPackage || (GIsSavingPackage && FModuleManager::Get().IsModuleLoaded(TEXT("AssetRegistry"))))
@@ -1533,6 +1551,7 @@ void FFindInBlueprintSearchManager::Initialize()
 		AssetRegistryModule->Get().OnAssetAdded().AddRaw(this, &FFindInBlueprintSearchManager::OnAssetAdded);
 		AssetRegistryModule->Get().OnAssetRemoved().AddRaw(this, &FFindInBlueprintSearchManager::OnAssetRemoved);
 		AssetRegistryModule->Get().OnAssetRenamed().AddRaw(this, &FFindInBlueprintSearchManager::OnAssetRenamed);
+		AssetRegistryModule->Get().OnFilesLoaded().AddRaw(this, &FFindInBlueprintSearchManager::OnAssetRegistryFilesLoaded);
 	}
 	else
 	{
@@ -1640,6 +1659,9 @@ void FFindInBlueprintSearchManager::OnAssetAdded(const FAssetData& InAssetData)
 
 void FFindInBlueprintSearchManager::ExtractUnloadedFiBData(const FAssetData& InAssetData, const FString& InFiBData, EFiBVersion InFiBDataVersion)
 {
+	CSV_SCOPED_TIMING_STAT(FindInBlueprint, ExtractUnloadedFiBData);
+	CSV_CUSTOM_STAT(FindInBlueprint, ExtractUnloadedCountThisFrame, 1, ECsvCustomStatOp::Accumulate);
+
 	if (SearchMap.Contains(InAssetData.ObjectPath))
 	{
 		return;
@@ -1744,6 +1766,11 @@ void FFindInBlueprintSearchManager::OnAssetRenamed(const struct FAssetData& InAs
 	}
 }
 
+void FFindInBlueprintSearchManager::OnAssetRegistryFilesLoaded()
+{
+	CSV_EVENT(FindInBlueprint, TEXT("OnAssetRegistryFilesLoaded"));
+}
+
 void FFindInBlueprintSearchManager::OnAssetLoaded(UObject* InAsset)
 {
 	const IBlueprintAssetHandler* Handler = FBlueprintAssetHandler::Get().FindHandler(InAsset->GetClass());
@@ -1782,7 +1809,9 @@ void FFindInBlueprintSearchManager::OnHotReload(bool bWasTriggeredAutomatically)
 }
 
 FString FFindInBlueprintSearchManager::GatherBlueprintSearchMetadata(const UBlueprint* Blueprint)
-{	
+{
+	CSV_SCOPED_TIMING_STAT(FindInBlueprint, GatherBlueprintSearchMetadata);
+
 	FTemporarilyUseFriendlyNodeTitles TemporarilyUseFriendlyNodeTitles;
 
 	FString SearchMetaData;
@@ -1874,6 +1903,9 @@ FString FFindInBlueprintSearchManager::GatherBlueprintSearchMetadata(const UBlue
 
 void FFindInBlueprintSearchManager::AddOrUpdateBlueprintSearchMetadata(UBlueprint* InBlueprint, bool bInForceReCache/* = false*/)
 {
+	CSV_SCOPED_TIMING_STAT(FindInBlueprint, AddOrUpdateBlueprintSearchMetadata);
+	CSV_CUSTOM_STAT(FindInBlueprint, AddOrUpdateCountThisFrame, 1, ECsvCustomStatOp::Accumulate);
+
 	// No need to update the cache in the following cases:
 	//	a) Indexing is disabled.
 	//	b) The Blueprint is not yet fully loaded. This ensures that we don't make attempts to re-index before load completion.
@@ -2548,6 +2580,8 @@ TSharedPtr< FJsonObject > FFindInBlueprintSearchManager::ConvertJsonStringToObje
 
 void FFindInBlueprintSearchManager::GlobalFindResultsClosed(const TSharedRef<SFindInBlueprints>& FindResults)
 {
+	CSV_EVENT(FindInBlueprint, TEXT("GlobalFindResultsClosed: %s"), *FindResults->GetHostTabId().ToString());
+
 	for (TWeakPtr<SFindInBlueprints> FindResultsPtr : GlobalFindResults)
 	{
 		if (FindResultsPtr.Pin() == FindResults)
@@ -2590,6 +2624,8 @@ FText FFindInBlueprintSearchManager::GetGlobalFindResultsTabLabel(int32 TabIdx)
 
 TSharedRef<SDockTab> FFindInBlueprintSearchManager::SpawnGlobalFindResultsTab(const FSpawnTabArgs& SpawnTabArgs, int32 TabIdx)
 {
+	CSV_EVENT(FindInBlueprint, TEXT("SpawnGlobalFindResultsTab: %d"), TabIdx);
+
 	TAttribute<FText> Label = TAttribute<FText>::Create(TAttribute<FText>::FGetter::CreateRaw(this, &FFindInBlueprintSearchManager::GetGlobalFindResultsTabLabel, TabIdx));
 
 	TSharedRef<SDockTab> NewTab = SNew(SDockTab)
@@ -2746,6 +2782,8 @@ void FFindInBlueprintSearchManager::CloseOrphanedGlobalFindResultsTabs(TSharedPt
 
 void FFindInBlueprintSearchManager::Tick(float DeltaTime)
 {
+	CSV_SCOPED_TIMING_STAT(FindInBlueprint, Tick);
+
 	if(IsCacheInProgress())
 	{
 		check(CachingObject);
diff --git a/Engine/Source/Editor/Kismet/Private/FindInBlueprints.cpp b/Engine/Source/Editor/Kismet/Private/FindInBlueprints.cpp
--- a/Engine/Source/Editor/Kismet/Private/FindInBlueprints.cpp
+++ b/Engine/Source/Editor/Kismet/Private/FindInBlueprints.cpp
@@ -1027,10 +1027,10 @@ void SFindInBlueprints::FocusForUse(bool bSetFindWithinBlueprint, FString NewSea
 	}
 }
 
-void SFindInBlueprints::MakeSearchQuery(FString InSearchString, bool bInIsFindWithinBlueprint, enum ESearchQueryFilter InSearchFilterForImaginaryDataReturn/* = ESearchQueryFilter::AllFilter*/, EFiBVersion InMinimiumVersionRequirement/* = EFiBVersion::FIB_VER_LATEST*/, FOnSearchComplete InOnSearchComplete/* = FOnSearchComplete()*/)
+void SFindInBlueprints::MakeSearchQuery(FString InSearchString, bool bInIsFindWithinBlueprint, const FStreamSearchOptions& InSearchOptions/* = FStreamSearchOptions()*/, FOnSearchComplete InOnSearchComplete/* = FOnSearchComplete()*/)
 {
 	SearchTextField->SetText(FText::FromString(InSearchString));
-	LastSearchedFiBVersion = InMinimiumVersionRequirement;
+	LastSearchedFiBVersion = InSearchOptions.MinimiumVersionRequirement;
 
 	if(ItemsFound.Num())
 	{
@@ -1053,6 +1053,8 @@ void SFindInBlueprints::MakeSearchQuery(FString InSearchString, bool bInIsFindWi
 
 		if (bInIsFindWithinBlueprint)
 		{
+			const double StartTime = FPlatformTime::Seconds();
+
 			if(StreamSearch.IsValid() && !StreamSearch->IsComplete())
 			{
 				StreamSearch->Stop();
@@ -1116,10 +1118,12 @@ void SFindInBlueprints::MakeSearchQuery(FString InSearchString, bool bInIsFindWi
 			}
 
 			TreeView->RequestTreeRefresh();
+
+			UE_LOG(LogFindInBlueprint, Log, TEXT("Search completed in %0.2f seconds."), FPlatformTime::Seconds() - StartTime);
 		}
 		else
 		{
-			LaunchStreamThread(InSearchString, InSearchFilterForImaginaryDataReturn, InMinimiumVersionRequirement, InOnSearchComplete);
+			LaunchStreamThread(InSearchString, InSearchOptions, InOnSearchComplete);
 		}
 	}
 }
@@ -1147,24 +1151,7 @@ ECheckBoxState SFindInBlueprints::OnGetFindModeChecked() const
 	return bIsInFindWithinBlueprintMode ? ECheckBoxState::Checked : ECheckBoxState::Unchecked;
 }
 
-void SFindInBlueprints::LaunchStreamThread(const FString& InSearchValue)
-{
-	if(StreamSearch.IsValid() && !StreamSearch->IsComplete())
-	{
-		StreamSearch->Stop();
-		StreamSearch->EnsureCompletion();
-	}
-	else
-	{
-		// If the stream search wasn't already running, register the active timer
-		RegisterActiveTimer( 0.f, FWidgetActiveTimerDelegate::CreateSP( this, &SFindInBlueprints::UpdateSearchResults ) );
-	}
-
-	StreamSearch = MakeShareable(new FStreamSearch(InSearchValue));
-	OnSearchComplete = FOnSearchComplete();
-}
-
-void SFindInBlueprints::LaunchStreamThread(const FString& InSearchValue, enum ESearchQueryFilter InSearchFilterForRawDataReturn, EFiBVersion InMinimiumVersionRequirement, FOnSearchComplete InOnSearchComplete)
+void SFindInBlueprints::LaunchStreamThread(const FString& InSearchValue, const FStreamSearchOptions& InSearchOptions, FOnSearchComplete InOnSearchComplete)
 {
 	if(StreamSearch.IsValid() && !StreamSearch->IsComplete())
 	{
@@ -1177,7 +1164,7 @@ void SFindInBlueprints::LaunchStreamThread(const FString& InSearchValue, enum ES
 		RegisterActiveTimer( 0.f, FWidgetActiveTimerDelegate::CreateSP( this, &SFindInBlueprints::UpdateSearchResults ) );
 	}
 
-	StreamSearch = MakeShareable(new FStreamSearch(InSearchValue, InSearchFilterForRawDataReturn, InMinimiumVersionRequirement));
+	StreamSearch = MakeShared<FStreamSearch>(InSearchValue, InSearchOptions);
 	OnSearchComplete = InOnSearchComplete;
 }
 
diff --git a/Engine/Source/Editor/Kismet/Private/ImaginaryBlueprintData.cpp b/Engine/Source/Editor/Kismet/Private/ImaginaryBlueprintData.cpp
--- a/Engine/Source/Editor/Kismet/Private/ImaginaryBlueprintData.cpp
+++ b/Engine/Source/Editor/Kismet/Private/ImaginaryBlueprintData.cpp
@@ -79,6 +79,8 @@ FImaginaryFiBData::FImaginaryFiBData(FImaginaryFiBDataWeakPtr InOuter, TSharedPt
 
 FSearchResult FImaginaryFiBData::CreateSearchResult(FSearchResult InParent) const
 {
+	CSV_SCOPED_TIMING_STAT(FindInBlueprint, CreateSearchResult);
+
 	FSearchResult ReturnSearchResult = CreateSearchResult_Internal(InParent);
 
 	FText DisplayName;
@@ -95,6 +97,9 @@ FSearchResult FImaginaryFiBData::CreateSearchResult(FSearchResult InParent) cons
 
 FSearchResult FImaginaryFiBData::CreateSearchTree(FSearchResult InParentSearchResult, FImaginaryFiBDataWeakPtr InCurrentPointer, TArray< const FImaginaryFiBData* >& InValidSearchResults, TMultiMap< const FImaginaryFiBData*, FComponentUniqueDisplay >& InMatchingSearchComponents)
 {
+	CSV_SCOPED_TIMING_STAT(FindInBlueprint, CreateSearchTree);
+	CSV_CUSTOM_STAT(FindInBlueprint, CreateSearchTreeIterations, 1, ECsvCustomStatOp::Accumulate);
+
 	FImaginaryFiBDataSharedPtr CurrentDataPtr = InCurrentPointer.Pin();
 	if (FImaginaryFiBData* CurrentData = CurrentDataPtr.Get())
 	{
@@ -206,6 +211,9 @@ void FImaginaryFiBData::ParseAllChildData_Internal(ESearchableValueStatus InSear
 
 void FImaginaryFiBData::ParseAllChildData(ESearchableValueStatus InSearchabilityOverride/* = ESearchableValueStatus::Searchable*/)
 {
+	CSV_SCOPED_TIMING_STAT(FindInBlueprint, ParseAllChildData);
+	CSV_CUSTOM_STAT(FindInBlueprint, ParseAllChildDataIterations, 1, ECsvCustomStatOp::Accumulate);
+
 	FScopeLock ScopeLock(&FImaginaryFiBData::ParseChildDataCriticalSection);
 	ParseAllChildData_Internal(InSearchabilityOverride);
 }
diff --git a/Engine/Source/Editor/Kismet/Private/SReplaceNodeReferences.cpp b/Engine/Source/Editor/Kismet/Private/SReplaceNodeReferences.cpp
--- a/Engine/Source/Editor/Kismet/Private/SReplaceNodeReferences.cpp
+++ b/Engine/Source/Editor/Kismet/Private/SReplaceNodeReferences.cpp
@@ -445,7 +445,10 @@ void SReplaceNodeReferences::OnSubmitSearchQuery(bool bFindAndReplace)
 		OnSearchComplete = FOnSearchComplete::CreateSP(this, &SReplaceNodeReferences::FindAllReplacementsComplete);
 	}
 
-	FindInBlueprints->MakeSearchQuery(SearchTerm, false, ESearchQueryFilter::NodesFilter, EFiBVersion::FIB_VER_VARIABLE_REFERENCE, OnSearchComplete);
+	FStreamSearchOptions SearchOptions;
+	SearchOptions.ImaginaryDataFilter = ESearchQueryFilter::NodesFilter;
+	SearchOptions.MinimiumVersionRequirement = EFiBVersion::FIB_VER_VARIABLE_REFERENCE;
+	FindInBlueprints->MakeSearchQuery(SearchTerm, false, SearchOptions, OnSearchComplete);
 }
 
 void SReplaceNodeReferences::FindAllReplacementsComplete(TArray<FImaginaryFiBDataSharedPtr>& InRawDataList)
@@ -486,7 +489,11 @@ void SReplaceNodeReferences::FindAllReplacementsComplete(TArray<FImaginaryFiBDat
 				FBlueprintEditorUtils::MarkBlueprintAsStructurallyModified(Blueprint);
 				FFindInBlueprintSearchManager::Get().AddOrUpdateBlueprintSearchMetadata(Blueprint);
 			}
-			FindInBlueprints->MakeSearchQuery(VariableReference.GetReferenceSearchString(SourceProperty->GetOwnerClass()), false, ESearchQueryFilter::NodesFilter, EFiBVersion::FIB_VER_VARIABLE_REFERENCE, FOnSearchComplete());
+
+			FStreamSearchOptions SearchOptions;
+			SearchOptions.ImaginaryDataFilter = ESearchQueryFilter::NodesFilter;
+			SearchOptions.MinimiumVersionRequirement = EFiBVersion::FIB_VER_VARIABLE_REFERENCE;
+			FindInBlueprints->MakeSearchQuery(VariableReference.GetReferenceSearchString(SourceProperty->GetOwnerClass()), false, SearchOptions);
 		}
 	}
 }
diff --git a/Engine/Source/Editor/Kismet/Public/FindInBlueprintManager.h b/Engine/Source/Editor/Kismet/Public/FindInBlueprintManager.h
--- a/Engine/Source/Editor/Kismet/Public/FindInBlueprintManager.h
+++ b/Engine/Source/Editor/Kismet/Public/FindInBlueprintManager.h
@@ -12,6 +12,11 @@
 #include "Input/Reply.h"
 #include "TickableEditorObject.h"
 
+DECLARE_LOG_CATEGORY_EXTERN(LogFindInBlueprint, Warning, All);
+
+/** CSV stats profiling category */
+CSV_DECLARE_CATEGORY_EXTERN(FindInBlueprint);
+
 struct FAssetData;
 class FFindInBlueprintsResult;
 class FImaginaryBlueprint;
@@ -312,6 +317,22 @@ typedef TSharedPtr<FFindInBlueprintsResult> FSearchResult;
 ////////////////////////////////////
 // FStreamSearch
 
+struct FStreamSearchOptions
+{
+	/** Filter to limit the FilteredImaginaryResults to */
+	enum ESearchQueryFilter ImaginaryDataFilter;
+
+	/** When searching, any Blueprint below this version will be considered out-of-date */
+	EFiBVersion MinimiumVersionRequirement;
+
+	/** Default constructor. */
+	FStreamSearchOptions()
+		:ImaginaryDataFilter(ESearchQueryFilter::AllFilter)
+		,MinimiumVersionRequirement(EFiBVersion::FIB_VER_LATEST)
+	{
+	}
+};
+
 /**
  * Async task for searching Blueprints
  */
@@ -319,8 +340,7 @@ class FStreamSearch : public FRunnable
 {
 public:
 	/** Constructor */
-	FStreamSearch(const FString& InSearchValue);
-	FStreamSearch(const FString& InSearchValue, enum ESearchQueryFilter InImaginaryDataFilter, EFiBVersion InMinimiumVersionRequirement);
+	FStreamSearch(const FString& InSearchValue, const FStreamSearchOptions& InSearchOptions = FStreamSearchOptions());
 
 	/** Begin FRunnable Interface */
 	virtual bool Init() override;
@@ -367,26 +387,27 @@ public:
 	/** The search value to filter results by */
 	FString SearchValue;
 
+	/** Options for setting up the search */
+	FStreamSearchOptions SearchOptions;
+
 	/** Prevents searching while other threads are pulling search results */
 	FCriticalSection SearchCriticalSection;
 
-	// Whether the thread has finished running
-	bool bThreadCompleted;
-
 	/** > 0 if we've been asked to abort work in progress at the next opportunity */
 	FThreadSafeCounter StopTaskCounter;
 
-	/** When searching, any Blueprint below this version will be considered out-of-date */
-	EFiBVersion MinimiumVersionRequirement;
+	/** Filtered (ImaginaryDataFilter) list of imaginary data results that met the search requirements. Must be declared as thread-safe since imaginary data is a shared resource. */
+	TArray<FImaginaryFiBDataSharedPtr> FilteredImaginaryResults;
 
 	/** A going count of all Blueprints below the MinimiumVersionRequirement */
 	int32 BlueprintCountBelowVersion;
 
-	/** Filtered (ImaginaryDataFilter) list of imaginary data results that met the search requirements. Must be declared as thread-safe since imaginary data is a shared resource. */
-	TArray<FImaginaryFiBDataSharedPtr> FilteredImaginaryResults;
+	// Whether the thread has finished running
+	bool bThreadCompleted;
 
-	/** Filter to limit the FilteredImaginaryResults to */
-	enum ESearchQueryFilter ImaginaryDataFilter;
+private:
+	/** Unique identifier for this search (used with benchmarking) */
+	int32 SearchId;
 };
 
 ////////////////////////////////////
@@ -565,6 +586,9 @@ private:
 	/** Callback hook from the Asset Registry, marks the asset for deletion from the cache */
 	void OnAssetRenamed(const struct FAssetData& InAssetData, const FString& InOldName);
 
+	/** Callback hook from the Asset Registry, signals that the initial asset discovery stage has been completed */
+	void OnAssetRegistryFilesLoaded();
+
 	/** Callback hook from the Asset Registry when an asset is loaded */
 	void OnAssetLoaded(class UObject* InAsset);
 
@@ -607,12 +631,6 @@ private:
 	TSharedPtr<SFindInBlueprints> OpenGlobalFindResultsTab();
 
 protected:
-	/** Tells if gathering data is currently allowed */
-	bool bEnableGatheringData;
-
-	/** Tells if deferred indexing is configured as disabled */
-	bool bDisableDeferredIndexing;
-
 	/** Maps the Blueprint paths to their index in the SearchArray */
 	TMap<FName, int32> SearchMap;
 
@@ -634,9 +652,6 @@ protected:
 	/** Critical section to safely modify cached data */
 	FCriticalSection SafeModifyCacheCriticalSection;
 
-	/** TRUE when the the FiB manager wants to pause all searches, helps manage the pausing procedure */
-	volatile bool bIsPausing;
-
 	/** Because we are unable to query for the module on another thread, cache it for use later */
 	class FAssetRegistryModule* AssetRegistryModule;
 
@@ -669,6 +684,18 @@ protected:
 
 	/** Global Find Results workspace menu item */
 	TSharedPtr<class FWorkspaceItem> GlobalFindResultsMenuItem;
+
+	/** Tells if gathering data is currently allowed */
+	bool bEnableGatheringData;
+
+	/** Tells if deferred indexing is configured as disabled */
+	bool bDisableDeferredIndexing;
+
+	/** Whether CSV profiling has been enabled (default=false) */
+	bool bEnableCSVStatsProfiling;
+
+	/** TRUE when the the FiB manager wants to pause all searches, helps manage the pausing procedure */
+	volatile bool bIsPausing;
 };
 
 struct KISMET_API FDisableGatheringDataOnScope
diff --git a/Engine/Source/Editor/Kismet/Public/FindInBlueprints.h b/Engine/Source/Editor/Kismet/Public/FindInBlueprints.h
--- a/Engine/Source/Editor/Kismet/Public/FindInBlueprints.h
+++ b/Engine/Source/Editor/Kismet/Public/FindInBlueprints.h
@@ -227,11 +227,10 @@ public:
 	 *
 	 * @param InSearchString						String to search using
 	 * @param bInIsFindWithinBlueprint				TRUE if searching within the current Blueprint only
-	 * @param InSearchFilterForImaginaryDataReturn	If requesting a callback on search complete for the filtered imaginary data, this is the filter that the raw data will be passed through. By default nothing is collected
-	 * @param InMinimiumVersionRequirement			The minimum search requirement Blueprints must be to be searched or they will be reported as out-of-date.
+	 * @param InSearchOptions						Optional search parameters.
 	 * @param InOnSearchComplete					Callback when the search is complete, passing the filtered imaginary data (if any).
 	 */
-	void MakeSearchQuery(FString InSearchString, bool bInIsFindWithinBlueprint, enum ESearchQueryFilter InSearchFilterForImaginaryDataReturn = ESearchQueryFilter::AllFilter, EFiBVersion InMinimiumVersionRequirement = EFiBVersion::FIB_VER_LATEST, FOnSearchComplete InOnSearchComplete = FOnSearchComplete());
+	void MakeSearchQuery(FString InSearchString, bool bInIsFindWithinBlueprint, const FStreamSearchOptions& InSearchOptions = FStreamSearchOptions(), FOnSearchComplete InOnSearchComplete = FOnSearchComplete());
 
 	/** Called when caching Blueprints is started */
 	void OnCacheStarted(EFiBCacheOpType InOpType, EFiBCacheOpFlags InOpFlags);
@@ -290,8 +289,7 @@ private:
 	TSharedRef<ITableRow> OnGenerateRow(FSearchResult InItem, const TSharedRef<STableViewBase>& OwnerTable);
 	
 	/** Launches a thread for streaming more content into the results widget */
-	void LaunchStreamThread(const FString& InSearchValue);
-	void LaunchStreamThread(const FString& InSearchValue, enum ESearchQueryFilter InSearchFilterForRawDataReturn, EFiBVersion InMinimiumVersionRequirement, FOnSearchComplete InOnSearchComplete);
+	void LaunchStreamThread(const FString& InSearchValue, const FStreamSearchOptions& InSearchOptions = FStreamSearchOptions(), FOnSearchComplete InOnSearchComplete = FOnSearchComplete());
 
 	/** Returns the percent complete on the search for the progress bar */
 	TOptional<float> GetPercentCompleteSearch() const;
