CIS fix (non-unity).

#rnx

diff --git a/Engine/Source/Editor/Kismet/Public/FindInBlueprintManager.h b/Engine/Source/Editor/Kismet/Public/FindInBlueprintManager.h
--- a/Engine/Source/Editor/Kismet/Public/FindInBlueprintManager.h
+++ b/Engine/Source/Editor/Kismet/Public/FindInBlueprintManager.h
@@ -11,6 +11,7 @@
 #include "SlateFwd.h"
 #include "Input/Reply.h"
 #include "TickableEditorObject.h"
+#include "ProfilingDebugging/CsvProfiler.h"
 
 DECLARE_LOG_CATEGORY_EXTERN(LogFindInBlueprint, Warning, All);
 
