Regenerate replication data when relinking derived types in the blueprint reinstancer
#rnx

diff --git a/Engine/Source/Editor/Kismet/Private/BlueprintCompilationManager.cpp b/Engine/Source/Editor/Kismet/Private/BlueprintCompilationManager.cpp
--- a/Engine/Source/Editor/Kismet/Private/BlueprintCompilationManager.cpp
+++ b/Engine/Source/Editor/Kismet/Private/BlueprintCompilationManager.cpp
@@ -1840,6 +1840,9 @@ void FBlueprintCompilationManagerImpl::ReinstanceBatch(TArray<FReinstancingJob>&
 		Class->SetSuperStruct(*NewParent);
 		Class->Bind();
 		Class->StaticLink(true);
+
+		Class->ClassFlags &= ~CLASS_ReplicationDataIsSetUp;
+		Class->SetUpRuntimeReplicationData();
 	}
 
 	// make new hierarchy
@@ -1874,6 +1877,9 @@ void FBlueprintCompilationManagerImpl::ReinstanceBatch(TArray<FReinstancingJob>&
 		// relink the new class:
 		NewestReinstancer->ClassToReinstance->Bind();
 		NewestReinstancer->ClassToReinstance->StaticLink(true);
+
+		NewestReinstancer->ClassToReinstance->ClassFlags &= ~CLASS_ReplicationDataIsSetUp;
+		NewestReinstancer->ClassToReinstance->SetUpRuntimeReplicationData();
 	}
 
 	// run UpdateBytecodeReferences:
