Add commands to dump the Blueprint search index for analysis/debugging, along with some minor revisions to JSON serialization.

diff --git a/Engine/Config/BaseEditor.ini b/Engine/Config/BaseEditor.ini
--- a/Engine/Config/BaseEditor.ini
+++ b/Engine/Config/BaseEditor.ini
@@ -432,6 +432,7 @@ ForceUnits=EUnit::Newtons
 [BlueprintSearchSettings]
 bDisableDeferredIndexing=false
 bEnableCsvStatsProfiling=false
+bEnableDeveloperMenuTools=false
 
 [BlueprintNativizationSettings]
 bUseStaticVariablesInWrappers=true
diff --git a/Engine/Source/Editor/Kismet/Private/BlueprintEditor.cpp b/Engine/Source/Editor/Kismet/Private/BlueprintEditor.cpp
--- a/Engine/Source/Editor/Kismet/Private/BlueprintEditor.cpp
+++ b/Engine/Source/Editor/Kismet/Private/BlueprintEditor.cpp
@@ -89,6 +89,7 @@
 #include "BlueprintFunctionNodeSpawner.h"
 #include "SBlueprintEditorToolbar.h"
 #include "FindInBlueprints.h"
+#include "ImaginaryBlueprintData.h"
 #include "SGraphTitleBar.h"
 #include "Kismet2/Kismet2NameValidators.h"
 #include "Kismet2/DebuggerCommands.h"
@@ -2711,6 +2712,14 @@ void FBlueprintEditor::CreateDefaultCommands()
 		FExecuteAction::CreateSP(this, &FBlueprintEditor::OpenNativeCodeGenerationTool),
 		FCanExecuteAction::CreateSP(this, &FBlueprintEditor::CanGenerateNativeCode));
 
+	ToolkitCommands->MapAction(FBlueprintEditorCommands::Get().GenerateSearchIndex,
+		FExecuteAction::CreateSP(this, &FBlueprintEditor::OnGenerateSearchIndexForDebugging),
+		FCanExecuteAction());
+
+	ToolkitCommands->MapAction(FBlueprintEditorCommands::Get().DumpCachedIndexData,
+		FExecuteAction::CreateSP(this, &FBlueprintEditor::OnDumpCachedIndexDataForBlueprint),
+		FCanExecuteAction());
+
 	ToolkitCommands->MapAction(FBlueprintEditorCommands::Get().ShowActionMenuItemSignatures,
 		FExecuteAction::CreateLambda([]()
 			{ 
@@ -2767,6 +2776,54 @@ bool FBlueprintEditor::CanGenerateNativeCode() const
 	return Blueprint && FNativeCodeGenerationTool::CanGenerate(*Blueprint);
 }
 
+void FBlueprintEditor::OnGenerateSearchIndexForDebugging()
+{
+	UBlueprint* Blueprint = GetBlueprintObj();
+	if (Blueprint)
+	{
+		FString FileLocation = FPaths::ConvertRelativePathToFull(FPaths::ProjectIntermediateDir() + TEXT("BlueprintSearchTools"));
+		FString FullPath = FString::Printf(TEXT("%s/%s.index.xml"), *FileLocation, *Blueprint->GetName());
+
+		FArchive* DumpFile = IFileManager::Get().CreateFileWriter(*FullPath);
+		if (DumpFile)
+		{
+			FString JsonOutput = FFindInBlueprintSearchManager::Get().GenerateSearchIndexForDebugging(Blueprint);
+			DumpFile->Serialize(TCHAR_TO_ANSI(*JsonOutput), JsonOutput.Len());
+
+			DumpFile->Close();
+			delete DumpFile;
+
+			UE_LOG(LogFindInBlueprint, Log, TEXT("Wrote search index to %s"), *FullPath);
+		}
+	}
+}
+
+void FBlueprintEditor::OnDumpCachedIndexDataForBlueprint()
+{
+	UBlueprint* Blueprint = GetBlueprintObj();
+	if (Blueprint)
+	{
+		FString FileLocation = FPaths::ConvertRelativePathToFull(FPaths::ProjectIntermediateDir() + TEXT("BlueprintSearchTools"));
+		FString FullPath = FString::Printf(TEXT("%s/%s.cache.csv"), *FileLocation, *Blueprint->GetName());
+
+		FArchive* DumpFile = IFileManager::Get().CreateFileWriter(*FullPath);
+		if (DumpFile)
+		{
+			const FName AssetPath = *Blueprint->GetPathName();
+			FSearchData SearchData = FFindInBlueprintSearchManager::Get().GetSearchDataForAssetPath(AssetPath);
+			if (SearchData.IsValid() && SearchData.ImaginaryBlueprint.IsValid())
+			{
+				SearchData.ImaginaryBlueprint->DumpParsedObject(*DumpFile);
+			}
+
+			DumpFile->Close();
+			delete DumpFile;
+
+			UE_LOG(LogFindInBlueprint, Log, TEXT("Wrote cached index data to %s"), *FullPath);
+		}
+	}
+}
+
 void FBlueprintEditor::FindInBlueprint_Clicked()
 {
 	SummonSearchUI(true);
diff --git a/Engine/Source/Editor/Kismet/Private/BlueprintEditorCommands.cpp b/Engine/Source/Editor/Kismet/Private/BlueprintEditorCommands.cpp
--- a/Engine/Source/Editor/Kismet/Private/BlueprintEditorCommands.cpp
+++ b/Engine/Source/Editor/Kismet/Private/BlueprintEditorCommands.cpp
@@ -70,6 +70,8 @@ void FBlueprintEditorCommands::RegisterCommands()
 	// Development commands
 	UI_COMMAND( SaveIntermediateBuildProducts, "Save Intermediate Build Products", "Should the compiler save intermediate build products for debugging.", EUserInterfaceActionType::ToggleButton, FInputChord() );
 	UI_COMMAND(GenerateNativeCode, "Generate Native Code", "Generate C++ code from the blueprint", EUserInterfaceActionType::Button, FInputChord());
+	UI_COMMAND(GenerateSearchIndex, "Generate Search Index", "Generate the search index for this blueprint.", EUserInterfaceActionType::Button, FInputChord());
+	UI_COMMAND(DumpCachedIndexData, "Dump Cached Index Data", "Dump the cached index data for this blueprint.", EUserInterfaceActionType::Button, FInputChord());
 	UI_COMMAND(ShowActionMenuItemSignatures, "Show Action Menu Item Signatures", "If enabled, tooltips on action menu items will show the associated action's signature id (can be used to setup custom favorites menus).", EUserInterfaceActionType::ToggleButton, FInputChord());
 	
 	// SCC commands
diff --git a/Engine/Source/Editor/Kismet/Private/BlueprintEditorCommands.h b/Engine/Source/Editor/Kismet/Private/BlueprintEditorCommands.h
--- a/Engine/Source/Editor/Kismet/Private/BlueprintEditorCommands.h
+++ b/Engine/Source/Editor/Kismet/Private/BlueprintEditorCommands.h
@@ -71,6 +71,8 @@ public:
 	// Development commands
 	TSharedPtr< FUICommandInfo > SaveIntermediateBuildProducts;
 	TSharedPtr< FUICommandInfo > GenerateNativeCode;
+	TSharedPtr< FUICommandInfo > GenerateSearchIndex;
+	TSharedPtr< FUICommandInfo > DumpCachedIndexData;
 	TSharedPtr< FUICommandInfo > ShowActionMenuItemSignatures;
 
 	// SSC commands
diff --git a/Engine/Source/Editor/Kismet/Private/FindInBlueprintManager.cpp b/Engine/Source/Editor/Kismet/Private/FindInBlueprintManager.cpp
--- a/Engine/Source/Editor/Kismet/Private/FindInBlueprintManager.cpp
+++ b/Engine/Source/Editor/Kismet/Private/FindInBlueprintManager.cpp
@@ -1,6 +1,7 @@
 // Copyright 1998-2019 Epic Games, Inc. All Rights Reserved.
 
 #include "FindInBlueprintManager.h"
+#include "Misc/CoreMisc.h"
 #include "Misc/MessageDialog.h"
 #include "HAL/FileManager.h"
 #include "Misc/Paths.h"
@@ -436,13 +437,15 @@ namespace BlueprintSearchMetaDataHelpers
 	};
 
 	/** Json Writer used for serializing FText's in the correct format for Find-in-Blueprints */
-	template < class PrintPolicy = TPrettyJsonPrintPolicy<TCHAR> >
-	class TJsonFindInBlueprintStringWriter : public TJsonStringWriter<PrintPolicy>
+	template<class PrintPolicy>
+	class TFindInBlueprintJsonStringWriter : public TJsonStringWriter<PrintPolicy>
 	{
 	public:
-		static TSharedRef< TJsonFindInBlueprintStringWriter > Create( FString* const InStream )
+		typedef PrintPolicy InnerPrintPolicy;
+
+		TFindInBlueprintJsonStringWriter(FString* const InOutString)
+			: TJsonStringWriter<PrintPolicy>(InOutString, 0)
 		{
-			return MakeShareable( new TJsonFindInBlueprintStringWriter( InStream ) );
 		}
 
 		using TJsonStringWriter<PrintPolicy>::WriteObjectStart;
@@ -490,12 +493,79 @@ namespace BlueprintSearchMetaDataHelpers
 			this->PreviousTokenWritten = this->WriteValueOnly( Value );
 		}
 
-		/** Converts the lookup table of ints (which are stored as identifiers and string values in the Json) and the FText's they represent to an FString. */
-		FString GetSerializedLookupTable()
+	protected:
+		virtual void WriteTextValue( const FText& Text )
+		{
+			TJsonStringWriter<PrintPolicy>::WriteStringValue(Text.ToString());
+		}
+
+		FORCEINLINE void WriteIdentifier( const FText& Identifier )
+		{
+			this->WriteCommaIfNeeded();
+			PrintPolicy::WriteLineTerminator(this->Stream);
+
+			PrintPolicy::WriteTabs(this->Stream, this->IndentLevel);
+
+			WriteTextValue( Identifier );
+			PrintPolicy::WriteChar(this->Stream, TCHAR(':'));
+		}
+
+	public:
+		/** Cached mapping of all searchable properties that have been discovered while gathering searchable data for the current Blueprint */
+		TMap<UStruct*, TArray<FSearchableProperty>> CachedPropertyMapping;
+	};
+
+	/** Json Writer used for serializing FText's in the correct format for Find-in-Blueprints */
+	class FFindInBlueprintJsonWriter : public TFindInBlueprintJsonStringWriter<TCondensedJsonPrintPolicy<TCHAR>>
+	{
+	public:
+		FFindInBlueprintJsonWriter(FString* const InOutString)
+			:TFindInBlueprintJsonStringWriter<TCondensedJsonPrintPolicy<TCHAR>>(InOutString)
+			, JsonOutput(InOutString)
+		{
+		}
+
+		virtual bool Close() override
+		{
+			// This will copy the JSON output to the string given as input (must do this first)
+			bool bResult = TFindInBlueprintJsonStringWriter<TCondensedJsonPrintPolicy<TCHAR>>::Close();
+
+			// Build the search metadata string for the asset tag (version + LUT + JSON)
+			*JsonOutput = FiBSerializationHelpers::Serialize(FSearchDataVersionInfo::Current.FiBDataVersion, false)
+				+ FiBSerializationHelpers::Serialize< TMap< int32, FText > >(LookupTable, true)
+				+ MoveTemp(*JsonOutput);
+
+			return bResult;
+		}
+
+	protected:
+		virtual void WriteStringValue(const FString& String) override
+		{
+			// We just want to make sure all strings are converted into FText hex strings, used by the FiB system
+			WriteTextValue(FText::FromString(String));
+		}
+
+		virtual void WriteTextValue(const FText& Text) override
 		{
-			return FiBSerializationHelpers::Serialize< TMap< int32, FText > >(LookupTable, true);
+			// Check to see if the value has already been added.
+			int32* TableLookupValuePtr = ReverseLookupTable.Find(FLookupTableItem(Text));
+			if (TableLookupValuePtr)
+			{
+				TFindInBlueprintJsonStringWriter<TCondensedJsonPrintPolicy<TCHAR>>::WriteStringValue(FString::FromInt(*TableLookupValuePtr));
+			}
+			else
+			{
+				// Add the FText to the table and write to the Json the ID to look the item up using
+				int32 TableLookupValue = LookupTable.Num();
+				{
+					LookupTable.Add(TableLookupValue, Text);
+					ReverseLookupTable.Add(FLookupTableItem(Text), TableLookupValue);
+				}
+				TFindInBlueprintJsonStringWriter<TCondensedJsonPrintPolicy<TCHAR>>::WriteStringValue(FString::FromInt(TableLookupValue));
+			}
 		}
 
+	private:
 		struct FLookupTableItem
 		{
 			FText Text;
@@ -531,97 +601,16 @@ namespace BlueprintSearchMetaDataHelpers
 			}
 		};
 
-	protected:
-		TJsonFindInBlueprintStringWriter( FString* const InOutString )
-			: TJsonStringWriter<PrintPolicy>( InOutString, 0 )
-		{
-		}
-
-		virtual void WriteStringValue( const FString& String ) override
-		{
-			// We just want to make sure all strings are converted into FText hex strings, used by the FiB system
-			WriteTextValue(FText::FromString(String));
-		}
-
-		void WriteTextValue( const FText& Text )
-		{
-			// Check to see if the value has already been added.
-			int32* TableLookupValuePtr = ReverseLookupTable.Find(FLookupTableItem(Text));
-			if(TableLookupValuePtr)
-			{
-				TJsonStringWriter<PrintPolicy>::WriteStringValue(FString::FromInt(*TableLookupValuePtr));
-			}
-			else
-			{
-				// Add the FText to the table and write to the Json the ID to look the item up using
-				int32 TableLookupValue = LookupTable.Num();
-				{
-					LookupTable.Add(TableLookupValue, Text);
-					ReverseLookupTable.Add(FLookupTableItem(Text), TableLookupValue);
-				}
-				TJsonStringWriter<PrintPolicy>::WriteStringValue( FString::FromInt(TableLookupValue) );
-			}
-		}
-
-		FORCEINLINE void WriteIdentifier( const FText& Identifier )
-		{
-			this->WriteCommaIfNeeded();
-			PrintPolicy::WriteLineTerminator(this->Stream);
-
-			PrintPolicy::WriteTabs(this->Stream, this->IndentLevel);
+		// Output stream
+		FString* JsonOutput;
 
-			WriteTextValue( Identifier );
-			PrintPolicy::WriteChar(this->Stream, TCHAR(':'));
-		}
-		
 		// This gets serialized
 		TMap< int32, FText > LookupTable;
 
 		// This is just locally needed for the write, to lookup the integer value by using the string of the FText
 		TMap< FLookupTableItem, int32 > ReverseLookupTable;
-
-	public:
-		/** Cached mapping of all searchable properties that have been discovered while gathering searchable data for the current Blueprint */
-		TMap<UStruct*, TArray<FSearchableProperty>> CachedPropertyMapping;
-	};
-
-	typedef TJsonWriter<TCHAR, TCondensedJsonPrintPolicy<TCHAR>> SearchMetaDataWriterParentClass;
-	typedef TJsonFindInBlueprintStringWriter<TCondensedJsonPrintPolicy<TCHAR>> SearchMetaDataWriter;
-
-	/** Json Writer used for serializing FText's in the correct format for Find-in-Blueprints */
-	template <class CharType = TCHAR>
-	class TJsonFindInBlueprintStringReader : public TJsonReader<CharType>
-	{
-	public:
-		static TSharedRef< TJsonFindInBlueprintStringReader< TCHAR > > Create( FArchive* const Stream, TMap< int32, FText >& InLookupTable  )
-		{
-			return MakeShareable( new TJsonFindInBlueprintStringReader( Stream, InLookupTable ) );
-		}
-
-		TJsonFindInBlueprintStringReader( FArchive* InStream,  TMap< int32, FText >& InLookupTable )
-			: TJsonReader<CharType>(InStream)
-			, LookupTable(MoveTemp(InLookupTable))
-		{
-
-		}
-
-		FORCEINLINE virtual const FString& GetIdentifier() const override
-		{
-			return this->Identifier;
-		}
-
-		FORCEINLINE virtual  const FString& GetValueAsString() const override
-		{ 
-			check( this->CurrentToken == EJsonToken::String ); 
-			// The string value from Json is a Hex value that must be looked up in the LookupTable to find the FText it represents
-			return this->StringValue;
-		}
-
-		TMap< int32, FText > LookupTable;
 	};
 
-	typedef TJsonFindInBlueprintStringReader<TCHAR> SearchMetaDataReader;
-
 	/**
 	 * Checks if Json value is searchable, eliminating data that not considered useful to search for
 	 *
@@ -705,7 +694,8 @@ namespace BlueprintSearchMetaDataHelpers
 	 * @param InWriter				Writer used for saving the Json
 	 * @param InPinType				The pin type to save
 	 */
-	void SavePinTypeToJson(TSharedRef< SearchMetaDataWriter>& InWriter, const FEdGraphPinType& InPinType)
+	template<class PrintPolicy>
+	void SavePinTypeToJson(const TSharedRef<TFindInBlueprintJsonStringWriter<PrintPolicy>>& InWriter, const FEdGraphPinType& InPinType)
 	{
 		// Only save strings that are not empty
 
@@ -734,7 +724,8 @@ namespace BlueprintSearchMetaDataHelpers
 	 * @param InBlueprint				Blueprint the property for the variable can be found in, if any
 	 * @param InVariableDescription		The variable description being serialized to Json
 	 */
-	void SaveVariableDescriptionToJson(TSharedRef< SearchMetaDataWriter>& InWriter, const UBlueprint* InBlueprint, const FBPVariableDescription& InVariableDescription)
+	template<class PrintPolicy>
+	void SaveVariableDescriptionToJson(const TSharedRef<TFindInBlueprintJsonStringWriter<PrintPolicy>>& InWriter, const UBlueprint* InBlueprint, const FBPVariableDescription& InVariableDescription)
 	{
 		FEdGraphPinType VariableType = InVariableDescription.VarType;
 
@@ -765,7 +756,7 @@ namespace BlueprintSearchMetaDataHelpers
 			if(BlueprintSearchMetaDataHelpers::CheckIfJsonValueIsSearchable(JsonValue))
 			{
 				TSharedRef< FJsonValue > JsonValueAsSharedRef = JsonValue.ToSharedRef();
-				FJsonSerializer::Serialize(JsonValue, FFindInBlueprintSearchTags::FiB_DefaultValue.ToString(), StaticCastSharedRef<SearchMetaDataWriterParentClass>(InWriter), false );
+				FJsonSerializer::Serialize<TCHAR, PrintPolicy>(JsonValue, FFindInBlueprintSearchTags::FiB_DefaultValue.ToString(), InWriter, false );
 			}
 		}
 
@@ -788,7 +779,8 @@ namespace BlueprintSearchMetaDataHelpers
 	 * @param InStruct				Struct or class that represent the UObject's layout
 	 * @param InSearchableType		Informs the system how it should examine the properties to determine if they are searchable. All sub-properties of searchable properties are automatically gathered unless marked as not being searchable
 	 */
-	void GatherSearchableProperties(TSharedRef< SearchMetaDataWriter>& InWriter, const void* InValue, UStruct* InStruct, EGatherSearchableType InSearchableType = SEARCHABLE_AS_DESIRED);
+	template<class PrintPolicy>
+	void GatherSearchableProperties(const TSharedRef<TFindInBlueprintJsonStringWriter<PrintPolicy>>& InWriter, const void* InValue, UStruct* InStruct, EGatherSearchableType InSearchableType = SEARCHABLE_AS_DESIRED);
 
 	/**
 	 * Examines a searchable property and digs in deeper if it is a UObject, UStruct, or an array, or serializes it straight out to Json
@@ -798,7 +790,8 @@ namespace BlueprintSearchMetaDataHelpers
 	 * @param InValue				Value to find the property in the UStruct
 	 * @param InStruct				Struct or class that represent the UObject's layout
 	 */
-	void GatherSearchablesFromProperty(TSharedRef< SearchMetaDataWriter>& InWriter, UProperty* InProperty, const void* InValue, UStruct* InStruct)
+	template<class PrintPolicy>
+	void GatherSearchablesFromProperty(const TSharedRef<TFindInBlueprintJsonStringWriter<PrintPolicy>>& InWriter, UProperty* InProperty, const void* InValue, UStruct* InStruct)
 	{
 		if (UArrayProperty* ArrayProperty = Cast<UArrayProperty>(InProperty))
 		{
@@ -840,7 +833,7 @@ namespace BlueprintSearchMetaDataHelpers
 					// Shallow conversion of property to string
 					TSharedPtr<FJsonValue> JsonValue;
 					JsonValue = FJsonObjectConverter::UPropertyToJsonValue(InProperty, InValue, 0, 0);
-					FJsonSerializer::Serialize(JsonValue, InProperty->GetName(), StaticCastSharedRef<SearchMetaDataWriterParentClass>(InWriter), false);
+					FJsonSerializer::Serialize<TCHAR, PrintPolicy>(JsonValue, InProperty->GetName(), InWriter, false);
 				}
 			}
 		}
@@ -848,11 +841,12 @@ namespace BlueprintSearchMetaDataHelpers
 		{
 			TSharedPtr<FJsonValue> JsonValue;
 			JsonValue = FJsonObjectConverter::UPropertyToJsonValue(InProperty, InValue, 0, 0);
-			FJsonSerializer::Serialize(JsonValue, InProperty->GetName(), StaticCastSharedRef<SearchMetaDataWriterParentClass>(InWriter), false);
+			FJsonSerializer::Serialize<TCHAR, PrintPolicy>(JsonValue, InProperty->GetName(), InWriter, false);
 		}
 	}
 
-	void GatherSearchableProperties(TSharedRef<SearchMetaDataWriter>& InWriter, const void* InValue, UStruct* InStruct, EGatherSearchableType InSearchableType)
+	template<class PrintPolicy>
+	void GatherSearchableProperties(const TSharedRef<TFindInBlueprintJsonStringWriter<PrintPolicy>>& InWriter, const void* InValue, UStruct* InStruct, EGatherSearchableType InSearchableType)
 	{
 		if (InValue)
 		{
@@ -1020,7 +1014,8 @@ namespace BlueprintSearchMetaDataHelpers
 	 * @param InWriter		The Json writer to use for serialization
 	 * @param InGraph		The graph to search through
 	 */
-	void GatherNodesFromGraph(TSharedRef< SearchMetaDataWriter>& InWriter, const UEdGraph* InGraph)
+	template<class PrintPolicy>
+	void GatherNodesFromGraph(const TSharedRef<TFindInBlueprintJsonStringWriter<PrintPolicy>>& InWriter, const UEdGraph* InGraph)
 	{
 		// Collect all macro graphs
 		InWriter->WriteArrayStart(FFindInBlueprintSearchTags::FiB_Nodes);
@@ -1097,7 +1092,8 @@ namespace BlueprintSearchMetaDataHelpers
 	 * @param InTitle			The array title to place these graphs into
 	 * @param InOutSubGraphs	All the subgraphs that need to be processed later
 	 */
-	void GatherGraphSearchData(TSharedRef< SearchMetaDataWriter>& InWriter, const UBlueprint* InBlueprint, const TArray< UEdGraph* >& InGraphArray, FText InTitle, TArray< UEdGraph* >* InOutSubGraphs)
+	template<class PrintPolicy>
+	void GatherGraphSearchData(const TSharedRef<TFindInBlueprintJsonStringWriter<PrintPolicy>>& InWriter, const UBlueprint* InBlueprint, const TArray< UEdGraph* >& InGraphArray, FText InTitle, TArray< UEdGraph* >* InOutSubGraphs)
 	{
 		if(InGraphArray.Num() > 0)
 		{
@@ -1158,6 +1154,102 @@ namespace BlueprintSearchMetaDataHelpers
 			InWriter->WriteArrayEnd();
 		}
 	}
+
+	template<class PrintPolicy>
+	void GatherBlueprintSearchMetadata(const TSharedRef<TFindInBlueprintJsonStringWriter<PrintPolicy>>& InWriter, const UBlueprint* Blueprint)
+	{
+		FTemporarilyUseFriendlyNodeTitles TemporarilyUseFriendlyNodeTitles;
+
+		TMap<FString, TMap<FString, int>> AllPaths;
+		InWriter->WriteObjectStart();
+
+		// Only pull properties if the Blueprint has been compiled
+		if (Blueprint->SkeletonGeneratedClass)
+		{
+			InWriter->WriteArrayStart(FFindInBlueprintSearchTags::FiB_Properties);
+			{
+				for (const FBPVariableDescription& Variable : Blueprint->NewVariables)
+				{
+					SaveVariableDescriptionToJson(InWriter, Blueprint, Variable);
+				}
+			}
+			InWriter->WriteArrayEnd(); // Properties
+		}
+
+		// Gather all graph searchable data
+		TArray< UEdGraph* > SubGraphs;
+
+		// Gather normal event graphs
+		GatherGraphSearchData(InWriter, Blueprint, Blueprint->UbergraphPages, FFindInBlueprintSearchTags::FiB_UberGraphs, &SubGraphs);
+
+		// We have interface graphs and function graphs to put into the Functions category. We cannot do them separately, so we must compile the full list
+		{
+			TArray<UEdGraph*> CompleteGraphList;
+			CompleteGraphList.Append(Blueprint->FunctionGraphs);
+			// Gather all interface graphs as functions
+			for (const FBPInterfaceDescription& InterfaceDesc : Blueprint->ImplementedInterfaces)
+			{
+				CompleteGraphList.Append(InterfaceDesc.Graphs);
+			}
+			GatherGraphSearchData(InWriter, Blueprint, CompleteGraphList, FFindInBlueprintSearchTags::FiB_Functions, &SubGraphs);
+		}
+
+		// Gather Macros
+		GatherGraphSearchData(InWriter, Blueprint, Blueprint->MacroGraphs, FFindInBlueprintSearchTags::FiB_Macros, &SubGraphs);
+
+		// Sub graphs are processed separately so that they do not become children in the TreeView, cluttering things up if the tree is deep
+		GatherGraphSearchData(InWriter, Blueprint, SubGraphs, FFindInBlueprintSearchTags::FiB_SubGraphs, nullptr);
+
+		// Gather all SCS components
+		// If we have an SCS but don't support it, then we remove it
+		if (Blueprint->SimpleConstructionScript)
+		{
+			// Remove any SCS variable nodes
+			const TArray<USCS_Node*>& AllSCSNodes = Blueprint->SimpleConstructionScript->GetAllNodes();
+			InWriter->WriteArrayStart(FFindInBlueprintSearchTags::FiB_Components);
+			for (TFieldIterator<UProperty> PropertyIt(Blueprint->SkeletonGeneratedClass, EFieldIteratorFlags::ExcludeSuper); PropertyIt; ++PropertyIt)
+			{
+				UProperty* Property = *PropertyIt;
+				UObjectPropertyBase* Obj = Cast<UObjectPropertyBase>(Property);
+				const bool bComponentProperty = Obj && Obj->PropertyClass ? Obj->PropertyClass->IsChildOf<UActorComponent>() : false;
+				FName PropName = Property->GetFName();
+				if (bComponentProperty && FBlueprintEditorUtils::FindSCS_Node(Blueprint, PropName) != INDEX_NONE)
+				{
+					FEdGraphPinType PropertyPinType;
+					if (UEdGraphSchema_K2::StaticClass()->GetDefaultObject<UEdGraphSchema_K2>()->ConvertPropertyToPinType(Property, PropertyPinType))
+					{
+						InWriter->WriteObjectStart();
+						{
+							InWriter->WriteValue(FFindInBlueprintSearchTags::FiB_Name, FText::FromName(PropName));
+							InWriter->WriteValue(FFindInBlueprintSearchTags::FiB_IsSCSComponent, true);
+							SavePinTypeToJson(InWriter, PropertyPinType);
+						}
+						InWriter->WriteObjectEnd();
+					}
+				}
+			}
+			InWriter->WriteArrayEnd(); // Components
+		}
+
+		InWriter->WriteObjectEnd();
+		InWriter->Close();
+	}
+
+	template<class JsonWriterType>
+	FString GatherBlueprintSearchMetadata(const UBlueprint* Blueprint)
+	{
+		CSV_SCOPED_TIMING_STAT(FindInBlueprint, GatherBlueprintSearchMetadata);
+
+		FString SearchMetaData;
+
+		// The search registry tags for a Blueprint are all in Json
+		TSharedRef<JsonWriterType> Writer = MakeShared<JsonWriterType>(&SearchMetaData);
+
+		typedef typename JsonWriterType::InnerPrintPolicy PrintPolicyType;
+		GatherBlueprintSearchMetadata<PrintPolicyType>(Writer, Blueprint);
+
+		return SearchMetaData;
+	}
 }
 
 class FCacheAllBlueprintsTickableObject
@@ -1493,6 +1585,7 @@ FFindInBlueprintSearchManager::FFindInBlueprintSearchManager()
 	, bEnableGatheringData(true)
 	, bDisableDeferredIndexing(false)
 	, bEnableCSVStatsProfiling(false)
+	, bEnableDeveloperMenuTools(false)
 	, bIsPausing(false)
 {
 	for (int32 TabIdx = 0; TabIdx < UE_ARRAY_COUNT(GlobalFindResultsTabIDs); TabIdx++)
@@ -1531,6 +1624,7 @@ void FFindInBlueprintSearchManager::Initialize()
 	// Init configuration
 	GConfig->GetBool(TEXT("BlueprintSearchSettings"), TEXT("bDisableDeferredIndexing"), bDisableDeferredIndexing, GEditorIni);
 	GConfig->GetBool(TEXT("BlueprintSearchSettings"), TEXT("bEnableCsvStatsProfiling"), bEnableCSVStatsProfiling, GEditorIni);
+	GConfig->GetBool(TEXT("BlueprintSearchSettings"), TEXT("bEnableDeveloperMenuTools"), bEnableDeveloperMenuTools, GEditorIni);
 
 	// If profiling has been enabled, turn on the stat category and begin a capture.
 	if (bEnableCSVStatsProfiling)
@@ -1728,6 +1822,24 @@ void FFindInBlueprintSearchManager::ExtractUnloadedFiBData(const FAssetData& InA
 	AddSearchDataToDatabase(MoveTemp(NewSearchData));
 }
 
+FSearchData FFindInBlueprintSearchManager::GetSearchDataForAssetPath(FName InAssetPath)
+{
+	FScopeLock ScopeLock(&SafeModifyCacheCriticalSection);
+
+	const int32* ArrayIdx = SearchMap.Find(InAssetPath);
+	if (ArrayIdx)
+	{
+		checkf(*ArrayIdx < SearchArray.Num(),
+			TEXT("ArrayIdx:%d, SearchArray.Num():%d"),
+			*ArrayIdx,
+			SearchArray.Num());
+
+		return SearchArray[*ArrayIdx];
+	}
+
+	return FSearchData();
+}
+
 int32 FFindInBlueprintSearchManager::AddSearchDataToDatabase(FSearchData InSearchData)
 {
 	FName AssetPath = InSearchData.AssetPath; // Copy before we move the data into the array
@@ -1808,99 +1920,6 @@ void FFindInBlueprintSearchManager::OnHotReload(bool bWasTriggeredAutomatically)
 	CachedAssetClasses.Reset();
 }
 
-FString FFindInBlueprintSearchManager::GatherBlueprintSearchMetadata(const UBlueprint* Blueprint)
-{
-	CSV_SCOPED_TIMING_STAT(FindInBlueprint, GatherBlueprintSearchMetadata);
-
-	FTemporarilyUseFriendlyNodeTitles TemporarilyUseFriendlyNodeTitles;
-
-	FString SearchMetaData;
-
-	// The search registry tags for a Blueprint are all in Json
-	TSharedRef< BlueprintSearchMetaDataHelpers::TJsonFindInBlueprintStringWriter<TCondensedJsonPrintPolicy<TCHAR>> > Writer = BlueprintSearchMetaDataHelpers::TJsonFindInBlueprintStringWriter<TCondensedJsonPrintPolicy<TCHAR>>::Create( &SearchMetaData );
-
-	TMap<FString, TMap<FString,int>> AllPaths;
-	Writer->WriteObjectStart();
-
-	// Only pull properties if the Blueprint has been compiled
-	if(Blueprint->SkeletonGeneratedClass)
-	{
-		Writer->WriteArrayStart(FFindInBlueprintSearchTags::FiB_Properties);
-		{
-			for (const FBPVariableDescription& Variable : Blueprint->NewVariables)
-			{
-				BlueprintSearchMetaDataHelpers::SaveVariableDescriptionToJson(Writer, Blueprint, Variable);
-			}
-		}
-		Writer->WriteArrayEnd(); // Properties
-	}
-
-	// Gather all graph searchable data
-	TArray< UEdGraph* > SubGraphs;
-
-	// Gather normal event graphs
-	BlueprintSearchMetaDataHelpers::GatherGraphSearchData(Writer, Blueprint, Blueprint->UbergraphPages, FFindInBlueprintSearchTags::FiB_UberGraphs, &SubGraphs);
-	
-	// We have interface graphs and function graphs to put into the Functions category. We cannot do them separately, so we must compile the full list
-	{
-		TArray<UEdGraph*> CompleteGraphList;
-		CompleteGraphList.Append(Blueprint->FunctionGraphs);
-		// Gather all interface graphs as functions
-		for (const FBPInterfaceDescription& InterfaceDesc : Blueprint->ImplementedInterfaces)
-		{
-			CompleteGraphList.Append(InterfaceDesc.Graphs);
-		}
-		BlueprintSearchMetaDataHelpers::GatherGraphSearchData(Writer, Blueprint, CompleteGraphList, FFindInBlueprintSearchTags::FiB_Functions, &SubGraphs);
-	}
-
-	// Gather Macros
-	BlueprintSearchMetaDataHelpers::GatherGraphSearchData(Writer, Blueprint, Blueprint->MacroGraphs, FFindInBlueprintSearchTags::FiB_Macros, &SubGraphs);
-
-	// Sub graphs are processed separately so that they do not become children in the TreeView, cluttering things up if the tree is deep
-	BlueprintSearchMetaDataHelpers::GatherGraphSearchData(Writer, Blueprint, SubGraphs, FFindInBlueprintSearchTags::FiB_SubGraphs, NULL);
-
-	// Gather all SCS components
-	// If we have an SCS but don't support it, then we remove it
-	if(Blueprint->SimpleConstructionScript)
-	{
-		// Remove any SCS variable nodes
-		const TArray<USCS_Node*>& AllSCSNodes = Blueprint->SimpleConstructionScript->GetAllNodes();
-		Writer->WriteArrayStart(FFindInBlueprintSearchTags::FiB_Components);
-		for (TFieldIterator<UProperty> PropertyIt(Blueprint->SkeletonGeneratedClass, EFieldIteratorFlags::ExcludeSuper); PropertyIt; ++PropertyIt)
-		{
-			UProperty* Property = *PropertyIt;
-			UObjectPropertyBase* Obj = Cast<UObjectPropertyBase>(Property);
-			const bool bComponentProperty = Obj && Obj->PropertyClass ? Obj->PropertyClass->IsChildOf<UActorComponent>() : false;
-			FName PropName = Property->GetFName();
-			if(bComponentProperty && FBlueprintEditorUtils::FindSCS_Node(Blueprint, PropName) != INDEX_NONE)
-			{
-				FEdGraphPinType PropertyPinType;
-				if(UEdGraphSchema_K2::StaticClass()->GetDefaultObject<UEdGraphSchema_K2>()->ConvertPropertyToPinType(Property, PropertyPinType))
-				{
-					Writer->WriteObjectStart();
-					{
-						Writer->WriteValue(FFindInBlueprintSearchTags::FiB_Name, FText::FromName(PropName));
-						Writer->WriteValue(FFindInBlueprintSearchTags::FiB_IsSCSComponent, true);
-						SavePinTypeToJson(Writer,  PropertyPinType);
-					}
-					Writer->WriteObjectEnd();
-				}
-			}
-		}
-		Writer->WriteArrayEnd(); // Components
-	}
-
-	Writer->WriteObjectEnd();
-	Writer->Close();
-
-	// Build the search metadata string for the asset tag (version + LUT + JSON)
-	SearchMetaData = FiBSerializationHelpers::Serialize(FSearchDataVersionInfo::Current.FiBDataVersion, false)
-		+ Writer->GetSerializedLookupTable()
-		+ SearchMetaData;
-
-	return SearchMetaData;
-}
-
 void FFindInBlueprintSearchManager::AddOrUpdateBlueprintSearchMetadata(UBlueprint* InBlueprint, bool bInForceReCache/* = false*/)
 {
 	CSV_SCOPED_TIMING_STAT(FindInBlueprint, AddOrUpdateBlueprintSearchMetadata);
@@ -1955,8 +1974,10 @@ void FFindInBlueprintSearchManager::AddOrUpdateBlueprintSearchMetadata(UBlueprin
 		// Cannot successfully gather most searchable data if there is no SkeletonGeneratedClass, so don't try, leave it as whatever it was last set to
 		if (InBlueprint->SkeletonGeneratedClass != nullptr)
 		{
+			using namespace BlueprintSearchMetaDataHelpers;
+
 			// Update search metadata string content
-			SearchArray[Index].Value = GatherBlueprintSearchMetadata(InBlueprint);
+			SearchArray[Index].Value = GatherBlueprintSearchMetadata<FFindInBlueprintJsonWriter>(InBlueprint);
 
 			// Update version info stored in database to latest
 			SearchArray[Index].VersionInfo = FSearchDataVersionInfo::Current;
@@ -2242,6 +2263,19 @@ void FFindInBlueprintSearchManager::BuildCache()
 	}
 }
 
+void FFindInBlueprintSearchManager::DumpCache(FArchive& Ar)
+{
+	FScopeLock ScopeLock(&SafeModifyCacheCriticalSection);
+
+	for (const FSearchData& SearchData : SearchArray)
+	{
+		if (SearchData.ImaginaryBlueprint.IsValid())
+		{
+			SearchData.ImaginaryBlueprint->DumpParsedObject(Ar);
+		}
+	}
+}
+
 FText FFindInBlueprintSearchManager::ConvertHexStringToFText(FString InHexString)
 {
 	TArray<uint8> SerializedData;
@@ -2269,6 +2303,12 @@ FString FFindInBlueprintSearchManager::ConvertFTextToHexString(FText InValue)
 	return BytesToHex(SerializedData.GetData(), SerializedData.Num());
 }
 
+FString FFindInBlueprintSearchManager::GenerateSearchIndexForDebugging(UBlueprint* InBlueprint)
+{
+	using namespace BlueprintSearchMetaDataHelpers;
+	return GatherBlueprintSearchMetadata<TFindInBlueprintJsonStringWriter<TPrettyJsonPrintPolicy<TCHAR>>>(InBlueprint);
+}
+
 void FFindInBlueprintSearchManager::OnCacheAllUnindexedAssets(bool bInSourceControlActive, bool bInCheckoutAndSave)
 {
 	// Multiple threads can be adding to this at the same time
@@ -2571,8 +2611,8 @@ TSharedPtr< FJsonObject > FFindInBlueprintSearchManager::ConvertJsonStringToObje
 	OutFTextLookupTable = LookupTable = FiBSerializationHelpers::Deserialize< TMap<int32, FText> >(ReaderStream, SizeOfData);
 
 	// The original BufferReader should be positioned at the Json
-	TSharedPtr< FJsonObject > JsonObject = NULL;
-	TSharedRef< TJsonReader<> > Reader = BlueprintSearchMetaDataHelpers::SearchMetaDataReader::Create( &ReaderStream, LookupTable );
+	TSharedPtr<FJsonObject> JsonObject = nullptr;
+	TSharedRef<TJsonReader<TCHAR>> Reader = TJsonReader<TCHAR>::Create(&ReaderStream);
 	FJsonSerializer::Deserialize( Reader, JsonObject );
 
 	return JsonObject;
@@ -2809,4 +2849,30 @@ TStatId FFindInBlueprintSearchManager::GetStatId() const
 	RETURN_QUICK_DECLARE_CYCLE_STAT(FFindInBlueprintSearchManager, STATGROUP_Tickables);
 }
 
+class FFiBDumpIndexCacheToFileExecHelper : public FSelfRegisteringExec
+{
+	virtual bool Exec(class UWorld* InWorld, const TCHAR* Cmd, FOutputDevice& Ar)
+	{
+		if (FParse::Command(&Cmd, TEXT("DUMPFIBINDEXCACHE")))
+		{
+			FString FileLocation = FPaths::ConvertRelativePathToFull(FPaths::ProjectIntermediateDir() + TEXT("BlueprintSearchTools"));
+			FString FullPath = FString::Printf(TEXT("%s/FullBlueprintIndexCacheDump.csv"), *FileLocation);
+
+			FArchive* DumpFile = IFileManager::Get().CreateFileWriter(*FullPath);
+			if (DumpFile)
+			{
+				FFindInBlueprintSearchManager::Get().DumpCache(*DumpFile);
+
+				DumpFile->Close();
+				delete DumpFile;
+
+				UE_LOG(LogFindInBlueprint, Log, TEXT("Wrote full index cache to %s"), *FullPath);
+			}
+			return true;
+		}
+		return false;
+	}
+};
+static FFiBDumpIndexCacheToFileExecHelper GFiBDumpIndexCacheToFileExec;
+
 #undef LOCTEXT_NAMESPACE
diff --git a/Engine/Source/Editor/Kismet/Private/ImaginaryBlueprintData.cpp b/Engine/Source/Editor/Kismet/Private/ImaginaryBlueprintData.cpp
--- a/Engine/Source/Editor/Kismet/Private/ImaginaryBlueprintData.cpp
+++ b/Engine/Source/Editor/Kismet/Private/ImaginaryBlueprintData.cpp
@@ -64,12 +64,6 @@ bool FComponentUniqueDisplay::operator==(const FComponentUniqueDisplay& Other)
 
 FCriticalSection FImaginaryFiBData::ParseChildDataCriticalSection;
 
-FImaginaryFiBData::FImaginaryFiBData(FImaginaryFiBDataWeakPtr InOuter)
-	: LookupTablePtr(nullptr)
-	, Outer(InOuter)
-{
-}
-
 FImaginaryFiBData::FImaginaryFiBData(FImaginaryFiBDataWeakPtr InOuter, TSharedPtr< FJsonObject > InUnparsedJsonObject, TMap<int32, FText>* InLookupTablePtr)
 	: UnparsedJsonObject(InUnparsedJsonObject)
 	, LookupTablePtr(InLookupTablePtr)
@@ -334,6 +328,40 @@ UObject* FImaginaryFiBData::GetObject(UBlueprint* InBlueprint) const
 	return CreateSearchResult(nullptr)->GetObject(InBlueprint);
 }
 
+void FImaginaryFiBData::DumpParsedObject(FArchive& Ar, int32 InTreeLevel) const
+{
+	FString CommaStr = TEXT(",");
+	for (int32 i = 0; i < InTreeLevel; ++i)
+	{
+		Ar.Serialize(TCHAR_TO_ANSI(*CommaStr), CommaStr.Len());
+	}
+
+	DumpParsedObject_Internal(Ar);
+
+	for (const TPair< FindInBlueprintsHelpers::FSimpleFTextKeyStorage, FSearchableValueInfo >& TagsValuePair : ParsedTagsAndValues)
+	{
+		FText Value = TagsValuePair.Value.GetDisplayText(*LookupTablePtr);
+		FString ValueAsString = Value.ToString();
+		ValueAsString.ReplaceInline(TEXT(" "), TEXT(""));
+
+		FString LineStr = FString::Printf(TEXT(",%s:%s"), *TagsValuePair.Key.Text.ToString(), *ValueAsString);
+		Ar.Serialize(TCHAR_TO_ANSI(*LineStr), LineStr.Len());
+	}
+
+	FString NewLine(TEXT("\n"));
+	Ar.Serialize(TCHAR_TO_ANSI(*NewLine), NewLine.Len());
+
+	for (const FImaginaryFiBDataSharedPtr Child : ParsedChildData)
+	{
+		Child->DumpParsedObject(Ar, InTreeLevel + 1);
+	}
+
+	if (InTreeLevel == 0)
+	{
+		Ar.Serialize(TCHAR_TO_ANSI(*NewLine), NewLine.Len());
+	}
+}
+
 ///////////////////////////
 // FFiBMetaData
 
@@ -432,6 +460,12 @@ void FCategorySectionHelper::ParseAllChildData_Internal(ESearchableValueStatus I
 	}
 }
 
+void FCategorySectionHelper::DumpParsedObject_Internal(FArchive& Ar) const
+{
+	FString OutputString = FString::Printf(TEXT("FCategorySectionHelper,CategoryName:%s,IsTagAndValueCategory:%s"), *CategoryName.ToString(), IsTagAndValueCategory() ? TEXT("true") : TEXT("false"));
+	Ar.Serialize(TCHAR_TO_ANSI(*OutputString), OutputString.Len());
+}
+
 //////////////////////////////////////////
 // FImaginaryBlueprint
 
@@ -462,6 +496,12 @@ FSearchResult FImaginaryBlueprint::CreateSearchResult_Internal(FSearchResult InP
 	return FSearchResult(new FFindInBlueprintsResult(ParsedTagsAndValues.Find(FindInBlueprintsHelpers::FSimpleFTextKeyStorage(FFindInBlueprintSearchTags::FiB_Path))->GetDisplayText(LookupTable)));
 }
 
+void FImaginaryBlueprint::DumpParsedObject_Internal(FArchive& Ar) const
+{
+	FString OutputString = FString::Printf(TEXT("FImaginaryBlueprint"));
+	Ar.Serialize(TCHAR_TO_ANSI(*OutputString), OutputString.Len());
+}
+
 UBlueprint* FImaginaryBlueprint::GetBlueprint() const
 {
 	return Cast<UBlueprint>(GetObject(nullptr));
@@ -574,6 +614,12 @@ FSearchResult FImaginaryGraph::CreateSearchResult_Internal(FSearchResult InParen
 	return FSearchResult(new FFindInBlueprintsGraph(FText::GetEmpty(), InParent, GraphType));
 }
 
+void FImaginaryGraph::DumpParsedObject_Internal(FArchive& Ar) const
+{
+	FString OutputString = FString::Printf(TEXT("FImaginaryGraph"));
+	Ar.Serialize(TCHAR_TO_ANSI(*OutputString), OutputString.Len());
+}
+
 bool FImaginaryGraph::IsCompatibleWithFilter(ESearchQueryFilter InSearchQueryFilter) const
 {
 	return InSearchQueryFilter == ESearchQueryFilter::AllFilter || 
@@ -643,6 +689,12 @@ FSearchResult FImaginaryGraphNode::CreateSearchResult_Internal(FSearchResult InP
 	return FSearchResult(new FFindInBlueprintsGraphNode(FText::GetEmpty(), InParent));
 }
 
+void FImaginaryGraphNode::DumpParsedObject_Internal(FArchive& Ar) const
+{
+	FString OutputString = FString::Printf(TEXT("FImaginaryGraphNode"));
+	Ar.Serialize(TCHAR_TO_ANSI(*OutputString), OutputString.Len());
+}
+
 bool FImaginaryGraphNode::IsCompatibleWithFilter(ESearchQueryFilter InSearchQueryFilter) const
 {
 	return InSearchQueryFilter == ESearchQueryFilter::AllFilter || InSearchQueryFilter == ESearchQueryFilter::NodesFilter;
@@ -741,6 +793,12 @@ FSearchResult FImaginaryProperty::CreateSearchResult_Internal(FSearchResult InPa
 	return FSearchResult(new FFindInBlueprintsProperty(FText::GetEmpty(), InParent));
 }
 
+void FImaginaryProperty::DumpParsedObject_Internal(FArchive& Ar) const
+{
+	FString OutputString = FString::Printf(TEXT("FImaginaryProperty"));
+	Ar.Serialize(TCHAR_TO_ANSI(*OutputString), OutputString.Len());
+}
+
 ESearchableValueStatus FImaginaryProperty::GetSearchabilityStatus(FString InKey)
 {
 	// This is a non-ideal way to assign searchability vs being a core display item and will be resolved in future versions of the FiB data in the AR
@@ -795,6 +853,12 @@ FSearchResult FImaginaryPin::CreateSearchResult_Internal(FSearchResult InParent)
 	return FSearchResult(new FFindInBlueprintsPin(FText::GetEmpty(), InParent, SchemaName));
 }
 
+void FImaginaryPin::DumpParsedObject_Internal(FArchive& Ar) const
+{
+	FString OutputString = FString::Printf(TEXT("FImaginaryPin"));
+	Ar.Serialize(TCHAR_TO_ANSI(*OutputString), OutputString.Len());
+}
+
 ESearchableValueStatus FImaginaryPin::GetSearchabilityStatus(FString InKey)
 {
 	if (FindInBlueprintsHelpers::IsTextEqualToString(FFindInBlueprintSearchTags::FiB_Name, InKey)
diff --git a/Engine/Source/Editor/Kismet/Private/SBlueprintEditorToolbar.cpp b/Engine/Source/Editor/Kismet/Private/SBlueprintEditorToolbar.cpp
--- a/Engine/Source/Editor/Kismet/Private/SBlueprintEditorToolbar.cpp
+++ b/Engine/Source/Editor/Kismet/Private/SBlueprintEditorToolbar.cpp
@@ -30,6 +30,7 @@
 #include "SBlueprintRevisionMenu.h"
 #include "ToolMenus.h"
 #include "BlueprintEditorContext.h"
+#include "FindInBlueprintManager.h"
 #include "ISourceCodeAccessor.h"
 #include "ISourceCodeAccessModule.h"
 
@@ -112,6 +113,16 @@ void FKismet2Menu::FillDeveloperMenu( FMenuBuilder& MenuBuilder )
 	}
 	MenuBuilder.EndSection();
 
+	if (FFindInBlueprintSearchManager::Get().ShouldEnableDeveloperMenuTools())
+	{
+		MenuBuilder.BeginSection("FileDeveloperSearchTools", LOCTEXT("SearchToolsHeading", "Search Tools"));
+		{
+			MenuBuilder.AddMenuEntry(FBlueprintEditorCommands::Get().GenerateSearchIndex);
+			MenuBuilder.AddMenuEntry(FBlueprintEditorCommands::Get().DumpCachedIndexData);
+		}
+		MenuBuilder.EndSection();
+	}
+
 	if (false)
 	{
 		MenuBuilder.BeginSection("FileDeveloperFindReferences");
diff --git a/Engine/Source/Editor/Kismet/Public/BlueprintEditor.h b/Engine/Source/Editor/Kismet/Public/BlueprintEditor.h
--- a/Engine/Source/Editor/Kismet/Public/BlueprintEditor.h
+++ b/Engine/Source/Editor/Kismet/Public/BlueprintEditor.h
@@ -574,6 +574,12 @@ public:
 	/** Can generate native code for current blueprint */
 	bool CanGenerateNativeCode() const;
 
+	/** Dumps the current blueprint search index to a JSON file for debugging purposes */
+	void OnGenerateSearchIndexForDebugging();
+
+	/** Dumps the currently-cached index data for the blueprint to a file for debugging */
+	void OnDumpCachedIndexDataForBlueprint();
+
 	/**
 	 * Check to see if we can customize the SCS editor for the passed-in scene component
 	 * @param	InComponentToCustomize	The component to check to see if a customization exists
diff --git a/Engine/Source/Editor/Kismet/Public/FindInBlueprintManager.h b/Engine/Source/Editor/Kismet/Public/FindInBlueprintManager.h
--- a/Engine/Source/Editor/Kismet/Public/FindInBlueprintManager.h
+++ b/Engine/Source/Editor/Kismet/Public/FindInBlueprintManager.h
@@ -167,6 +167,11 @@ struct FSearchData
 		, bMarkedForDeletion(false)
 	{
 	}
+
+	bool IsValid() const
+	{
+		return AssetPath != NAME_None;
+	}
 };
 
 /** Filters are used by functions for searching to decide whether items can call certain functions or match the requirements of a function */
@@ -430,6 +435,14 @@ public:
 	virtual TStatId GetStatId() const override;
 	//~ End FTickableObject Interface
 
+	/**
+	 * Given an asset path, locate and return a copy of its matching search data in the index cache.
+	 *
+	 * @param InAssetPath				Asset path (search index key).
+	 * @return							Matching search data from the index cache. Will return invalid (empty) search data if a matching entry was not found.
+	 */
+	FSearchData GetSearchDataForAssetPath(FName InAssetPath);
+
 	/**
 	 * Gathers the Blueprint's search metadata and adds or updates it in the cache
 	 *
@@ -479,12 +492,6 @@ public:
 	 */
 	const FSearchData* QuerySingleBlueprint(UBlueprint* InBlueprint, bool bInRebuildSearchData);
 
-	/** Converts a string of hex characters, previously converted by ConvertFTextToHexString, to an FText. */
-	static FText ConvertHexStringToFText(FString InHexString);
-
-	/** Serializes an FText to memory and converts the memory into a string of hex characters */
-	static FString ConvertFTextToHexString(FText InValue);
-
 	/** Returns the number of assets that are waiting to be re-indexed */
 	int32 GetNumberPendingAssets() const;
 
@@ -550,13 +557,12 @@ public:
 	/** Returns a weak reference to the widget that initiated the current caching operation */
 	TWeakPtr<SFindInBlueprints> GetSourceCachingWidget() const { return SourceCachingWidget; }
 
-	/** Given a fully constructed Find-in-Blueprint FString of searchable data, will parse and construct a JsonObject */
-	static TSharedPtr< class FJsonObject > ConvertJsonStringToObject(FSearchDataVersionInfo InVersionInfo, FString InJsonString, TMap<int32, FText>& OutFTextLookupTable);
-
 	void EnableGatheringData(bool bInEnableGatheringData) { bEnableGatheringData = bInEnableGatheringData; }
 
 	bool IsGatheringDataEnabled() const { return bEnableGatheringData; }
 
+	bool ShouldEnableDeveloperMenuTools() const { return bEnableDeveloperMenuTools; }
+
 	/** Find or create the global find results widget */
 	TSharedPtr<SFindInBlueprints> GetGlobalFindResults();
 
@@ -568,6 +574,22 @@ public:
 
 	void GlobalFindResultsClosed(const TSharedRef<SFindInBlueprints>& FindResults);
 
+	/** Dumps the full index cache to the given stream (for debugging purposes) */
+	void DumpCache(FArchive& Ar);
+
+public:
+	/** Converts a string of hex characters, previously converted by ConvertFTextToHexString, to an FText. */
+	static FText ConvertHexStringToFText(FString InHexString);
+
+	/** Serializes an FText to memory and converts the memory into a string of hex characters */
+	static FString ConvertFTextToHexString(FText InValue);
+
+	/** Given a fully constructed Find-in-Blueprint FString of searchable data, will parse and construct a JsonObject */
+	static TSharedPtr< class FJsonObject > ConvertJsonStringToObject(FSearchDataVersionInfo InVersionInfo, FString InJsonString, TMap<int32, FText>& OutFTextLookupTable);
+
+	/** Generates a human-readable search index for the given Blueprint (for debugging purposes) */
+	static FString GenerateSearchIndexForDebugging(UBlueprint* InBlueprint);
+
 private:
 	/** Initializes the FiB manager */
 	void Initialize();
@@ -599,9 +621,6 @@ private:
 	/** Callback hook from the Hot Reload manager that indicates that a module has been hot-reloaded */
 	void OnHotReload(bool bWasTriggeredAutomatically);
 
-	/** Helper to gathers the Blueprint's search metadata */
-	FString GatherBlueprintSearchMetadata(const UBlueprint* Blueprint);
-
 	/** Cleans the cache of any excess data from Blueprints that have been moved, renamed, or deleted. Occurs during post-garbage collection */
 	void CleanCache();
 
@@ -695,6 +714,9 @@ protected:
 	/** Whether CSV profiling has been enabled (default=false) */
 	bool bEnableCSVStatsProfiling;
 
+	/** Whether to enable Blueprint editor developer menu tools */
+	bool bEnableDeveloperMenuTools;
+
 	/** TRUE when the the FiB manager wants to pause all searches, helps manage the pausing procedure */
 	volatile bool bIsPausing;
 };
diff --git a/Engine/Source/Editor/Kismet/Public/ImaginaryBlueprintData.h b/Engine/Source/Editor/Kismet/Public/ImaginaryBlueprintData.h
--- a/Engine/Source/Editor/Kismet/Public/ImaginaryBlueprintData.h
+++ b/Engine/Source/Editor/Kismet/Public/ImaginaryBlueprintData.h
@@ -103,8 +103,7 @@ struct FComponentUniqueDisplay
 class FImaginaryFiBData : public ITextFilterExpressionContext, public TSharedFromThis<FImaginaryFiBData, ESPMode::ThreadSafe>
 {
 public:
-	FImaginaryFiBData(FImaginaryFiBDataWeakPtr InOuter);
-	FImaginaryFiBData(FImaginaryFiBDataWeakPtr InOuter, TSharedPtr< FJsonObject > InUnparsedJsonObject, TMap<int32, FText>* InLookupTablePtr);
+	FImaginaryFiBData(FImaginaryFiBDataWeakPtr InOuter, TSharedPtr< FJsonObject > InUnparsedJsonObject = TSharedPtr<FJsonObject>(), TMap<int32, FText>* InLookupTablePtr = nullptr);
 
 	/** ITextFilterExpressionContext Interface */
 	// We don't actually use these overrides, see FFiBContextHelper for how we call the alternate functions. These will assert if they are accidentally called.
@@ -116,7 +115,7 @@ public:
 	virtual bool IsCategory() const { return false; }
 
 	/** Returns TRUE if this item is considered a Tag and Value category, where it's contents should be considered no different than the parent owner */
-	virtual bool IsTagAndValueCategory() { return false; }
+	virtual bool IsTagAndValueCategory() const { return false; }
 
 	/** Checks if the filter is compatible with the current object, returns TRUE by default */
 	virtual bool IsCompatibleWithFilter(ESearchQueryFilter InSearchQueryFilter) const;
@@ -167,6 +166,9 @@ public:
 		return Outer;
 	}
 
+	/** Dumps the parsed object (including all children) to the given archive */
+	void DumpParsedObject(FArchive& Ar, int32 InTreeLevel = 0) const;
+
 	/** Builds a SearchTree ready to be displayed in the Find-in-Blueprints window */
 	static FSearchResult CreateSearchTree(FSearchResult InParentSearchResult, FImaginaryFiBDataWeakPtr InCurrentPointer, TArray< const FImaginaryFiBData* >& InValidSearchResults, TMultiMap< const FImaginaryFiBData*, FComponentUniqueDisplay >& InMatchingSearchComponents);
 
@@ -195,6 +197,9 @@ protected:
 	/** Internal version of the ParseAllChildData function, handles the bulk of the work */
 	virtual void ParseAllChildData_Internal(ESearchableValueStatus InSearchabilityOverride = ESearchableValueStatus::Searchable);
 
+	/** Internal helper function for dumping parsed object info */
+	virtual void DumpParsedObject_Internal(FArchive& Ar) const {}
+
 protected:
 	/** The unparsed Json object representing this item. Auto-cleared after parsing. Does not need to be declared as thread-safe because it's only accessed when parsing and that is already a critical section. */
 	TSharedPtr< FJsonObject > UnparsedJsonObject;
@@ -262,7 +267,7 @@ public:
 	/** FImaginaryFiBData Interface */
 	virtual void ParseAllChildData_Internal(ESearchableValueStatus InSearchabilityOverride/* = ESearchableValueStatus::ESearchableValueStatus::Searchable*/) override;
 	virtual bool IsCategory() const override { return true; }
-	virtual bool IsTagAndValueCategory() override { return bIsTagAndValue; }
+	virtual bool IsTagAndValueCategory() const override { return bIsTagAndValue; }
 	/** End FImaginaryFiBData Interface */
 
 	/** Returns the category name prepared for checking as a function */
@@ -275,6 +280,7 @@ protected:
 	/** FImaginaryFiBData Interface */
 	virtual bool CanCallFilter(ESearchQueryFilter InSearchQueryFilter) const override;
 	virtual FSearchResult CreateSearchResult_Internal(FSearchResult InParent) const override;
+	virtual void DumpParsedObject_Internal(FArchive& Ar) const override;
 	/** End FImaginaryFiBData Interface */
 
 protected:
@@ -304,6 +310,7 @@ protected:
 	/** FImaginaryFiBData Interface */
 	virtual bool TrySpecialHandleJsonValue(FText InKey, TSharedPtr< FJsonValue > InJsonValue) override;
 	virtual FSearchResult CreateSearchResult_Internal(FSearchResult InParent) const override;
+	virtual void DumpParsedObject_Internal(FArchive& Ar) const override;
 	/** End FImaginaryFiBData Interface */
 
 	/** Helper function to parse an array of Json Object representing graphs */
@@ -342,6 +349,7 @@ protected:
 	virtual bool TrySpecialHandleJsonValue(FText InKey, TSharedPtr< FJsonValue > InJsonValue) override;
 	virtual ESearchableValueStatus GetSearchabilityStatus(FString InKey) override;
 	virtual FSearchResult CreateSearchResult_Internal(FSearchResult InParent) const override;
+	virtual void DumpParsedObject_Internal(FArchive& Ar) const override;
 	/** End FImaginaryFiBData Interface */
 
 protected:
@@ -359,6 +367,7 @@ public:
 	virtual bool IsCompatibleWithFilter(ESearchQueryFilter InSearchQueryFilter) const override;
 	virtual bool CanCallFilter(ESearchQueryFilter InSearchQueryFilter) const override;
 	virtual void ParseAllChildData_Internal(ESearchableValueStatus InSearchabilityOverride/* = ESearchableValueStatus::Searchable*/) override;
+	virtual void DumpParsedObject_Internal(FArchive& Ar) const override;
 	/** End FImaginaryFiBData Interface */
 
 protected:
@@ -387,6 +396,7 @@ protected:
 	/** FImaginaryFiBData Interface */
 	virtual ESearchableValueStatus GetSearchabilityStatus(FString InKey) override;
 	virtual FSearchResult CreateSearchResult_Internal(FSearchResult InParent) const override;
+	virtual void DumpParsedObject_Internal(FArchive& Ar) const override;
 	/** End FImaginaryFiBData Interface */
 };
 
@@ -416,6 +426,7 @@ protected:
 	virtual bool TrySpecialHandleJsonValue(FText InKey, TSharedPtr< FJsonValue > InJsonValue);
 	virtual ESearchableValueStatus GetSearchabilityStatus(FString InKey) override;
 	virtual FSearchResult CreateSearchResult_Internal(FSearchResult InParent) const override;
+	virtual void DumpParsedObject_Internal(FArchive& Ar) const override;
 	/** End FImaginaryFiBData Interface */
 
 protected:
