Adjust code so that PVS isn't bothered by it
#rnx

diff --git a/Engine/Source/Editor/Kismet/Private/FindInBlueprintManager.cpp b/Engine/Source/Editor/Kismet/Private/FindInBlueprintManager.cpp
--- a/Engine/Source/Editor/Kismet/Private/FindInBlueprintManager.cpp
+++ b/Engine/Source/Editor/Kismet/Private/FindInBlueprintManager.cpp
@@ -532,7 +532,7 @@ namespace BlueprintSearchMetaDataHelpers
 
 			// Build the search metadata string for the asset tag (version + LUT + JSON)
 			*JsonOutput = FiBSerializationHelpers::Serialize(FSearchDataVersionInfo::Current.FiBDataVersion, false)
-				+ FiBSerializationHelpers::Serialize< TMap< int32, FText > >(LookupTable, true)
+				+ FiBSerializationHelpers::Serialize(LookupTable, true)
 				+ MoveTemp(*JsonOutput);
 
 			return bResult;
