PR #5613: Removed unnecessary background coloring in class selector drop downs (Contributed by [REDACTED_PERSON])

diff --git a/Engine/Source/Editor/DetailCustomizations/Private/ActorDetails.cpp b/Engine/Source/Editor/DetailCustomizations/Private/ActorDetails.cpp
--- a/Engine/Source/Editor/DetailCustomizations/Private/ActorDetails.cpp
+++ b/Engine/Source/Editor/DetailCustomizations/Private/ActorDetails.cpp
@@ -318,15 +318,10 @@ TSharedRef<SWidget> FActorDetails::OnGetConvertContent()
 
 	return
 		SNew(SBox)
-			.WidthOverride(280)
+		.WidthOverride(280)
+		.MaxDesiredHeight(500)
 		[
-			SNew(SVerticalBox)
-			+SVerticalBox::Slot()
-				.AutoHeight()
-				.MaxHeight(500)
-			[
-				ClassPicker
-			]			
+			ClassPicker
 		];
 }
 
diff --git a/Engine/Source/Editor/Kismet/Private/BlueprintDetailsCustomization.cpp b/Engine/Source/Editor/Kismet/Private/BlueprintDetailsCustomization.cpp
--- a/Engine/Source/Editor/Kismet/Private/BlueprintDetailsCustomization.cpp
+++ b/Engine/Source/Editor/Kismet/Private/BlueprintDetailsCustomization.cpp
@@ -4775,8 +4775,6 @@ bool FBlueprintGraphActionDetails::IsConstFunctionVisible() const
 	UK2Node_EditablePinBase * FunctionEntryNode = FunctionEntryNodePtr.Get();
 	if(FunctionEntryNode)
 	{
-		UBlueprint* Blueprint = FunctionEntryNode->GetBlueprint();
-
 		bSupportedType = FunctionEntryNode->IsA<UK2Node_FunctionEntry>();
 		bIsEditable = FunctionEntryNode->IsEditable();
 	}
@@ -5131,21 +5129,16 @@ TSharedRef<SWidget> FBlueprintInterfaceLayout::OnGetAddInterfaceMenuContent()
 	TArray<UBlueprint*> Blueprints;
 	Blueprints.Add(Blueprint);
 	TSharedRef<SWidget> ClassPicker = FBlueprintEditorUtils::ConstructBlueprintInterfaceClassPicker(Blueprints, FOnClassPicked::CreateSP(this, &FBlueprintInterfaceLayout::OnClassPicked));
-	return
-		SNew(SBorder)
-		.BorderImage(FEditorStyle::GetBrush("Menu.Background"))
+	// Achieving fixed width by nesting items within a fixed width box.
+	return SNew(SBox)
+		.WidthOverride(350.0f)
 		[
-			// Achieving fixed width by nesting items within a fixed width box.
-			SNew(SBox)
-			.WidthOverride(350.0f)
+			SNew(SVerticalBox)
+			+SVerticalBox::Slot()
+			.MaxHeight(400.0f)
+			.AutoHeight()
 			[
-				SNew(SVerticalBox)
-				+SVerticalBox::Slot()
-				.MaxHeight(400.0f)
-				.AutoHeight()
-				[
-					ClassPicker
-				]
+				ClassPicker
 			]
 		];
 }
@@ -5188,21 +5181,16 @@ TSharedRef<SWidget> FBlueprintGlobalOptionsDetails::GetParentClassMenuContent()
 	Blueprints.Add(GetBlueprintObj());
 	TSharedRef<SWidget> ClassPicker = FBlueprintEditorUtils::ConstructBlueprintParentClassPicker(Blueprints, FOnClassPicked::CreateSP(this, &FBlueprintGlobalOptionsDetails::OnClassPicked));
 
-	return
-		SNew(SBorder)
-		.BorderImage(FEditorStyle::GetBrush("Menu.Background"))
+	// Achieving fixed width by nesting items within a fixed width box.
+	return SNew(SBox)
+		.WidthOverride(350.0f)
 		[
-			// Achieving fixed width by nesting items within a fixed width box.
-			SNew(SBox)
-			.WidthOverride(350.0f)
+			SNew(SVerticalBox)
+			+SVerticalBox::Slot()
+			.MaxHeight(400.0f)
+			.AutoHeight()
 			[
-				SNew(SVerticalBox)
-				+SVerticalBox::Slot()
-				.MaxHeight(400.0f)
-				.AutoHeight()
-				[
-					ClassPicker
-				]
+				ClassPicker
 			]
 		];
 }
diff --git a/Engine/Source/Editor/PropertyEditor/Private/UserInterface/PropertyEditor/SPropertyEditorClass.cpp b/Engine/Source/Editor/PropertyEditor/Private/UserInterface/PropertyEditor/SPropertyEditorClass.cpp
--- a/Engine/Source/Editor/PropertyEditor/Private/UserInterface/PropertyEditor/SPropertyEditorClass.cpp
+++ b/Engine/Source/Editor/PropertyEditor/Private/UserInterface/PropertyEditor/SPropertyEditorClass.cpp
@@ -202,6 +202,7 @@ FText SPropertyEditorClass::GetDisplayValueAsString() const
 TSharedRef<SWidget> SPropertyEditorClass::GenerateClassPicker()
 {
 	FClassViewerInitializationOptions Options;
+	Options.bShowBackgroundBorder = false;
 	Options.bShowUnloadedBlueprints = true;
 	Options.bShowNoneOption = bAllowNone; 
 
diff --git a/Engine/Source/Editor/PropertyEditor/Private/UserInterface/PropertyEditor/SPropertyEditorEditInline.cpp b/Engine/Source/Editor/PropertyEditor/Private/UserInterface/PropertyEditor/SPropertyEditorEditInline.cpp
--- a/Engine/Source/Editor/PropertyEditor/Private/UserInterface/PropertyEditor/SPropertyEditorEditInline.cpp
+++ b/Engine/Source/Editor/PropertyEditor/Private/UserInterface/PropertyEditor/SPropertyEditorEditInline.cpp
@@ -167,6 +167,7 @@ bool SPropertyEditorEditInline::IsClassAllowed( UClass* CheckClass, bool bAllowA
 TSharedRef<SWidget> SPropertyEditorEditInline::GenerateClassPicker()
 {
 	FClassViewerInitializationOptions Options;
+	Options.bShowBackgroundBorder = false;
 	Options.bShowUnloadedBlueprints = true;
 	Options.NameTypeToDisplay = EClassViewerNameTypeToDisplay::DisplayName;
 
diff --git a/Engine/Source/Editor/PropertyEditor/Private/UserInterface/PropertyEditor/SPropertyMenuActorPicker.cpp b/Engine/Source/Editor/PropertyEditor/Private/UserInterface/PropertyEditor/SPropertyMenuActorPicker.cpp
--- a/Engine/Source/Editor/PropertyEditor/Private/UserInterface/PropertyEditor/SPropertyMenuActorPicker.cpp
+++ b/Engine/Source/Editor/PropertyEditor/Private/UserInterface/PropertyEditor/SPropertyMenuActorPicker.cpp
@@ -95,11 +95,7 @@ void SPropertyMenuActorPicker::Construct( const FArguments& InArgs )
 			.WidthOverride(PropertyEditorAssetConstants::SceneOutlinerWindowSize.X)
 			.HeightOverride(PropertyEditorAssetConstants::SceneOutlinerWindowSize.Y)
 			[
-				SNew( SBorder )
-				.BorderImage( FEditorStyle::GetBrush("Menu.Background") )
-				[
-					SceneOutlinerModule.CreateSceneOutliner(InitOptions, FOnActorPicked::CreateSP(this, &SPropertyMenuActorPicker::OnActorSelected))
-				]
+				SceneOutlinerModule.CreateSceneOutliner(InitOptions, FOnActorPicked::CreateSP(this, &SPropertyMenuActorPicker::OnActorSelected))
 			];
 
 		MenuBuilder.AddWidget(MenuContent.ToSharedRef(), FText::GetEmpty(), true);
diff --git a/Engine/Source/Editor/UnrealEd/Private/Kismet2/BlueprintEditorUtils.cpp b/Engine/Source/Editor/UnrealEd/Private/Kismet2/BlueprintEditorUtils.cpp
--- a/Engine/Source/Editor/UnrealEd/Private/Kismet2/BlueprintEditorUtils.cpp
+++ b/Engine/Source/Editor/UnrealEd/Private/Kismet2/BlueprintEditorUtils.cpp
@@ -7668,6 +7668,7 @@ TSharedRef<SWidget> FBlueprintEditorUtils::ConstructBlueprintParentClassPicker(
 	// Fill in options
 	FClassViewerInitializationOptions Options;
 	Options.Mode = EClassViewerMode::ClassPicker;
+	Options.bShowBackgroundBorder = false;
 
 	TSharedPtr<FBlueprintReparentFilter> Filter = MakeShareable(new FBlueprintReparentFilter);
 	Options.ClassFilter = Filter;
@@ -7823,6 +7824,7 @@ TSharedRef<SWidget> FBlueprintEditorUtils::ConstructBlueprintInterfaceClassPicke
 	// Fill in options
 	FClassViewerInitializationOptions Options;
 	Options.Mode = EClassViewerMode::ClassPicker;
+	Options.bShowBackgroundBorder = false;
 
 	TSharedPtr<FBlueprintInterfaceFilter> Filter = MakeShareable(new FBlueprintInterfaceFilter);
 	Options.ClassFilter = Filter;
