Fixed blueprint File->Developers sub-menu not opening until clicked.

Years ago a FName was being passed to a function where a boolean was expected, it should have been changed to false instead of true.

diff --git a/Engine/Source/Editor/Kismet/Private/SBlueprintEditorToolbar.cpp b/Engine/Source/Editor/Kismet/Private/SBlueprintEditorToolbar.cpp
--- a/Engine/Source/Editor/Kismet/Private/SBlueprintEditorToolbar.cpp
+++ b/Engine/Source/Editor/Kismet/Private/SBlueprintEditorToolbar.cpp
@@ -88,7 +88,7 @@ void FKismet2Menu::FillFileMenuBlueprintSection( FMenuBuilder& MenuBuilder, FBlu
 				LOCTEXT("DeveloperMenu", "Developer"),
 				LOCTEXT("DeveloperMenu_ToolTip", "Open the developer menu"),
 				FNewMenuDelegate::CreateStatic( &FKismet2Menu::FillDeveloperMenu ),
-				true);
+				false);
 		}
 		MenuBuilder.EndSection();
 	}
