Switched deprecated CanSpawnTab in FindInBlueprintManager.

diff --git a/Engine/Source/Editor/Kismet/Private/FindInBlueprintManager.cpp b/Engine/Source/Editor/Kismet/Private/FindInBlueprintManager.cpp
--- a/Engine/Source/Editor/Kismet/Private/FindInBlueprintManager.cpp
+++ b/Engine/Source/Editor/Kismet/Private/FindInBlueprintManager.cpp
@@ -2588,7 +2588,7 @@ void FFindInBlueprintSearchManager::EnableGlobalFindResults(bool bEnable)
 		for (int32 TabIdx = 0; TabIdx < ARRAY_COUNT(GlobalFindResultsTabIDs); TabIdx++)
 		{
 			const FName TabID = GlobalFindResultsTabIDs[TabIdx];
-			if (!GlobalTabManager->CanSpawnTab(TabID))
+			if (!GlobalTabManager->HasTabSpawner(TabID))
 			{
 				const FText DisplayName = FText::Format(LOCTEXT("GlobalFindResultsDisplayName", "Find in Blueprints {0}"), FText::AsNumber(TabIdx + 1));
 
