FPropertyAndParent now contains the full chain of parent properties of a property node, so that containers nested in containers can be checked for selection in SKismetInspector::IsPropertyVisible.

diff --git a/Engine/Plugins/Experimental/ControlRig/Source/ControlRigEditor/Private/EditMode/SControlRigEditModeTools.cpp b/Engine/Plugins/Experimental/ControlRig/Source/ControlRigEditor/Private/EditMode/SControlRigEditModeTools.cpp
--- a/Engine/Plugins/Experimental/ControlRig/Source/ControlRigEditor/Private/EditMode/SControlRigEditModeTools.cpp
+++ b/Engine/Plugins/Experimental/ControlRig/Source/ControlRigEditor/Private/EditMode/SControlRigEditModeTools.cpp
@@ -217,7 +217,8 @@ bool SControlRigEditModeTools::ShouldShowPropertyOnDetailCustomization(const FPr
 		}
 	}
 
-	return ShouldPropertyBeVisible(InPropertyAndParent.Property) || (InPropertyAndParent.ParentProperty && ShouldPropertyBeVisible(*InPropertyAndParent.ParentProperty));
+	return ShouldPropertyBeVisible(InPropertyAndParent.Property) || 
+		(InPropertyAndParent.ParentProperties.Num() > 0 && ShouldPropertyBeVisible(*InPropertyAndParent.ParentProperties[0]));
 }
 
 bool SControlRigEditModeTools::IsReadOnlyPropertyOnDetailCustomization(const FPropertyAndParent& InPropertyAndParent) const
@@ -245,7 +246,8 @@ bool SControlRigEditModeTools::IsReadOnlyPropertyOnDetailCustomization(const FPr
 		}
 	}
 
-	return !(ShouldPropertyBeEnabled(InPropertyAndParent.Property) || (InPropertyAndParent.ParentProperty && ShouldPropertyBeEnabled(*InPropertyAndParent.ParentProperty)));
+	return !(ShouldPropertyBeEnabled(InPropertyAndParent.Property) || 
+		(InPropertyAndParent.ParentProperties.Num() > 0 && ShouldPropertyBeEnabled(*InPropertyAndParent.ParentProperties[0])));
 }
 
 static bool bPickerChangingSelection = false;
diff --git a/Engine/Source/Editor/Kismet/Private/SKismetInspector.cpp b/Engine/Source/Editor/Kismet/Private/SKismetInspector.cpp
--- a/Engine/Source/Editor/Kismet/Private/SKismetInspector.cpp
+++ b/Engine/Source/Editor/Kismet/Private/SKismetInspector.cpp
@@ -765,6 +765,21 @@ bool SKismetInspector::IsStructViewPropertyReadOnly(const struct FPropertyAndPar
 	return false;
 }
 
+bool SKismetInspector::IsAnyParentContainerSelected(const FPropertyAndParent& PropertyAndParent) const
+{
+	for (const UProperty* CurrentProperty : PropertyAndParent.ParentProperties)
+	{
+		const UProperty* CurrentOuter = Cast<UProperty>(CurrentProperty->GetOuter());
+
+		if (CurrentOuter != nullptr && SelectedObjectProperties.Find(MakeWeakObjectPtr(const_cast<UProperty*>(CurrentOuter))))
+		{
+			return true;
+		}
+	}
+
+	return false;
+} 
+
 bool SKismetInspector::IsPropertyVisible( const FPropertyAndParent& PropertyAndParent ) const
 {
 	const UProperty& Property = PropertyAndParent.Property;
@@ -832,30 +847,23 @@ bool SKismetInspector::IsPropertyVisible( const FPropertyAndParent& PropertyAndP
 		// If the current property is selected, it is visible.
 		return true;
 	}
-	else if ( PropertyAndParent.ParentProperty )
+	else if ( PropertyAndParent.ParentProperties.Num() > 0 && SelectedObjectProperties.Num() > 0 )
 	{
-		const UProperty* ParentProperty = PropertyAndParent.ParentProperty;
-		const UProperty* ParentPropertyOuter = nullptr;
-		
-		if (ParentProperty)
-		{
-			ParentPropertyOuter = Cast<UProperty>(ParentProperty->GetOuter());
-		}
+		const UProperty* ParentProperty = PropertyAndParent.ParentProperties[0];
 
 		if ( SelectedObjectProperties.Find( MakeWeakObjectPtr( const_cast<UProperty*>( ParentProperty ) ) ) )
 		{
 			// If its parent is selected, it should be visible
 			return true;
 		}
-		else if ( ParentPropertyOuter && SelectedObjectProperties.Find( MakeWeakObjectPtr( const_cast<UProperty*>( ParentPropertyOuter ) ) ) )
+		else if ( IsAnyParentContainerSelected(PropertyAndParent) )
 		{
-			// If its parent is part of a container and the container property is selected, it should be visible
 			return true;
 		}
 	}
 
 
-	return !SelectedObjectProperties.Num();
+	return SelectedObjectProperties.Num() == 0;
 }
 
 void SKismetInspector::SetPropertyWindowContents(TArray<UObject*> Objects)
diff --git a/Engine/Source/Editor/Kismet/Public/SKismetInspector.h b/Engine/Source/Editor/Kismet/Public/SKismetInspector.h
--- a/Engine/Source/Editor/Kismet/Public/SKismetInspector.h
+++ b/Engine/Source/Editor/Kismet/Public/SKismetInspector.h
@@ -202,4 +202,6 @@ protected:
 
 	ECheckBoxState GetPublicViewCheckboxState() const;
 	void SetPublicViewCheckboxState(ECheckBoxState InIsChecked);
+
+	bool IsAnyParentContainerSelected(const FPropertyAndParent& PropertyAndParent) const;
 };
diff --git a/Engine/Source/Editor/LevelEditor/Private/SActorDetails.cpp b/Engine/Source/Editor/LevelEditor/Private/SActorDetails.cpp
--- a/Engine/Source/Editor/LevelEditor/Private/SActorDetails.cpp
+++ b/Engine/Source/Editor/LevelEditor/Private/SActorDetails.cpp
@@ -617,7 +617,8 @@ bool SActorDetails::IsPropertyReadOnly(const FPropertyAndParent& PropertyAndPare
 		{
 			TSet<const UProperty*> UCSModifiedProperties;
 			Component->GetUCSModifiedProperties(UCSModifiedProperties);
-			if (UCSModifiedProperties.Contains(&PropertyAndParent.Property) || (PropertyAndParent.ParentProperty && UCSModifiedProperties.Contains(PropertyAndParent.ParentProperty)))
+			if (UCSModifiedProperties.Contains(&PropertyAndParent.Property) || 
+				(PropertyAndParent.ParentProperties.Num() > 0 && UCSModifiedProperties.Contains(PropertyAndParent.ParentProperties[0])))
 			{
 				bIsReadOnly = true;
 				break;
diff --git a/Engine/Source/Editor/PropertyEditor/Private/DetailLayoutBuilderImpl.cpp b/Engine/Source/Editor/PropertyEditor/Private/DetailLayoutBuilderImpl.cpp
--- a/Engine/Source/Editor/PropertyEditor/Private/DetailLayoutBuilderImpl.cpp
+++ b/Engine/Source/Editor/PropertyEditor/Private/DetailLayoutBuilderImpl.cpp
@@ -626,14 +626,14 @@ bool FDetailLayoutBuilderImpl::IsPropertyVisible( TSharedRef<IPropertyHandle> Pr
 	{
 		TArray<UObject*> OuterObjects;
 		PropertyHandle->GetOuterObjects(OuterObjects);
-
-		TArray<TWeakObjectPtr<UObject> > Objects;
+		
+		TArray<TWeakObjectPtr<UObject>> Objects;
 		for (auto OuterObject : OuterObjects)
 		{
 			Objects.Add(OuterObject);
 		}
 
-		FPropertyAndParent PropertyAndParent(*PropertyHandle->GetProperty(), PropertyHandle->GetParentHandle().IsValid() ? PropertyHandle->GetParentHandle()->GetProperty() : nullptr, Objects );
+		FPropertyAndParent PropertyAndParent(PropertyHandle, Objects);
 
 		return IsPropertyVisible(PropertyAndParent);
 	}
diff --git a/Engine/Source/Editor/PropertyEditor/Private/DetailLayoutHelpers.cpp b/Engine/Source/Editor/PropertyEditor/Private/DetailLayoutHelpers.cpp
--- a/Engine/Source/Editor/PropertyEditor/Private/DetailLayoutHelpers.cpp
+++ b/Engine/Source/Editor/PropertyEditor/Private/DetailLayoutHelpers.cpp
@@ -117,7 +117,8 @@ namespace DetailLayoutHelpers
 						}
 					}
 
-					FPropertyAndParent PropertyAndParent(*Property, ParentProperty, Objects);
+					TSharedPtr<IPropertyHandle> ChildHandle = PropertyEditorHelpers::GetPropertyHandle(ChildNodePtr.ToSharedRef(), nullptr, nullptr);
+					FPropertyAndParent PropertyAndParent(ChildHandle.ToSharedRef(), Objects);
 					const bool bIsUserVisible = InUpdateArgs.IsPropertyVisible(PropertyAndParent);
 
 					// Inners of customized in structs should not be taken into consideration for customizing.  They are not designed to be individually customized when their parent is already customized
diff --git a/Engine/Source/Editor/PropertyEditor/Private/DetailPropertyRow.cpp b/Engine/Source/Editor/PropertyEditor/Private/DetailPropertyRow.cpp
--- a/Engine/Source/Editor/PropertyEditor/Private/DetailPropertyRow.cpp
+++ b/Engine/Source/Editor/PropertyEditor/Private/DetailPropertyRow.cpp
@@ -330,7 +330,7 @@ void FDetailPropertyRow::GenerateChildrenForPropertyNode( TSharedPtr<FPropertyNo
 					GenerateChildrenForPropertyNode( ChildNode, OutChildren );
 				}
 				// Only struct children can have custom visibility that is different from their parent.
-				else if ( !bStructProperty || LayoutBuilder.IsPropertyVisible( FPropertyAndParent(*ChildNode->GetProperty(), ParentProperty, Objects ) ) )
+				else if ( !bStructProperty || LayoutBuilder.IsPropertyVisible( FPropertyAndParent( PropertyEditorHelpers::GetPropertyHandle(ChildNode.ToSharedRef(), nullptr, nullptr).ToSharedRef(), Objects ) ) )
 				{	
 					TArray<TSharedRef<FDetailTreeNode>> PropNodes;
 					bool bHasKeyNode = false;
diff --git a/Engine/Source/Editor/PropertyEditor/Private/SPropertyTreeViewImpl.cpp b/Engine/Source/Editor/PropertyEditor/Private/SPropertyTreeViewImpl.cpp
--- a/Engine/Source/Editor/PropertyEditor/Private/SPropertyTreeViewImpl.cpp
+++ b/Engine/Source/Editor/PropertyEditor/Private/SPropertyTreeViewImpl.cpp
@@ -695,7 +695,8 @@ void SPropertyTreeViewImpl::OnGetChildrenForPropertyNode( TSharedPtr<FPropertyNo
 				}
 			}
 
-			FPropertyAndParent PropertyAndParent(*Property, InPropertyNode->GetProperty(), Objects);
+			TSharedPtr<IPropertyHandle> ChildHandle = PropertyEditorHelpers::GetPropertyHandle(ChildNode.ToSharedRef(), nullptr, nullptr);
+			FPropertyAndParent PropertyAndParent(ChildHandle.ToSharedRef(), Objects);
 
 			bPropertyVisible = IsPropertyVisible.Execute( PropertyAndParent );
 		}
diff --git a/Engine/Source/Editor/PropertyEditor/Public/PropertyEditorDelegates.cpp b/Engine/Source/Editor/PropertyEditor/Public/PropertyEditorDelegates.cpp
new file mode 100644
--- /dev/null
+++ b/Engine/Source/Editor/PropertyEditor/Public/PropertyEditorDelegates.cpp
@@ -0,0 +1,21 @@
+// Copyright 1998-2019 Epic Games, Inc. All Rights Reserved.
+
+#include "PropertyEditorDelegates.h"
+#include "PropertyHandle.h"
+
+FPropertyAndParent::FPropertyAndParent(const TSharedRef<IPropertyHandle>& InPropertyHandle, const TArray<TWeakObjectPtr<UObject>>& InObjects) :
+	Property(*InPropertyHandle->GetProperty()),
+	Objects(InObjects)
+{
+	TSharedPtr<IPropertyHandle> ParentHandle = InPropertyHandle->GetParentHandle();
+	while (ParentHandle.IsValid())
+	{
+		const UProperty* ParentProperty = ParentHandle->GetProperty();
+		if (ParentProperty != nullptr)
+		{
+			ParentProperties.Add(ParentProperty);
+		}
+
+		ParentHandle = ParentHandle->GetParentHandle();
+	}
+}
\ No newline at end of file
diff --git a/Engine/Source/Editor/PropertyEditor/Public/PropertyEditorDelegates.h b/Engine/Source/Editor/PropertyEditor/Public/PropertyEditorDelegates.h
--- a/Engine/Source/Editor/PropertyEditor/Public/PropertyEditorDelegates.h
+++ b/Engine/Source/Editor/PropertyEditor/Public/PropertyEditorDelegates.h
@@ -7,20 +7,17 @@
 
 class FPropertyPath;
 struct FPropertyChangedEvent;
+class IPropertyHandle;
 
 struct FPropertyAndParent
 {
-	FPropertyAndParent( const UProperty& InProperty, const UProperty* const InParentProperty, const TArray< TWeakObjectPtr<UObject> >& InObjects )
-		: Property( InProperty )
-		, ParentProperty( InParentProperty )
-		, Objects( InObjects )
-	{}
+	FPropertyAndParent(const TSharedRef<IPropertyHandle>& InPropertyHandle, const TArray< TWeakObjectPtr< UObject > >& InObjects);
 
 	/** The property always exists */
 	const UProperty& Property;
 
-	/** The parent property may not exist */
-	const UProperty* const ParentProperty;
+	/** The entire chain of parent properties, all the way to the property root. ParentProperties[0] is the immediate parent.*/
+	TArray< const UProperty* > ParentProperties;
 
 	/** The objects for these properties */
 	TArray< TWeakObjectPtr< UObject > > Objects;
