Buildfix
#rnx

diff --git a/Engine/Source/Editor/Kismet/Public/SSCSEditorMenuContext.h b/Engine/Source/Editor/Kismet/Public/SSCSEditorMenuContext.h
--- a/Engine/Source/Editor/Kismet/Public/SSCSEditorMenuContext.h
+++ b/Engine/Source/Editor/Kismet/Public/SSCSEditorMenuContext.h
@@ -6,6 +6,8 @@
 
 #include "SSCSEditorMenuContext.generated.h"
 
+class SSCSEditor;
+
 UCLASS()
 class USSCSEditorMenuContext : public UObject
 {
