Refactoring the majority of AssetEditorManager cases to AssetEditorSubsystem

diff --git a/Engine/Source/Editor/FontEditor/Private/SCompositeFontEditor.cpp b/Engine/Source/Editor/FontEditor/Private/SCompositeFontEditor.cpp
--- a/Engine/Source/Editor/FontEditor/Private/SCompositeFontEditor.cpp
+++ b/Engine/Source/Editor/FontEditor/Private/SCompositeFontEditor.cpp
@@ -12,7 +12,7 @@
 #include "EditorStyleSet.h"
 #include "EditorFontGlyphs.h"
 #include "EditorDirectories.h"
-#include "Toolkits/AssetEditorManager.h"
+
 #include "IFontEditor.h"
 #include "Widgets/Text/SInlineEditableTextBlock.h"
 #include "Widgets/Input/SNumericEntryBox.h"
@@ -27,6 +27,7 @@
 #include "Misc/FileHelper.h"
 #include "IContentBrowserSingleton.h"
 #include "FileHelpers.h"
+#include "Subsystems/AssetEditorSubsystem.h"
 
 #define LOCTEXT_NAMESPACE "FontEditor"
 
@@ -925,7 +926,7 @@ void STypefaceEntryEditor::OnTypefaceEntryFontPathPicked(const FString& InNewFon
 			if (NewFontFaceAsset)
 			{
 				OnFontFaceAssetChanged(FAssetData(NewFontFaceAsset));
-				FAssetEditorManager::Get().OpenEditorForAsset(NewFontFaceAsset);
+				GEditor->GetEditorSubsystem<UAssetEditorSubsystem>()->OpenEditorForAsset(NewFontFaceAsset);
 			}
 		}
 	}
@@ -1056,7 +1057,7 @@ FReply STypefaceEntryEditor::OnUpgradeDataClicked()
 		if (NewFontFaceAsset)
 		{
 			OnFontFaceAssetChanged(FAssetData(NewFontFaceAsset));
-			FAssetEditorManager::Get().OpenEditorForAsset(NewFontFaceAsset);
+			GEditor->GetEditorSubsystem<UAssetEditorSubsystem>()->OpenEditorForAsset(NewFontFaceAsset);
 		}
 	}
 
diff --git a/Engine/Source/Editor/GameProjectGeneration/Private/SNewClassDialog.cpp b/Engine/Source/Editor/GameProjectGeneration/Private/SNewClassDialog.cpp
--- a/Engine/Source/Editor/GameProjectGeneration/Private/SNewClassDialog.cpp
+++ b/Engine/Source/Editor/GameProjectGeneration/Private/SNewClassDialog.cpp
@@ -34,10 +34,11 @@
 #include "TutorialMetaData.h"
 #include "Kismet2/KismetEditorUtilities.h"
 #include "AssetRegistryModule.h"
-#include "Toolkits/AssetEditorManager.h"
+
 #include "Framework/Notifications/NotificationManager.h"
 #include "Widgets/Notifications/SNotificationList.h"
 #include "FeaturedClasses.inl"
+#include "Subsystems/AssetEditorSubsystem.h"
 
 #define LOCTEXT_NAMESPACE "GameProjectGeneration"
 
@@ -1217,7 +1218,7 @@ void SNewClassDialog::FinishClicked()
 					ContentBrowserModule.Get().SyncBrowserToAssets(SyncAssets);
 
 					// Open the editor for the new asset
-					FAssetEditorManager::Get().OpenEditorForAsset(NewBP);
+					GEditor->GetEditorSubsystem<UAssetEditorSubsystem>()->OpenEditorForAsset(NewBP);
 
 					// Successfully created the code and potentially opened the IDE. Close the dialog.
 					CloseContainingWindow();
diff --git a/Engine/Source/Editor/HierarchicalLODOutliner/Private/HLODOutliner.cpp b/Engine/Source/Editor/HierarchicalLODOutliner/Private/HLODOutliner.cpp
--- a/Engine/Source/Editor/HierarchicalLODOutliner/Private/HLODOutliner.cpp
+++ b/Engine/Source/Editor/HierarchicalLODOutliner/Private/HLODOutliner.cpp
@@ -48,9 +48,10 @@
 #include "Widgets/Layout/SWrapBox.h"
 #include "Framework/Notifications/NotificationManager.h"
 #include "Widgets/Notifications/SNotificationList.h"
-#include "Toolkits/AssetEditorManager.h"
+
 #include "Engine/HLODProxy.h"
 #include "Misc/ScopedSlowTask.h"
+#include "Subsystems/AssetEditorSubsystem.h"
 
 #define LOCTEXT_NAMESPACE "HLODOutliner"
 
@@ -2133,7 +2134,7 @@ namespace HLODOutliner
 								GetObjectsWithOuter(HLODPackage, Objects);
 								for(UObject* PackageObject : Objects)
 								{
-									FAssetEditorManager::Get().CloseAllEditorsForAsset(PackageObject);
+									GEditor->GetEditorSubsystem<UAssetEditorSubsystem>()->CloseAllEditorsForAsset(PackageObject);
 								}
 							}
 						}
diff --git a/Engine/Source/Editor/IntroTutorials/Private/EditorTutorial.cpp b/Engine/Source/Editor/IntroTutorials/Private/EditorTutorial.cpp
--- a/Engine/Source/Editor/IntroTutorials/Private/EditorTutorial.cpp
+++ b/Engine/Source/Editor/IntroTutorials/Private/EditorTutorial.cpp
@@ -8,12 +8,13 @@
 #include "Settings/ContentBrowserSettings.h"
 #include "EngineGlobals.h"
 #include "Editor.h"
-#include "Toolkits/AssetEditorManager.h"
+
 #include "IIntroTutorials.h"
 #include "IntroTutorials.h"
 #include "LevelEditor.h"
 #include "Widgets/Docking/SDockTab.h"
 #include "Misc/RuntimeErrors.h"
+#include "Subsystems/AssetEditorSubsystem.h"
 
 #if WITH_EDITORONLY_DATA
 namespace
@@ -102,7 +103,7 @@ void UEditorTutorial::OpenAsset(UObject* Asset)
 {
 	if (ensureAsRuntimeWarning(Asset != nullptr))
 	{
-		FAssetEditorManager::Get().OpenEditorForAsset(Asset);
+		GEditor->GetEditorSubsystem<UAssetEditorSubsystem>()->OpenEditorForAsset(Asset);
 	}
 }
 
diff --git a/Engine/Source/Editor/IntroTutorials/Private/STutorialOverlay.cpp b/Engine/Source/Editor/IntroTutorials/Private/STutorialOverlay.cpp
--- a/Engine/Source/Editor/IntroTutorials/Private/STutorialOverlay.cpp
+++ b/Engine/Source/Editor/IntroTutorials/Private/STutorialOverlay.cpp
@@ -13,12 +13,13 @@
 #include "Framework/Docking/TabManager.h"
 #include "Engine/Blueprint.h"
 #include "Editor.h"
-#include "Toolkits/AssetEditorManager.h"
+
 #include "IntroTutorials.h"
 #include "STutorialContent.h"
 #include "Kismet2/KismetEditorUtilities.h"
 #include "LevelEditor.h"
 #include "Kismet2/BlueprintEditorUtils.h"
+#include "Subsystems/AssetEditorSubsystem.h"
 
 static FName IntroTutorialsModuleName("IntroTutorials");
 
@@ -204,13 +205,13 @@ void STutorialOverlay::OpenBrowserForWidgetAnchor(UEditorTutorial* InTutorial, c
 			FString Name = WidgetContent.WidgetAnchor.OuterName.RightChop(Space + 1);
 			TArray<FString> AssetPaths;
 			AssetPaths.Add(Name);
-			FAssetEditorManager::Get().OpenEditorsForAssets(AssetPaths);
+			GEditor->GetEditorSubsystem<UAssetEditorSubsystem>()->OpenEditorsForAssets(AssetPaths);
 			UObject* Blueprint = FindObject<UObject>(ANY_PACKAGE, *Name);
 
 			// If we found a blueprint
 			if(Blueprint != nullptr)
 			{
-				IAssetEditorInstance* PotentialAssetEditor = FAssetEditorManager::Get().FindEditorForAsset(Blueprint, false);
+				IAssetEditorInstance* PotentialAssetEditor = GEditor->GetEditorSubsystem<UAssetEditorSubsystem>()->FindEditorForAsset(Blueprint, false);
 				if(PotentialAssetEditor != nullptr)
 				{
 					AssetEditor = PotentialAssetEditor;
@@ -225,7 +226,7 @@ void STutorialOverlay::OpenBrowserForWidgetAnchor(UEditorTutorial* InTutorial, c
 			UObject* AssetObject = InTutorial->AssetToUse.ResolveObject();
 			if(AssetObject != nullptr)
 			{
-				IAssetEditorInstance* PotentialAssetEditor = FAssetEditorManager::Get().FindEditorForAsset(AssetObject, false);
+				IAssetEditorInstance* PotentialAssetEditor = GEditor->GetEditorSubsystem<UAssetEditorSubsystem>()->FindEditorForAsset(AssetObject, false);
 				if(PotentialAssetEditor != nullptr)
 				{
 					AssetEditor = PotentialAssetEditor;
diff --git a/Engine/Source/Editor/IntroTutorials/Private/STutorialRoot.cpp b/Engine/Source/Editor/IntroTutorials/Private/STutorialRoot.cpp
--- a/Engine/Source/Editor/IntroTutorials/Private/STutorialRoot.cpp
+++ b/Engine/Source/Editor/IntroTutorials/Private/STutorialRoot.cpp
@@ -4,7 +4,7 @@
 #include "Framework/Application/SlateApplication.h"
 #include "Templates/SubclassOf.h"
 #include "Kismet/GameplayStatics.h"
-#include "Toolkits/AssetEditorManager.h"
+
 #include "IntroTutorials.h"
 #include "EditorTutorial.h"
 #include "SEditorTutorials.h"
@@ -17,6 +17,7 @@
 #include "Interfaces/IAnalyticsProvider.h"
 #include "Framework/Notifications/NotificationManager.h"
 #include "Widgets/Notifications/SNotificationList.h"
+#include "Subsystems/AssetEditorSubsystem.h"
 
 #define LOCTEXT_NAMESPACE "STutorialRoot"
 
@@ -133,7 +134,7 @@ void STutorialRoot::LaunchTutorial(UEditorTutorial* InTutorial, IIntroTutorials:
 		{
 			TArray<FString> AssetPaths;
 			AssetPaths.Add(InTutorial->AssetToUse.ToString());
-			FAssetEditorManager::Get().OpenEditorsForAssets(AssetPaths);
+			GEditor->GetEditorSubsystem<UAssetEditorSubsystem>()->OpenEditorsForAssets(AssetPaths);
 
 			UObject* Asset = InTutorial->AssetToUse.ResolveObject();
 			if(Asset != nullptr)
diff --git a/Engine/Source/Editor/IntroTutorials/Private/TutorialText.cpp b/Engine/Source/Editor/IntroTutorials/Private/TutorialText.cpp
--- a/Engine/Source/Editor/IntroTutorials/Private/TutorialText.cpp
+++ b/Engine/Source/Editor/IntroTutorials/Private/TutorialText.cpp
@@ -7,7 +7,7 @@
 #include "Framework/Text/SlateTextRun.h"
 #include "Framework/Text/TextDecorators.h"
 #include "Engine/Blueprint.h"
-#include "Toolkits/AssetEditorManager.h"
+
 #include "IIntroTutorials.h"
 #include "IntroTutorials.h"
 #include "IDocumentation.h"
@@ -23,6 +23,7 @@
 #include "AnalyticsEventAttribute.h"
 #include "Interfaces/IAnalyticsProvider.h"
 #include "TutorialHyperlinkDecorator.h"
+#include "Subsystems/AssetEditorSubsystem.h"
 
 #define LOCTEXT_NAMESPACE "TutorialText"
 
@@ -215,7 +216,7 @@ static void ParseAssetLink(const FString& InternalLink, const FString* Action)
 		}
 		else
 		{
-			FAssetEditorManager::Get().OpenEditorForAsset(RequiredObject);
+			GEditor->GetEditorSubsystem<UAssetEditorSubsystem>()->OpenEditorForAsset(RequiredObject);
 		}
 
 		if( FEngineAnalytics::IsAvailable() )
diff --git a/Engine/Source/Editor/Kismet/Private/BlueprintDetailsCustomization.cpp b/Engine/Source/Editor/Kismet/Private/BlueprintDetailsCustomization.cpp
--- a/Engine/Source/Editor/Kismet/Private/BlueprintDetailsCustomization.cpp
+++ b/Engine/Source/Editor/Kismet/Private/BlueprintDetailsCustomization.cpp
@@ -75,6 +75,7 @@
 #include "Widgets/Input/STextComboBox.h"
 
 #include "UObject/TextProperty.h"
+#include "Subsystems/AssetEditorSubsystem.h"
 
 #define LOCTEXT_NAMESPACE "BlueprintDetailsCustomization"
 
@@ -5209,7 +5210,7 @@ void FBlueprintInterfaceLayout::OnBrowseToInterface(TWeakObjectPtr<UObject> Asse
 {
 	if (Asset.IsValid())
 	{
-		FAssetEditorManager::Get().OpenEditorForAsset(Asset.Get());
+		GEditor->GetEditorSubsystem<UAssetEditorSubsystem>()->OpenEditorForAsset(Asset.Get());
 	}
 }
 
diff --git a/Engine/Source/Editor/Kismet/Private/BlueprintEditor.cpp b/Engine/Source/Editor/Kismet/Private/BlueprintEditor.cpp
--- a/Engine/Source/Editor/Kismet/Private/BlueprintEditor.cpp
+++ b/Engine/Source/Editor/Kismet/Private/BlueprintEditor.cpp
@@ -159,6 +159,7 @@
 #include "Preferences/BlueprintEditorOptions.h"
 #include "Framework/MultiBox/MultiBoxBuilder.h"
 #include "Widgets/Input/SNumericEntryBox.h"
+#include "Subsystems/AssetEditorSubsystem.h"
 
 #define LOCTEXT_NAMESPACE "BlueprintEditor"
 
@@ -2186,7 +2187,7 @@ FReply FBlueprintEditor::OnEditParentClassClicked()
 			UBlueprintGeneratedClass* ParentBlueprintGeneratedClass = Cast<UBlueprintGeneratedClass>( ParentClass );
 			if ( ParentBlueprintGeneratedClass != NULL )
 			{
-				FAssetEditorManager::Get().OpenEditorForAsset( ParentBlueprintGeneratedClass->ClassGeneratedBy );
+				GEditor->GetEditorSubsystem<UAssetEditorSubsystem>()->OpenEditorForAsset( ParentBlueprintGeneratedClass->ClassGeneratedBy );
 			}
 		}
 	}
@@ -3642,7 +3643,7 @@ void FBlueprintEditor::JumpToHyperlink(const UObject* ObjectReference, bool bReq
 	}
 	else if(const UBlueprintGeneratedClass* Class = Cast<const UBlueprintGeneratedClass>(ObjectReference))
 	{
-		FAssetEditorManager::Get().OpenEditorForAsset(Class->ClassGeneratedBy);
+		GEditor->GetEditorSubsystem<UAssetEditorSubsystem>()->OpenEditorForAsset(Class->ClassGeneratedBy);
 	}
 	else if (const UTimelineTemplate* Timeline = Cast<const UTimelineTemplate>(ObjectReference))
 	{
@@ -3650,7 +3651,7 @@ void FBlueprintEditor::JumpToHyperlink(const UObject* ObjectReference, bool bReq
 	}
 	else if ((ObjectReference != nullptr) && ObjectReference->IsAsset())
 	{
-		FAssetEditorManager::Get().OpenEditorForAsset(const_cast<UObject*>(ObjectReference));
+		GEditor->GetEditorSubsystem<UAssetEditorSubsystem>()->OpenEditorForAsset(const_cast<UObject*>(ObjectReference));
 	}
 	else
 	{
@@ -7357,7 +7358,7 @@ void FBlueprintEditor::Tick(float DeltaTime)
 
 	if (bPendingDeferredClose)
 	{
-		IAssetEditorInstance* EditorInst = FAssetEditorManager::Get().FindEditorForAsset(GetBlueprintObj(), /*bFocusIfOpen =*/false);
+		IAssetEditorInstance* EditorInst = GEditor->GetEditorSubsystem<UAssetEditorSubsystem>()->FindEditorForAsset(GetBlueprintObj(), /*bFocusIfOpen =*/false);
 		check(EditorInst != nullptr);
 		EditorInst->CloseWindow();
 	}
diff --git a/Engine/Source/Editor/Kismet/Private/BlueprintEditorViewportContextMenuExtender.cpp b/Engine/Source/Editor/Kismet/Private/BlueprintEditorViewportContextMenuExtender.cpp
--- a/Engine/Source/Editor/Kismet/Private/BlueprintEditorViewportContextMenuExtender.cpp
+++ b/Engine/Source/Editor/Kismet/Private/BlueprintEditorViewportContextMenuExtender.cpp
@@ -12,9 +12,10 @@
 #include "Engine/Selection.h"
 #include "Framework/MultiBox/MultiBoxBuilder.h"
 #include "Editor.h"
-#include "Toolkits/AssetEditorManager.h"
+
 #include "Kismet2/KismetEditorUtilities.h"
 #include "CreateBlueprintFromActorDialog.h"
+#include "Subsystems/AssetEditorSubsystem.h"
 
 DEFINE_LOG_CATEGORY_STATIC(LogViewportBlueprintMenu, Log, All);
 
@@ -42,7 +43,7 @@ void EditKismetCodeFor( TWeakObjectPtr<UBlueprint> BlueprintRef )
 	{
 		// Open the blueprint
 		// @todo toolkit major: Needs world-centric support (pass in LevelEditor.  See FLevelEditorActionCallbacks::OpenLevelBlueprint)
-		FAssetEditorManager::Get().OpenEditorForAsset( Blueprint );
+		GEditor->GetEditorSubsystem<UAssetEditorSubsystem>()->OpenEditorForAsset( Blueprint );
 	}
 	else
 	{
diff --git a/Engine/Source/Editor/Kismet/Private/SBlueprintDiff.cpp b/Engine/Source/Editor/Kismet/Private/SBlueprintDiff.cpp
--- a/Engine/Source/Editor/Kismet/Private/SBlueprintDiff.cpp
+++ b/Engine/Source/Editor/Kismet/Private/SBlueprintDiff.cpp
@@ -24,7 +24,8 @@
 #include "WidgetBlueprint.h"
 #include "HAL/PlatformApplicationMisc.h"
 #include "Framework/Application/SlateApplication.h"
-#include "Toolkits/AssetEditorManager.h"
+
+#include "Subsystems/AssetEditorSubsystem.h"
 
 #define LOCTEXT_NAMESPACE "SBlueprintDif"
 
@@ -909,7 +910,7 @@ void SBlueprintDiff::Construct( const FArguments& InArgs)
 	{
 		WeakParentWindow = InArgs._ParentWindow;
 
-		AssetEditorCloseDelegate = FAssetEditorManager::Get().OnAssetEditorRequestClose().AddSP(this, &SBlueprintDiff::OnCloseAssetEditor);
+		AssetEditorCloseDelegate = GEditor->GetEditorSubsystem<UAssetEditorSubsystem>()->OnAssetEditorRequestClose().AddSP(this, &SBlueprintDiff::OnCloseAssetEditor);
 	}
 
 	FToolBarBuilder ToolbarBuilder(TSharedPtr< const FUICommandList >(), FMultiBoxCustomization::None);
@@ -1049,7 +1050,7 @@ SBlueprintDiff::~SBlueprintDiff()
 {
 	if (AssetEditorCloseDelegate.IsValid())
 	{
-		FAssetEditorManager::Get().OnAssetEditorRequestClose().Remove(AssetEditorCloseDelegate);
+		GEditor->GetEditorSubsystem<UAssetEditorSubsystem>()->OnAssetEditorRequestClose().Remove(AssetEditorCloseDelegate);
 	}
 }
 
@@ -1062,7 +1063,7 @@ void SBlueprintDiff::OnCloseAssetEditor(UObject* Asset, EAssetEditorCloseReason
 
 		if (AssetEditorCloseDelegate.IsValid())
 		{
-			FAssetEditorManager::Get().OnAssetEditorRequestClose().Remove(AssetEditorCloseDelegate);
+			GEditor->GetEditorSubsystem<UAssetEditorSubsystem>()->OnAssetEditorRequestClose().Remove(AssetEditorCloseDelegate);
 		}
 
 		if (WeakParentWindow.IsValid())
diff --git a/Engine/Source/Editor/Kismet/Private/SSCSEditor.cpp b/Engine/Source/Editor/Kismet/Private/SSCSEditor.cpp
--- a/Engine/Source/Editor/Kismet/Private/SSCSEditor.cpp
+++ b/Engine/Source/Editor/Kismet/Private/SSCSEditor.cpp
@@ -74,6 +74,7 @@
 #include "ToolMenus.h"
 #include "SSCSEditorMenuContext.h"
 #include "Kismet2/ComponentEditorContextMenuContex.h"
+#include "Subsystems/AssetEditorSubsystem.h"
 
 #define LOCTEXT_NAMESPACE "SSCSEditor"
 
@@ -5096,7 +5097,7 @@ UClass* SSCSEditor::CreateNewBPComponent(TSubclassOf<UActorComponent> ComponentC
 					GEditor->SyncBrowserToObjects(Objects);
 
 					// Open the editor for the new blueprint
-					FAssetEditorManager::Get().OpenEditorForAsset(NewBP);
+					GEditor->GetEditorSubsystem<UAssetEditorSubsystem>()->OpenEditorForAsset(NewBP);
 				}
 			}
 		}
@@ -6413,7 +6414,7 @@ void SSCSEditor::OnOpenBlueprintEditor(bool bForceCodeEditing) const
 			}
 			else
 			{
-				FAssetEditorManager::Get().OpenEditorForAsset(Blueprint);
+				GEditor->GetEditorSubsystem<UAssetEditorSubsystem>()->OpenEditorForAsset(Blueprint);
 			}
 		}
 	}
diff --git a/Engine/Source/Editor/Kismet/Private/Tests/BlueprintEditorTests.cpp b/Engine/Source/Editor/Kismet/Private/Tests/BlueprintEditorTests.cpp
--- a/Engine/Source/Editor/Kismet/Private/Tests/BlueprintEditorTests.cpp
+++ b/Engine/Source/Editor/Kismet/Private/Tests/BlueprintEditorTests.cpp
@@ -23,7 +23,7 @@
 #include "Engine/BlueprintGeneratedClass.h"
 #include "AssetData.h"
 #include "Editor.h"
-#include "Toolkits/AssetEditorManager.h"
+
 #include "Kismet2/KismetEditorUtilities.h"
 #include "EdGraphSchema_K2.h"
 #include "EdGraphSchema_K2_Actions.h"
@@ -48,6 +48,7 @@
 #include "Tests/AutomationTestSettings.h"
 #include "Tests/AutomationEditorCommon.h"
 #include "Tests/AutomationEditorPromotionCommon.h"
+#include "Subsystems/AssetEditorSubsystem.h"
 
 #if WITH_DEV_AUTOMATION_TESTS
 
@@ -195,7 +196,7 @@ namespace BlueprintEditorPromotionUtils
 	*/
 	static USCS_Node* CreateBlueprintComponent(UBlueprint* InBlueprint, UObject* InAsset)
 	{
-		IAssetEditorInstance* OpenEditor = FAssetEditorManager::Get().FindEditorForAsset(InBlueprint, true);
+		IAssetEditorInstance* OpenEditor = GEditor->GetEditorSubsystem<UAssetEditorSubsystem>()->FindEditorForAsset(InBlueprint, true);
 		FBlueprintEditor* CurrentBlueprintEditor = (FBlueprintEditor*)OpenEditor;
 		TSubclassOf<UActorComponent> ComponentClass = FComponentAssetBrokerage::GetPrimaryComponentForAsset(InAsset->GetClass());
 
@@ -411,7 +412,7 @@ namespace BlueprintEditorPromotionUtils
 	*/
 	static void PromotePinToVariable(UBlueprint* InBlueprint, UEdGraphNode* Node, const FName& PinName)
 	{
-		IAssetEditorInstance* OpenEditor = FAssetEditorManager::Get().FindEditorForAsset(InBlueprint, true);
+		IAssetEditorInstance* OpenEditor = GEditor->GetEditorSubsystem<UAssetEditorSubsystem>()->FindEditorForAsset(InBlueprint, true);
 		FBlueprintEditor* CurrentBlueprintEditor = (FBlueprintEditor*)OpenEditor;
 
 		UEdGraphPin* PinToPromote = Node->FindPin(PinName);
@@ -751,7 +752,7 @@ namespace BlueprintEditorPromotionTestHelper
 		bool Cleanup()
 		{
 			//Make sure no editors are open
-			FAssetEditorManager::Get().CloseAllAssetEditors();
+			GEditor->GetEditorSubsystem<UAssetEditorSubsystem>()->CloseAllAssetEditors();
 
 			//remove all assets in the build package
 			// Load the asset registry module
@@ -881,7 +882,7 @@ namespace BlueprintEditorPromotionTestHelper
 				BlueprintPackage->MarkPackageDirty();
 
 				Test->AddInfo(TEXT("Opening the blueprint editor for the first time"));
-				FAssetEditorManager::Get().OpenEditorForAsset(BlueprintObject);
+				GEditor->GetEditorSubsystem<UAssetEditorSubsystem>()->OpenEditorForAsset(BlueprintObject);
 			}
 
 			return true;
@@ -895,7 +896,7 @@ namespace BlueprintEditorPromotionTestHelper
 		{
 			if (BlueprintObject)
 			{
-				IAssetEditorInstance* AssetEditor = FAssetEditorManager::Get().FindEditorForAsset(BlueprintObject, true);
+				IAssetEditorInstance* AssetEditor = GEditor->GetEditorSubsystem<UAssetEditorSubsystem>()->FindEditorForAsset(BlueprintObject, true);
 				IBlueprintEditor* BlueprintEditor = (IBlueprintEditor*)AssetEditor;
 				Test->TestTrue(TEXT("Opened correct initial editor"), BlueprintEditor->GetCurrentMode() != FBlueprintEditorApplicationModes::BlueprintDefaultsMode);
 			}
@@ -912,7 +913,7 @@ namespace BlueprintEditorPromotionTestHelper
 			if (BlueprintObject)
 			{
 				Test->AddInfo(TEXT("Closing the blueprint editor"));
-				FAssetEditorManager::Get().CloseAllAssetEditors();
+				GEditor->GetEditorSubsystem<UAssetEditorSubsystem>()->CloseAllAssetEditors();
 			}
 			return true;
 		}
@@ -926,7 +927,7 @@ namespace BlueprintEditorPromotionTestHelper
 			if (BlueprintObject)
 			{
 				Test->AddInfo(TEXT("Opening the blueprint editor for the second time"));
-				FAssetEditorManager::Get().OpenEditorForAsset(BlueprintObject);
+				GEditor->GetEditorSubsystem<UAssetEditorSubsystem>()->OpenEditorForAsset(BlueprintObject);
 			}
 			return true;
 		}
@@ -939,7 +940,7 @@ namespace BlueprintEditorPromotionTestHelper
 		{
 			if (BlueprintObject)
 			{
-				IAssetEditorInstance* AssetEditor = FAssetEditorManager::Get().FindEditorForAsset(BlueprintObject, true);
+				IAssetEditorInstance* AssetEditor = GEditor->GetEditorSubsystem<UAssetEditorSubsystem>()->FindEditorForAsset(BlueprintObject, true);
 				IBlueprintEditor* BlueprintEditor = (IBlueprintEditor*)AssetEditor;
 				
 				bool bCorrectEditorOpened = BlueprintEditor->GetCurrentMode() == FBlueprintEditorApplicationModes::BlueprintDefaultsMode;
@@ -962,7 +963,7 @@ namespace BlueprintEditorPromotionTestHelper
 		{
 			if (BlueprintObject)
 			{
-				IAssetEditorInstance* AssetEditor = FAssetEditorManager::Get().FindEditorForAsset(BlueprintObject, true);
+				IAssetEditorInstance* AssetEditor = GEditor->GetEditorSubsystem<UAssetEditorSubsystem>()->FindEditorForAsset(BlueprintObject, true);
 				IBlueprintEditor* BlueprintEditor = (IBlueprintEditor*)AssetEditor;
 
 				MeshNode = BlueprintEditorPromotionUtils::CreateBlueprintComponent(BlueprintObject, FirstBlueprintMesh);
@@ -1007,7 +1008,7 @@ namespace BlueprintEditorPromotionTestHelper
 		{
 			if (BlueprintObject)
 			{
-				IAssetEditorInstance* AssetEditor = FAssetEditorManager::Get().FindEditorForAsset(BlueprintObject, true);
+				IAssetEditorInstance* AssetEditor = GEditor->GetEditorSubsystem<UAssetEditorSubsystem>()->FindEditorForAsset(BlueprintObject, true);
 				IBlueprintEditor* BlueprintEditor = (IBlueprintEditor*)AssetEditor;
 
 				//FEditorPromotionTestUtilities::TakeScreenshot(TEXT("BlueprintComponentVariables"), true);
@@ -1042,7 +1043,7 @@ namespace BlueprintEditorPromotionTestHelper
 		{
 			if (BlueprintObject)
 			{
-				IAssetEditorInstance* AssetEditor = FAssetEditorManager::Get().FindEditorForAsset(BlueprintObject, true);
+				IAssetEditorInstance* AssetEditor = GEditor->GetEditorSubsystem<UAssetEditorSubsystem>()->FindEditorForAsset(BlueprintObject, true);
 				FBlueprintEditor* BlueprintEditor = (FBlueprintEditor*)AssetEditor;
 
 				UEdGraph* ConstructionGraph = FBlueprintEditorUtils::FindUserConstructionScript(BlueprintObject);
@@ -1105,7 +1106,7 @@ namespace BlueprintEditorPromotionTestHelper
 		{
 			if (BlueprintObject)
 			{
-				IAssetEditorInstance* AssetEditor = FAssetEditorManager::Get().FindEditorForAsset(BlueprintObject, true);
+				IAssetEditorInstance* AssetEditor = GEditor->GetEditorSubsystem<UAssetEditorSubsystem>()->FindEditorForAsset(BlueprintObject, true);
 				FBlueprintEditor* BlueprintEditor = (FBlueprintEditor*)AssetEditor;
 
 				BlueprintEditorPromotionUtils::PromotePinToVariable(BlueprintObject, AddMeshNode, UEdGraphSchema_K2::PN_ReturnValue);
@@ -1129,7 +1130,7 @@ namespace BlueprintEditorPromotionTestHelper
 		{
 			if (BlueprintObject)
 			{
-				IAssetEditorInstance* AssetEditor = FAssetEditorManager::Get().FindEditorForAsset(BlueprintObject, true);
+				IAssetEditorInstance* AssetEditor = GEditor->GetEditorSubsystem<UAssetEditorSubsystem>()->FindEditorForAsset(BlueprintObject, true);
 				FBlueprintEditor* BlueprintEditor = (FBlueprintEditor*)AssetEditor;
 
 				//FEditorPromotionTestUtilities::TakeScreenshot(TEXT("BlueprintMeshVariable"), true);
@@ -1163,7 +1164,7 @@ namespace BlueprintEditorPromotionTestHelper
 		{
 			if (BlueprintObject)
 			{
-				IAssetEditorInstance* AssetEditor = FAssetEditorManager::Get().FindEditorForAsset(BlueprintObject, true);
+				IAssetEditorInstance* AssetEditor = GEditor->GetEditorSubsystem<UAssetEditorSubsystem>()->FindEditorForAsset(BlueprintObject, true);
 				FBlueprintEditor* BlueprintEditor = (FBlueprintEditor*)AssetEditor;
 
 				UEdGraph* EventGraph = FBlueprintEditorUtils::FindEventGraph(BlueprintObject);
diff --git a/Engine/Source/Editor/Kismet/Public/BlueprintEditorModule.h b/Engine/Source/Editor/Kismet/Public/BlueprintEditorModule.h
--- a/Engine/Source/Editor/Kismet/Public/BlueprintEditorModule.h
+++ b/Engine/Source/Editor/Kismet/Public/BlueprintEditorModule.h
@@ -117,7 +117,7 @@ public:
 	 *
 	 * Note: This function should not be called directly, use one of the following instead:
 	 *	- FKismetEditorUtilities::BringKismetToFocusAttentionOnObject
-	 *  - FAssetEditorManager::Get().OpenEditorForAsset
+	 *  - GEditor->GetEditorSubsystem<UAssetEditorSubsystem>()->OpenEditorForAsset
 	 *
 	 * @param	Mode					Mode that this editor should operate in
 	 * @param	InitToolkitHost			When Mode is WorldCentric, this is the level editor instance to spawn this editor within
diff --git a/Engine/Source/Editor/LevelEditor/Private/LevelEditorActions.cpp b/Engine/Source/Editor/LevelEditor/Private/LevelEditorActions.cpp
--- a/Engine/Source/Editor/LevelEditor/Private/LevelEditorActions.cpp
+++ b/Engine/Source/Editor/LevelEditor/Private/LevelEditorActions.cpp
@@ -36,7 +36,7 @@
 #include "FileHelpers.h"
 #include "EditorModeInterpolation.h"
 #include "UnrealEdGlobals.h"
-#include "Toolkits/AssetEditorManager.h"
+
 #include "LevelEditor.h"
 #include "Matinee/MatineeActor.h"
 #include "Engine/LevelScriptBlueprint.h"
@@ -95,6 +95,7 @@
 #if WITH_LIVE_CODING
 #include "ILiveCodingModule.h"
 #endif
+#include "Subsystems/AssetEditorSubsystem.h"
 
 DEFINE_LOG_CATEGORY_STATIC(LevelEditorActions, Log, All);
 
@@ -130,7 +131,7 @@ namespace LevelEditorActionsHelpers
 		{
 			// @todo Re-enable once world centric works
 			const bool bOpenWorldCentric = false;
-			FAssetEditorManager::Get().OpenEditorForAsset(
+			GEditor->GetEditorSubsystem<UAssetEditorSubsystem>()->OpenEditorForAsset(
 				Blueprint,
 				bOpenWorldCentric ? EToolkitMode::WorldCentric : EToolkitMode::Standalone,
 				InLevelEditor.Pin()  );
@@ -1241,7 +1242,7 @@ void FLevelEditorActionCallbacks::EditAsset_Clicked( const EToolkitMode::Type To
 			{
 				for (auto Asset : ReferencedAssets)
 				{
-					FAssetEditorManager::Get().OpenEditorForAsset(Asset, ToolkitMode, LevelEditorSharedPtr);
+					GEditor->GetEditorSubsystem<UAssetEditorSubsystem>()->OpenEditorForAsset(Asset, ToolkitMode, LevelEditorSharedPtr);
 				}
 			}
 		}
@@ -2327,7 +2328,7 @@ void FLevelEditorActionCallbacks::OpenLevelBlueprint( TWeakPtr< SLevelEditor > L
 		{
 			// @todo Re-enable once world centric works
 			const bool bOpenWorldCentric = false;
-			FAssetEditorManager::Get().OpenEditorForAsset(
+			GEditor->GetEditorSubsystem<UAssetEditorSubsystem>()->OpenEditorForAsset(
 				LevelScriptBlueprint,
 				bOpenWorldCentric ? EToolkitMode::WorldCentric : EToolkitMode::Standalone,
 				LevelEditor.Pin()  );
@@ -2355,7 +2356,7 @@ void FLevelEditorActionCallbacks::CreateBlankBlueprintClass()
 		{
 			// @todo Re-enable once world centric works
 			const bool bOpenWorldCentric = false;
-			FAssetEditorManager::Get().OpenEditorForAsset(
+			GEditor->GetEditorSubsystem<UAssetEditorSubsystem>()->OpenEditorForAsset(
 				Blueprint,
 				bOpenWorldCentric ? EToolkitMode::WorldCentric : EToolkitMode::Standalone);
 		}
diff --git a/Engine/Source/Editor/LevelEditor/Private/LevelEditorToolBar.cpp b/Engine/Source/Editor/LevelEditor/Private/LevelEditorToolBar.cpp
--- a/Engine/Source/Editor/LevelEditor/Private/LevelEditorToolBar.cpp
+++ b/Engine/Source/Editor/LevelEditor/Private/LevelEditorToolBar.cpp
@@ -26,7 +26,7 @@
 #include "GameFramework/HUD.h"
 #include "GameFramework/GameStateBase.h"
 #include "Engine/TextureStreamingTypes.h"
-#include "Toolkits/AssetEditorManager.h"
+
 #include "LevelEditor.h"
 #include "LevelEditorActions.h"
 #include "SourceCodeNavigation.h"
@@ -60,6 +60,7 @@
 #include "RHIShaderPlatformDefinitions.inl"
 #include "LevelEditorMenuContext.h"
 #include "ToolMenus.h"
+#include "Subsystems/AssetEditorSubsystem.h"
 
 namespace LevelEditorActionHelpers
 {
@@ -198,7 +199,7 @@ namespace LevelEditorActionHelpers
 			{
 				// @todo Re-enable once world centric works
 				const bool bOpenWorldCentric = false;
-				FAssetEditorManager::Get().OpenEditorForAsset(
+				GEditor->GetEditorSubsystem<UAssetEditorSubsystem>()->OpenEditorForAsset(
 					BlueprintClass,
 					bOpenWorldCentric ? EToolkitMode::WorldCentric : EToolkitMode::Standalone,
 					InLevelEditor.Pin() );
@@ -219,7 +220,7 @@ namespace LevelEditorActionHelpers
 			{
 				// @todo Re-enable once world centric works
 				const bool bOpenWorldCentric = false;
-				FAssetEditorManager::Get().OpenEditorForAsset(
+				GEditor->GetEditorSubsystem<UAssetEditorSubsystem>()->OpenEditorForAsset(
 					BlueprintClass,
 					bOpenWorldCentric ? EToolkitMode::WorldCentric : EToolkitMode::Standalone,
 					InLevelEditor.Pin() );
@@ -240,7 +241,7 @@ namespace LevelEditorActionHelpers
 			{
 				// @todo Re-enable once world centric works
 				const bool bOpenWorldCentric = false;
-				FAssetEditorManager::Get().OpenEditorForAsset(
+				GEditor->GetEditorSubsystem<UAssetEditorSubsystem>()->OpenEditorForAsset(
 					BlueprintClass,
 					bOpenWorldCentric ? EToolkitMode::WorldCentric : EToolkitMode::Standalone,
 					InLevelEditor.Pin() );
@@ -261,7 +262,7 @@ namespace LevelEditorActionHelpers
 			{
 				// @todo Re-enable once world centric works
 				const bool bOpenWorldCentric = false;
-				FAssetEditorManager::Get().OpenEditorForAsset(
+				GEditor->GetEditorSubsystem<UAssetEditorSubsystem>()->OpenEditorForAsset(
 					BlueprintClass,
 					bOpenWorldCentric ? EToolkitMode::WorldCentric : EToolkitMode::Standalone,
 					InLevelEditor.Pin() );
@@ -282,7 +283,7 @@ namespace LevelEditorActionHelpers
 			{
 				// @todo Re-enable once world centric works
 				const bool bOpenWorldCentric = false;
-				FAssetEditorManager::Get().OpenEditorForAsset(
+				GEditor->GetEditorSubsystem<UAssetEditorSubsystem>()->OpenEditorForAsset(
 					BlueprintClass,
 					bOpenWorldCentric ? EToolkitMode::WorldCentric : EToolkitMode::Standalone,
 					InLevelEditor.Pin() );
@@ -810,7 +811,7 @@ void LevelEditorActionHelpers::OnCreateGameModeClassPicked(UClass* InChosenClass
 		{
 			// @todo Re-enable once world centric works
 			const bool bOpenWorldCentric = false;
-			FAssetEditorManager::Get().OpenEditorForAsset(
+			GEditor->GetEditorSubsystem<UAssetEditorSubsystem>()->OpenEditorForAsset(
 				Blueprint,
 				bOpenWorldCentric ? EToolkitMode::WorldCentric : EToolkitMode::Standalone,
 				InLevelEditor.Pin()  );
@@ -897,7 +898,7 @@ void LevelEditorActionHelpers::OnCreateGameStateClassPicked(UClass* InChosenClas
 		{
 			// @todo Re-enable once world centric works
 			const bool bOpenWorldCentric = false;
-			FAssetEditorManager::Get().OpenEditorForAsset(
+			GEditor->GetEditorSubsystem<UAssetEditorSubsystem>()->OpenEditorForAsset(
 				Blueprint,
 				bOpenWorldCentric ? EToolkitMode::WorldCentric : EToolkitMode::Standalone,
 				InLevelEditor.Pin()  );
@@ -971,7 +972,7 @@ void LevelEditorActionHelpers::OnCreatePawnClassPicked(UClass* InChosenClass, TW
 		{
 			// @todo Re-enable once world centric works
 			const bool bOpenWorldCentric = false;
-			FAssetEditorManager::Get().OpenEditorForAsset(
+			GEditor->GetEditorSubsystem<UAssetEditorSubsystem>()->OpenEditorForAsset(
 				Blueprint,
 				bOpenWorldCentric ? EToolkitMode::WorldCentric : EToolkitMode::Standalone,
 				InLevelEditor.Pin()  );
@@ -1045,7 +1046,7 @@ void LevelEditorActionHelpers::OnCreateHUDClassPicked(UClass* InChosenClass, TWe
 		{
 			// @todo Re-enable once world centric works
 			const bool bOpenWorldCentric = false;
-			FAssetEditorManager::Get().OpenEditorForAsset(
+			GEditor->GetEditorSubsystem<UAssetEditorSubsystem>()->OpenEditorForAsset(
 				Blueprint,
 				bOpenWorldCentric ? EToolkitMode::WorldCentric : EToolkitMode::Standalone,
 				InLevelEditor.Pin()  );
@@ -1119,7 +1120,7 @@ void LevelEditorActionHelpers::OnCreatePlayerControllerClassPicked(UClass* InCho
 		{
 			// @todo Re-enable once world centric works
 			const bool bOpenWorldCentric = false;
-			FAssetEditorManager::Get().OpenEditorForAsset(
+			GEditor->GetEditorSubsystem<UAssetEditorSubsystem>()->OpenEditorForAsset(
 				Blueprint,
 				bOpenWorldCentric ? EToolkitMode::WorldCentric : EToolkitMode::Standalone,
 				InLevelEditor.Pin()  );
@@ -2143,7 +2144,7 @@ TSharedRef< SWidget > FLevelEditorToolBar::GenerateOpenBlueprintMenuContent( TSh
 			UBlueprint* SelectedBP = Cast<UBlueprint>(AssetData.GetAsset());
 			if(SelectedBP)
 			{
-				FAssetEditorManager::Get().OpenEditorForAsset(SelectedBP);
+				GEditor->GetEditorSubsystem<UAssetEditorSubsystem>()->OpenEditorForAsset(SelectedBP);
 			}
 		}
 
@@ -2254,7 +2255,7 @@ void FLevelEditorToolBar::OnOpenSubLevelBlueprint( ULevel* InLevel )
 
 	if( LevelScriptBlueprint )
 	{
-		FAssetEditorManager::Get().OpenEditorForAsset(LevelScriptBlueprint);
+		GEditor->GetEditorSubsystem<UAssetEditorSubsystem>()->OpenEditorForAsset(LevelScriptBlueprint);
 	}
 	else
 	{
@@ -2364,7 +2365,7 @@ void FLevelEditorToolBar::OnCinematicsActorPicked( AActor* Actor )
 
 		if (Asset != nullptr)
 		{
-			FAssetEditorManager::Get().OpenEditorForAsset(Asset);
+			GEditor->GetEditorSubsystem<UAssetEditorSubsystem>()->OpenEditorForAsset(Asset);
 		}
 	}
 }
diff --git a/Engine/Source/Editor/MainFrame/Private/Frame/MainFrameHandler.cpp b/Engine/Source/Editor/MainFrame/Private/Frame/MainFrameHandler.cpp
--- a/Engine/Source/Editor/MainFrame/Private/Frame/MainFrameHandler.cpp
+++ b/Engine/Source/Editor/MainFrame/Private/Frame/MainFrameHandler.cpp
@@ -65,10 +65,6 @@ void FMainFrameHandler::ShutDownEditor()
 		MainTabPtr.Pin()->SetOnTabClosed(SDockTab::FOnTabClosedCallback());
 	}
 
-	// Inform the AssetEditorManager that the editor is exiting so that it may save open assets
-	// and report usage stats
-	FAssetEditorManager::Get().OnExit();
-
 	if (RootWindow.IsValid())
 	{
 		RootWindow->SetRequestDestroyWindowOverride(FRequestDestroyWindowOverride());
diff --git a/Engine/Source/Editor/MainFrame/Private/Frame/MainFrameHandler.h b/Engine/Source/Editor/MainFrame/Private/Frame/MainFrameHandler.h
--- a/Engine/Source/Editor/MainFrame/Private/Frame/MainFrameHandler.h
+++ b/Engine/Source/Editor/MainFrame/Private/Frame/MainFrameHandler.h
@@ -11,7 +11,7 @@
 #include "Framework/Docking/TabManager.h"
 #include "Framework/Docking/LayoutService.h"
 #include "EngineGlobals.h"
-#include "Toolkits/AssetEditorManager.h"
+
 #include "Editor/UnrealEdEngine.h"
 #include "EditorModeManager.h"
 #include "EditorModes.h"
@@ -20,6 +20,7 @@
 #include "LevelEditor.h"
 #include "ILevelViewport.h"
 #include "MainFrameLog.h"
+#include "Subsystems/AssetEditorSubsystem.h"
 
 const FText StaticGetApplicationTitle( const bool bIncludeGameName );
 
@@ -255,7 +256,7 @@ public:
 				LevelEditor.FocusViewport();
 
 				// Restore any assets we had open. Note we don't do this on immersive PIE as its annoying to the user.
-				FAssetEditorManager::Get().RequestRestorePreviouslyOpenAssets();
+				GEditor->GetEditorSubsystem<UAssetEditorSubsystem>()->RequestRestorePreviouslyOpenAssets();
 			}
 		}
 	}
diff --git a/Engine/Source/Editor/MaterialEditor/Private/MaterialEditingLibrary.cpp b/Engine/Source/Editor/MaterialEditor/Private/MaterialEditingLibrary.cpp
--- a/Engine/Source/Editor/MaterialEditor/Private/MaterialEditingLibrary.cpp
+++ b/Engine/Source/Editor/MaterialEditor/Private/MaterialEditingLibrary.cpp
@@ -28,6 +28,7 @@
 #include "SceneTypes.h"
 #include "AssetRegistryModule.h"
 #include "DebugViewModeHelpers.h"
+#include "Subsystems/AssetEditorSubsystem.h"
 
 DEFINE_LOG_CATEGORY_STATIC(LogMaterialEditingLibrary, Warning, All);
 
@@ -258,8 +259,8 @@ namespace MaterialEditingLibraryImpl
 
 void UMaterialEditingLibrary::RebuildMaterialInstanceEditors(UMaterial* BaseMaterial)
 {
-	FAssetEditorManager& AssetEditorManager = FAssetEditorManager::Get();
-	TArray<UObject*> EditedAssets = AssetEditorManager.GetAllEditedAssets();
+	UAssetEditorSubsystem* AssetEditorSubsystem = GEditor->GetEditorSubsystem<UAssetEditorSubsystem>();
+	TArray<UObject*> EditedAssets = AssetEditorSubsystem->GetAllEditedAssets();
 
 	for (int32 AssetIdx = 0; AssetIdx < EditedAssets.Num(); AssetIdx++)
 	{
@@ -283,7 +284,7 @@ void UMaterialEditingLibrary::RebuildMaterialInstanceEditors(UMaterial* BaseMate
 			UMaterial* MICOriginalMaterial = SourceInstance->GetMaterial();
 			if (MICOriginalMaterial == BaseMaterial)
 			{
-				IAssetEditorInstance* EditorInstance = AssetEditorManager.FindEditorForAsset(EditedAsset, false);
+				IAssetEditorInstance* EditorInstance = AssetEditorSubsystem->FindEditorForAsset(EditedAsset, false);
 				if (EditorInstance != nullptr)
 				{
 					FMaterialInstanceEditor* OtherEditor = static_cast<FMaterialInstanceEditor*>(EditorInstance);
@@ -296,8 +297,8 @@ void UMaterialEditingLibrary::RebuildMaterialInstanceEditors(UMaterial* BaseMate
 
 void UMaterialEditingLibrary::RebuildMaterialInstanceEditors(UMaterialFunction* BaseFunction)
 {
-	FAssetEditorManager& AssetEditorManager = FAssetEditorManager::Get();
-	TArray<UObject*> EditedAssets = AssetEditorManager.GetAllEditedAssets();
+	UAssetEditorSubsystem* AssetEditorSubsystem = GEditor->GetEditorSubsystem<UAssetEditorSubsystem>();
+	TArray<UObject*> EditedAssets = AssetEditorSubsystem->GetAllEditedAssets();
 
 	for (int32 AssetIdx = 0; AssetIdx < EditedAssets.Num(); AssetIdx++)
 	{
@@ -311,7 +312,7 @@ void UMaterialEditingLibrary::RebuildMaterialInstanceEditors(UMaterialFunction*
 			// Update function instances that are children of this material function	
 			if (BaseFunction && BaseFunction == FunctionInstance->GetBaseFunction())
 			{
-				IAssetEditorInstance* EditorInstance = AssetEditorManager.FindEditorForAsset(EditedAsset, false);
+				IAssetEditorInstance* EditorInstance = AssetEditorSubsystem->FindEditorForAsset(EditedAsset, false);
 				if (EditorInstance)
 				{
 					FMaterialInstanceEditor* OtherEditor = static_cast<FMaterialInstanceEditor*>(EditorInstance);
@@ -339,7 +340,7 @@ void UMaterialEditingLibrary::RebuildMaterialInstanceEditors(UMaterialFunction*
 
 				if (BaseFunction && (DependentFunctions.Contains(BaseFunction) || DependentFunctions.Contains(BaseFunction->ParentFunction)))
 				{
-					IAssetEditorInstance* EditorInstance = AssetEditorManager.FindEditorForAsset(EditedAsset, false);
+					IAssetEditorInstance* EditorInstance = AssetEditorSubsystem->FindEditorForAsset(EditedAsset, false);
 					if (EditorInstance != nullptr)
 					{
 						FMaterialInstanceEditor* OtherEditor = static_cast<FMaterialInstanceEditor*>(EditorInstance);
@@ -809,7 +810,7 @@ void UMaterialEditingLibrary::UpdateMaterialFunction(UMaterialFunctionInterface*
 						}
 
 						// if this instance was opened in an editor notify the change
-						auto* MaterialEditor = (IMaterialEditor*)FAssetEditorManager::Get().FindEditorForAsset(CurrentMaterial, false);
+						auto* MaterialEditor = (IMaterialEditor*)GEditor->GetEditorSubsystem<UAssetEditorSubsystem>()->FindEditorForAsset(CurrentMaterial, false);
 						if (MaterialEditor)
 						{
 							MaterialEditor->NotifyExternalMaterialChange();
@@ -833,7 +834,7 @@ void UMaterialEditingLibrary::UpdateMaterialFunction(UMaterialFunctionInterface*
 						CurrentInstance->PostEditChange();
 
 						// if this instance was opened in an editor notify the change
-						auto* MaterialEditor = (IMaterialEditor*)FAssetEditorManager::Get().FindEditorForAsset(CurrentInstance, false);
+						auto* MaterialEditor = (IMaterialEditor*)GEditor->GetEditorSubsystem<UAssetEditorSubsystem>()->FindEditorForAsset(CurrentInstance, false);
 						if (MaterialEditor)
 						{
 							MaterialEditor->NotifyExternalMaterialChange();
diff --git a/Engine/Source/Editor/MaterialEditor/Private/MaterialEditor.cpp b/Engine/Source/Editor/MaterialEditor/Private/MaterialEditor.cpp
--- a/Engine/Source/Editor/MaterialEditor/Private/MaterialEditor.cpp
+++ b/Engine/Source/Editor/MaterialEditor/Private/MaterialEditor.cpp
@@ -129,6 +129,7 @@
 #include "IPropertyRowGenerator.h"
 #include "Widgets/Layout/SScrollBox.h"
 #include "UObject/TextProperty.h"
+#include "Subsystems/AssetEditorSubsystem.h"
 
 #define LOCTEXT_NAMESPACE "MaterialEditor"
 
@@ -5200,7 +5201,7 @@ void FMaterialEditor::OnNodeDoubleClicked(class UEdGraphNode* Node)
 		}
 		if (ObjectToEdit)
 		{
-			FAssetEditorManager::Get().OpenEditorForAsset(ObjectToEdit);
+			GEditor->GetEditorSubsystem<UAssetEditorSubsystem>()->OpenEditorForAsset(ObjectToEdit);
 		}
 	}
 }
diff --git a/Engine/Source/Editor/MaterialEditor/Private/MaterialEditorUtilities.cpp b/Engine/Source/Editor/MaterialEditor/Private/MaterialEditorUtilities.cpp
--- a/Engine/Source/Editor/MaterialEditor/Private/MaterialEditorUtilities.cpp
+++ b/Engine/Source/Editor/MaterialEditor/Private/MaterialEditorUtilities.cpp
@@ -33,6 +33,7 @@
 #include "Misc/ScopedSlowTask.h"
 #include "Templates/UniquePtr.h"
 #include "Materials/MaterialFunctionInstance.h"
+#include "Subsystems/AssetEditorSubsystem.h"
 
 #define LOCTEXT_NAMESPACE "MaterialEditorUtilities"
 
@@ -746,13 +747,13 @@ void FMaterialEditorUtilities::OpenSelectedParentEditor(UMaterialInterface* InMa
 		{
 			// Show material editor
 			UMaterial* Material = Cast<UMaterial>(InMaterialInterface);
-			FAssetEditorManager::Get().OpenEditorForAsset(Material);
+			GEditor->GetEditorSubsystem<UAssetEditorSubsystem>()->OpenEditorForAsset(Material);
 		}
 		else if (InMaterialInterface->IsA(UMaterialInstance::StaticClass()))
 		{
 			// Show material instance editor
 			UMaterialInstance* MaterialInstance = Cast<UMaterialInstance>(InMaterialInterface);
-			FAssetEditorManager::Get().OpenEditorForAsset(MaterialInstance);
+			GEditor->GetEditorSubsystem<UAssetEditorSubsystem>()->OpenEditorForAsset(MaterialInstance);
 		}
 	}
 }
@@ -766,12 +767,12 @@ void FMaterialEditorUtilities::OpenSelectedParentEditor(UMaterialFunctionInterfa
 		{
 			// Show function instance editor
 			UMaterialFunctionInstance* FunctionInstance = Cast<UMaterialFunctionInstance>(InMaterialFunction);
-			FAssetEditorManager::Get().OpenEditorForAsset(FunctionInstance);
+			GEditor->GetEditorSubsystem<UAssetEditorSubsystem>()->OpenEditorForAsset(FunctionInstance);
 		}
 		else
 		{
 			// Show function editor
-			FAssetEditorManager::Get().OpenEditorForAsset(InMaterialFunction);
+			GEditor->GetEditorSubsystem<UAssetEditorSubsystem>()->OpenEditorForAsset(InMaterialFunction);
 		}
 	}
 }
diff --git a/Engine/Source/Editor/MaterialEditor/Private/MaterialInstanceEditor.cpp b/Engine/Source/Editor/MaterialEditor/Private/MaterialInstanceEditor.cpp
--- a/Engine/Source/Editor/MaterialEditor/Private/MaterialInstanceEditor.cpp
+++ b/Engine/Source/Editor/MaterialEditor/Private/MaterialInstanceEditor.cpp
@@ -47,6 +47,7 @@
 #include "Widgets/Layout/SScrollBox.h"
 #include "DebugViewModeHelpers.h"
 #include "Widgets/Input/SButton.h"
+#include "Subsystems/AssetEditorSubsystem.h"
 
 #define LOCTEXT_NAMESPACE "MaterialInstanceEditor"
 
@@ -1458,13 +1459,13 @@ void FMaterialInstanceEditor::OpenSelectedParentEditor(UMaterialInterface* InMat
 		{
 			// Show material editor
 			UMaterial* Material = Cast<UMaterial>(InMaterialInterface);
-			FAssetEditorManager::Get().OpenEditorForAsset(Material);
+			GEditor->GetEditorSubsystem<UAssetEditorSubsystem>()->OpenEditorForAsset(Material);
 		}
 		else if(InMaterialInterface->IsA(UMaterialInstance::StaticClass()))
 		{
 			// Show material instance editor
 			UMaterialInstance* MaterialInstance = Cast<UMaterialInstance>(InMaterialInterface);
-			FAssetEditorManager::Get().OpenEditorForAsset(MaterialInstance);
+			GEditor->GetEditorSubsystem<UAssetEditorSubsystem>()->OpenEditorForAsset(MaterialInstance);
 		}
 	}
 }
@@ -1480,12 +1481,12 @@ void FMaterialInstanceEditor::OpenSelectedParentEditor(UMaterialFunctionInterfac
 		{
 			// Show function instance editor
 			UMaterialFunctionInstance* FunctionInstance = Cast<UMaterialFunctionInstance>(InMaterialFunction);
-			FAssetEditorManager::Get().OpenEditorForAsset(FunctionInstance);
+			GEditor->GetEditorSubsystem<UAssetEditorSubsystem>()->OpenEditorForAsset(FunctionInstance);
 		}
 		else
 		{
 			// Show function editor
-			FAssetEditorManager::Get().OpenEditorForAsset(InMaterialFunction);
+			GEditor->GetEditorSubsystem<UAssetEditorSubsystem>()->OpenEditorForAsset(InMaterialFunction);
 		}
 	}
 }
diff --git a/Engine/Source/Editor/MaterialEditor/Private/SMaterialLayersFunctionsTree.cpp b/Engine/Source/Editor/MaterialEditor/Private/SMaterialLayersFunctionsTree.cpp
--- a/Engine/Source/Editor/MaterialEditor/Private/SMaterialLayersFunctionsTree.cpp
+++ b/Engine/Source/Editor/MaterialEditor/Private/SMaterialLayersFunctionsTree.cpp
@@ -33,10 +33,12 @@
 #include "Widgets/Text/SInlineEditableTextBlock.h"
 #include "Materials/MaterialFunctionInstance.h"
 #include "Framework/Application/SlateApplication.h"
-#include "Toolkits/AssetEditorManager.h"
+
 #include "Widgets/Input/SEditableTextBox.h"
 #include "Curves/CurveLinearColor.h"
 #include "Curves/CurveLinearColorAtlas.h"
+#include "Subsystems/AssetEditorSubsystem.h"
+#include "Editor.h"
 
 
 #define LOCTEXT_NAMESPACE "MaterialLayerCustomization"
@@ -1448,7 +1450,7 @@ FReply SMaterialLayersFunctionsInstanceTree::OnThumbnailDoubleClick(const FGeome
 	}
 	if (AssetToOpen != nullptr)
 	{
-		FAssetEditorManager::Get().OpenEditorForAsset(AssetToOpen);
+		GEditor->GetEditorSubsystem<UAssetEditorSubsystem>()->OpenEditorForAsset(AssetToOpen);
 		return FReply::Handled();
 	}
 	return FReply::Unhandled();
diff --git a/Engine/Source/Editor/MaterialEditor/Private/Tests/MaterialEditorTests.cpp b/Engine/Source/Editor/MaterialEditor/Private/Tests/MaterialEditorTests.cpp
--- a/Engine/Source/Editor/MaterialEditor/Private/Tests/MaterialEditorTests.cpp
+++ b/Engine/Source/Editor/MaterialEditor/Private/Tests/MaterialEditorTests.cpp
@@ -6,7 +6,7 @@
 #include "MaterialGraph/MaterialGraph.h"
 #include "Engine/Texture.h"
 #include "AssetData.h"
-#include "Toolkits/AssetEditorManager.h"
+
 
 //Automation
 #include "Tests/AutomationTestSettings.h"
@@ -18,6 +18,7 @@
 //Materials
 #include "MaterialEditor.h"
 #include "Materials/MaterialExpressionTextureSample.h"
+#include "Subsystems/AssetEditorSubsystem.h"
 
 #if WITH_DEV_AUTOMATION_TESTS
 
@@ -50,7 +51,7 @@ namespace MaterialEditorPromotionTestUtils
 	*/
 	static bool AssignNormalToMaterial(UTexture* InNormalTexture, UMaterial* InMaterial)
 	{
-		IAssetEditorInstance* OpenEditor = FAssetEditorManager::Get().FindEditorForAsset(InMaterial, true);
+		IAssetEditorInstance* OpenEditor = GEditor->GetEditorSubsystem<UAssetEditorSubsystem>()->FindEditorForAsset(InMaterial, true);
 		IMaterialEditor* CurrentMaterialEditor = (IMaterialEditor*)OpenEditor;
 
 		//Create the texture sample and auto assign the normal texture
@@ -163,7 +164,7 @@ namespace MaterialEditorPromotionTestHelper
 				{
 					UE_LOG(LogEditorMaterialEditorPromotionTests, Display, TEXT("Created new material (%s) from texture (%s)"), *CreatedMaterial->GetName(), *DiffuseTexture->GetName());
 					// @TODO: Break out TEST 2: Open a material
-					FAssetEditorManager::Get().OpenEditorForAsset(CreatedMaterial);
+					GEditor->GetEditorSubsystem<UAssetEditorSubsystem>()->OpenEditorForAsset(CreatedMaterial);
 					UE_LOG(LogEditorMaterialEditorPromotionTests, Display, TEXT("Opened the material editor"));
 				}
 				else
@@ -197,13 +198,13 @@ namespace MaterialEditorPromotionTestHelper
 
 				if (NormalTexture)
 				{
-					IAssetEditorInstance* AssetEditor = FAssetEditorManager::Get().FindEditorForAsset(CreatedMaterial, true);
+					IAssetEditorInstance* AssetEditor = GEditor->GetEditorSubsystem<UAssetEditorSubsystem>()->FindEditorForAsset(CreatedMaterial, true);
 					FMaterialEditor* MaterialEditor = (FMaterialEditor*)AssetEditor;
 
 					// @TODO: Break out TEST 3: Add normal map to a material
 					if (MaterialEditorPromotionTestUtils::AssignNormalToMaterial(NormalTexture, CreatedMaterial))
 					{
-						FAssetEditorManager::Get().FindEditorForAsset(CreatedMaterial, true);
+						GEditor->GetEditorSubsystem<UAssetEditorSubsystem>()->FindEditorForAsset(CreatedMaterial, true);
 						// @TODO: Break out TEST 4: Compile a material
 						MaterialEditor->UpdateMaterialAfterGraphChange();
 						MaterialEditor->bMaterialDirty = false; // Remove dirty flag so window closes without prompt
diff --git a/Engine/Source/Editor/MovieSceneTools/Private/Sections/CinematicShotSection.cpp b/Engine/Source/Editor/MovieSceneTools/Private/Sections/CinematicShotSection.cpp
--- a/Engine/Source/Editor/MovieSceneTools/Private/Sections/CinematicShotSection.cpp
+++ b/Engine/Source/Editor/MovieSceneTools/Private/Sections/CinematicShotSection.cpp
@@ -14,11 +14,12 @@
 #include "EditorStyleSet.h"
 #include "MovieSceneToolHelpers.h"
 #include "MovieSceneTimeHelpers.h"
-#include "Toolkits/AssetEditorManager.h"
+
 #include "Tracks/MovieSceneCameraCutTrack.h"
 #include "Sections/MovieSceneCameraCutSection.h"
 #include "Evaluation/MovieSceneEvaluationTemplateInstance.h"
 #include "CommonMovieSceneTools.h"
+#include "Subsystems/AssetEditorSubsystem.h"
 
 #define LOCTEXT_NAMESPACE "FCinematicShotSection"
 
@@ -527,7 +528,7 @@ FReply FCinematicShotSection::OnSectionDoubleClicked(const FGeometry& SectionGeo
 		{
 			if (MouseEvent.IsControlDown())
 			{
-				FAssetEditorManager::Get().OpenEditorForAsset(SectionObject.GetSequence());
+				GEditor->GetEditorSubsystem<UAssetEditorSubsystem>()->OpenEditorForAsset(SectionObject.GetSequence());
 			}
 			else
 			{
diff --git a/Engine/Source/Editor/MovieSceneTools/Private/TrackEditors/SkeletalAnimationTrackEditor.cpp b/Engine/Source/Editor/MovieSceneTools/Private/TrackEditors/SkeletalAnimationTrackEditor.cpp
--- a/Engine/Source/Editor/MovieSceneTools/Private/TrackEditors/SkeletalAnimationTrackEditor.cpp
+++ b/Engine/Source/Editor/MovieSceneTools/Private/TrackEditors/SkeletalAnimationTrackEditor.cpp
@@ -38,12 +38,13 @@
 #include "Misc/MessageDialog.h"
 #include "Framework/Notifications/NotificationManager.h"
 #include "Widgets/Notifications/SNotificationList.h"
-#include "Toolkits/AssetEditorManager.h"
+
 #include "Engine/SCS_Node.h"
 #include "Engine/SimpleConstructionScript.h"
 #include "Engine/Blueprint.h"
 
 #include "CommonMovieSceneTools.h"
+#include "Subsystems/AssetEditorSubsystem.h"
 
 namespace SkeletalAnimationEditorConstants
 {
@@ -344,7 +345,7 @@ bool FSkeletalAnimationSection::CreatePoseAsset(const TArray<UObject*> NewAssets
 			Info.bUseLargeFont = false;
 			Info.Hyperlink = FSimpleDelegate::CreateLambda([NewAssets]()
 			{
-				FAssetEditorManager::Get().OpenEditorForAssets(NewAssets);
+				GEditor->GetEditorSubsystem<UAssetEditorSubsystem>()->OpenEditorForAssets(NewAssets);
 			});
 			Info.HyperlinkText = FText::Format(LOCTEXT("OpenNewPoseAssetHyperlink", "Open {0}"), FText::FromString(NewAssets[0]->GetName()));
 				
diff --git a/Engine/Source/Editor/Persona/Private/PersonaToolkit.cpp b/Engine/Source/Editor/Persona/Private/PersonaToolkit.cpp
--- a/Engine/Source/Editor/Persona/Private/PersonaToolkit.cpp
+++ b/Engine/Source/Editor/Persona/Private/PersonaToolkit.cpp
@@ -14,6 +14,7 @@
 #include "PersonaAssetFamily.h"
 #include "Interfaces/Interface_PreviewMeshProvider.h"
 #include "AnimationEditorPreviewActor.h"
+#include "Subsystems/AssetEditorSubsystem.h"
 
 FPersonaToolkit::FPersonaToolkit()
 	: Skeleton(nullptr)
@@ -319,8 +320,8 @@ void FPersonaToolkit::SetPreviewMesh(class USkeletalMesh* InSkeletalMesh, bool b
 			}
 			check(AssetToReopen);
 
-			FAssetEditorManager::Get().CloseAllEditorsForAsset(AssetToReopen);
-			FAssetEditorManager::Get().OpenEditorForAsset(AssetToReopen);
+			GEditor->GetEditorSubsystem<UAssetEditorSubsystem>()->CloseAllEditorsForAsset(AssetToReopen);
+			GEditor->GetEditorSubsystem<UAssetEditorSubsystem>()->OpenEditorForAsset(AssetToReopen);
 			return;
 		}
 
diff --git a/Engine/Source/Editor/Persona/Private/SAnimNotifyPanel.cpp b/Engine/Source/Editor/Persona/Private/SAnimNotifyPanel.cpp
--- a/Engine/Source/Editor/Persona/Private/SAnimNotifyPanel.cpp
+++ b/Engine/Source/Editor/Persona/Private/SAnimNotifyPanel.cpp
@@ -12,7 +12,7 @@
 #include "Widgets/Text/STextBlock.h"
 #include "Framework/MultiBox/MultiBoxBuilder.h"
 #include "Animation/AnimSequence.h"
-#include "Toolkits/AssetEditorManager.h"
+
 #include "Layout/WidgetPath.h"
 #include "Framework/Application/MenuStack.h"
 #include "Fonts/FontMeasure.h"
@@ -43,6 +43,7 @@
 #include "ClassViewerFilter.h"
 #include "ClassViewerModule.h"
 #include "SSkeletonAnimNotifies.h"
+#include "Subsystems/AssetEditorSubsystem.h"
 
 // Track Panel drawing
 const float NotificationTrackHeight = 20.0f;
@@ -3165,7 +3166,7 @@ void SAnimNotifyTrack::OnManageNotifies()
 
 void SAnimNotifyTrack::OnOpenNotifySource(UBlueprint* InSourceBlueprint) const
 {
-	FAssetEditorManager::Get().OpenEditorForAsset(InSourceBlueprint);
+	GEditor->GetEditorSubsystem<UAssetEditorSubsystem>()->OpenEditorForAsset(InSourceBlueprint);
 }
 
 bool SAnimNotifyTrack::IsSingleNodeSelected()
diff --git a/Engine/Source/Editor/Persona/Private/SAnimSegmentsPanel.cpp b/Engine/Source/Editor/Persona/Private/SAnimSegmentsPanel.cpp
--- a/Engine/Source/Editor/Persona/Private/SAnimSegmentsPanel.cpp
+++ b/Engine/Source/Editor/Persona/Private/SAnimSegmentsPanel.cpp
@@ -8,7 +8,7 @@
 #include "AssetData.h"
 #include "Animation/AnimSequenceBase.h"
 #include "Animation/AnimSequence.h"
-#include "Toolkits/AssetEditorManager.h"
+
 #include "Animation/AnimCompositeBase.h"
 
 #include "ScopedTransaction.h"
@@ -16,6 +16,8 @@
 #include "IContentBrowserSingleton.h"
 #include "ContentBrowserModule.h"
 #include "Framework/Application/SlateApplication.h"
+#include "Subsystems/AssetEditorSubsystem.h"
+#include "Editor.h"
 
 #define LOCTEXT_NAMESPACE "AnimSegmentPanel"
 
@@ -402,7 +404,7 @@ void SAnimSegmentsPanel::OpenAsset(int32 AnimSegmentIndex)
 		UAnimSequenceBase* Asset = AnimTrack->AnimSegments[AnimSegmentIndex].AnimReference;
 		if (Asset)
 		{
-			FAssetEditorManager::Get().OpenEditorForAsset(Asset);
+			GEditor->GetEditorSubsystem<UAssetEditorSubsystem>()->OpenEditorForAsset(Asset);
 		}
 	}
 }
diff --git a/Engine/Source/Editor/Persona/Private/SAnimationEditorViewport.cpp b/Engine/Source/Editor/Persona/Private/SAnimationEditorViewport.cpp
--- a/Engine/Source/Editor/Persona/Private/SAnimationEditorViewport.cpp
+++ b/Engine/Source/Editor/Persona/Private/SAnimationEditorViewport.cpp
@@ -33,9 +33,10 @@
 #include "Editor/EditorPerProjectUserSettings.h"
 #include "Materials/Material.h"
 #include "EditorFontGlyphs.h"
-#include "Toolkits/AssetEditorManager.h"
+
 #include "SkeletalMeshTypes.h"
 #include "IPersonaToolkit.h"
+#include "Subsystems/AssetEditorSubsystem.h"
 
 #define LOCTEXT_NAMESPACE "PersonaViewportToolbar"
 
@@ -2187,7 +2188,7 @@ void SAnimationEditorViewportTabBody::AddPostProcessNotification()
 		{
 			if(PreviewComponent->SkeletalMesh && PreviewComponent->SkeletalMesh->PostProcessAnimBlueprint)
 			{
-				FAssetEditorManager::Get().OpenEditorForAssets(TArray<UObject*>({ PreviewComponent->SkeletalMesh->PostProcessAnimBlueprint->ClassGeneratedBy }));
+				GEditor->GetEditorSubsystem<UAssetEditorSubsystem>()->OpenEditorForAssets(TArray<UObject*>({ PreviewComponent->SkeletalMesh->PostProcessAnimBlueprint->ClassGeneratedBy }));
 			}
 		}
 
diff --git a/Engine/Source/Editor/Persona/Private/SAnimationSequenceBrowser.cpp b/Engine/Source/Editor/Persona/Private/SAnimationSequenceBrowser.cpp
--- a/Engine/Source/Editor/Persona/Private/SAnimationSequenceBrowser.cpp
+++ b/Engine/Source/Editor/Persona/Private/SAnimationSequenceBrowser.cpp
@@ -12,7 +12,7 @@
 #include "Animation/AnimationAsset.h"
 #include "Animation/AnimSequenceBase.h"
 #include "Animation/AnimSequence.h"
-#include "Toolkits/AssetEditorManager.h"
+
 #include "EditorStyleSet.h"
 #include "Animation/DebugSkelMeshComponent.h"
 #include "IPersonaPreviewScene.h"
@@ -40,6 +40,7 @@
 #include "Framework/Application/SlateApplication.h"
 #include "Preferences/PersonaOptions.h"
 #include "Settings/SkeletalMeshEditorSettings.h"
+#include "Subsystems/AssetEditorSubsystem.h"
 
 #define LOCTEXT_NAMESPACE "SequenceBrowser"
 
@@ -520,7 +521,7 @@ void SAnimationSequenceBrowser::RetargetAnimationHandler(USkeleton* OldSkeleton,
 		FAssetRegistryModule::AssetCreated(AssetToOpen);
 
 		// once all success, attempt to open new editor with new skeleton
-		FAssetEditorManager::Get().OpenEditorForAsset(AssetToOpen);
+		GEditor->GetEditorSubsystem<UAssetEditorSubsystem>()->OpenEditorForAsset(AssetToOpen);
 	}
 }
 
diff --git a/Engine/Source/Editor/Persona/Private/SAssetFamilyShortcutBar.cpp b/Engine/Source/Editor/Persona/Private/SAssetFamilyShortcutBar.cpp
--- a/Engine/Source/Editor/Persona/Private/SAssetFamilyShortcutBar.cpp
+++ b/Engine/Source/Editor/Persona/Private/SAssetFamilyShortcutBar.cpp
@@ -12,7 +12,7 @@
 #include "Framework/MultiBox/MultiBoxDefs.h"
 #include "Widgets/Text/STextBlock.h"
 #include "Framework/MultiBox/MultiBoxBuilder.h"
-#include "Toolkits/AssetEditorManager.h"
+
 #include "EditorStyleSet.h"
 #include "IAssetFamily.h"
 #include "IContentBrowserSingleton.h"
@@ -23,6 +23,7 @@
 #include "Widgets/Images/SImage.h"
 #include "Widgets/Input/SComboButton.h"
 #include "Framework/Application/SlateApplication.h"
+#include "Subsystems/AssetEditorSubsystem.h"
 
 #define LOCTEXT_NAMESPACE "SAssetFamilyShortcutBar"
 
@@ -54,7 +55,7 @@ public:
 		AssetRegistryModule.Get().OnAssetRemoved().AddSP(this, &SAssetShortcut::HandleAssetRemoved);
 		AssetRegistryModule.Get().OnAssetRenamed().AddSP(this, &SAssetShortcut::HandleAssetRenamed);
 
-		FAssetEditorManager::Get().OnAssetEditorRequestedOpen().AddSP(this, &SAssetShortcut::HandleAssetOpened);
+		GEditor->GetEditorSubsystem<UAssetEditorSubsystem>()->OnAssetEditorRequestedOpen().AddSP(this, &SAssetShortcut::HandleAssetOpened);
 		AssetFamily->GetOnAssetOpened().AddSP(this, &SAssetShortcut::HandleAssetOpened);
 
 		AssetThumbnail = MakeShareable(new FAssetThumbnail(InAssetData, AssetShortcutConstants::ThumbnailSize, AssetShortcutConstants::ThumbnailSize, InThumbnailPool));
@@ -183,7 +184,7 @@ public:
 		}
 
 		AssetFamily->GetOnAssetOpened().RemoveAll(this);
-		FAssetEditorManager::Get().OnAssetEditorRequestedOpen().RemoveAll(this);
+		GEditor->GetEditorSubsystem<UAssetEditorSubsystem>()->OnAssetEditorRequestedOpen().RemoveAll(this);
 		UnRegisterActiveTimer(DirtyStateTimerHandle.ToSharedRef());
 	}
 
@@ -193,7 +194,7 @@ public:
 		{
 			TArray<UObject*> Assets;
 			Assets.Add(AssetData.GetAsset());
-			FAssetEditorManager::Get().OpenEditorForAssets(Assets);
+			GEditor->GetEditorSubsystem<UAssetEditorSubsystem>()->OpenEditorForAssets(Assets);
 		}
 	}
 
@@ -285,7 +286,7 @@ public:
 
 			TArray<UObject*> Assets;
 			Assets.Add(InAssetData.GetAsset());
-			FAssetEditorManager::Get().OpenEditorForAssets(Assets);
+			GEditor->GetEditorSubsystem<UAssetEditorSubsystem>()->OpenEditorForAssets(Assets);
 		}
 		else if(AssetData.IsValid())
 		{
@@ -294,7 +295,7 @@ public:
 			// Assume that as we are set to 'toggle' mode with no 'none' selection allowed, we are selecting the currently selected item
 			TArray<UObject*> Assets;
 			Assets.Add(AssetData.GetAsset());
-			FAssetEditorManager::Get().OpenEditorForAssets(Assets);
+			GEditor->GetEditorSubsystem<UAssetEditorSubsystem>()->OpenEditorForAssets(Assets);
 		}
 	}
 
diff --git a/Engine/Source/Editor/Persona/Private/SMontageEditor.cpp b/Engine/Source/Editor/Persona/Private/SMontageEditor.cpp
--- a/Engine/Source/Editor/Persona/Private/SMontageEditor.cpp
+++ b/Engine/Source/Editor/Persona/Private/SMontageEditor.cpp
@@ -6,7 +6,7 @@
 #include "Widgets/Layout/SBorder.h"
 #include "Widgets/Text/STextBlock.h"
 #include "Editor.h"
-#include "Toolkits/AssetEditorManager.h"
+
 #include "Animation/DebugSkelMeshComponent.h"
 #include "Animation/EditorCompositeSection.h"
 #include "IDocumentation.h"
@@ -22,6 +22,7 @@
 #include "Widgets/Images/SImage.h"
 #include "Factories/AnimMontageFactory.h"
 #include "Editor.h"
+#include "Subsystems/AssetEditorSubsystem.h"
 
 #define LOCTEXT_NAMESPACE "AnimSequenceEditor"
 
@@ -244,7 +245,7 @@ FReply SMontageEditor::OnEditParentClassClicked()
 		UObject* ParentClass = MontageObj->ParentAsset;
 		if (ParentClass != NULL)
 		{
-			FAssetEditorManager::Get().OpenEditorForAsset(ParentClass);
+			GEditor->GetEditorSubsystem<UAssetEditorSubsystem>()->OpenEditorForAsset(ParentClass);
 		}
 	}
 
diff --git a/Engine/Source/Editor/Persona/Private/SSlotNameReferenceWindow.cpp b/Engine/Source/Editor/Persona/Private/SSlotNameReferenceWindow.cpp
--- a/Engine/Source/Editor/Persona/Private/SSlotNameReferenceWindow.cpp
+++ b/Engine/Source/Editor/Persona/Private/SSlotNameReferenceWindow.cpp
@@ -1,7 +1,7 @@
 // Copyright 1998-2019 Epic Games, Inc. All Rights Reserved.
 
 #include "SSlotNameReferenceWindow.h"
-#include "Toolkits/AssetEditorManager.h"
+
 #include "EditorStyleSet.h"
 #include "Styling/CoreStyle.h"
 #include "Widgets/Images/SImage.h"
@@ -10,6 +10,7 @@
 #include "AnimGraphNode_Slot.h"
 #include "EdGraph/EdGraph.h"
 #include "BlueprintEditor.h"
+#include "Subsystems/AssetEditorSubsystem.h"
 
 #define LOCTEXT_NAMESPACE "SkeletonSlotNames"
 
@@ -103,7 +104,7 @@ FReply SMontageReferenceListRow::OnViewInContentBrowserClicked()
 
 FReply SMontageReferenceListRow::OnOpenAssetClicked()
 {
-	FAssetEditorManager::Get().OpenEditorForAsset(ReferenceInfo->AssetData.GetAsset());
+	GEditor->GetEditorSubsystem<UAssetEditorSubsystem>()->OpenEditorForAsset(ReferenceInfo->AssetData.GetAsset());
 
 	return FReply::Handled();
 }
@@ -181,9 +182,9 @@ FReply SBlueprintReferenceListRow::OnViewInContentBrowserClicked()
 
 FReply SBlueprintReferenceListRow::OnOpenAssetClicked()
 {
-	FAssetEditorManager::Get().OpenEditorForAsset(ReferenceInfo->AnimBlueprint);
+	GEditor->GetEditorSubsystem<UAssetEditorSubsystem>()->OpenEditorForAsset(ReferenceInfo->AnimBlueprint);
 	
-	FBlueprintEditor* BPEditor = static_cast<FBlueprintEditor*>(FAssetEditorManager::Get().FindEditorForAsset(ReferenceInfo->AnimBlueprint, true));
+	FBlueprintEditor* BPEditor = static_cast<FBlueprintEditor*>(GEditor->GetEditorSubsystem<UAssetEditorSubsystem>()->FindEditorForAsset(ReferenceInfo->AnimBlueprint, true));
 	if(BPEditor)
 	{
 		// Should only ever get an FPersona object back for an animBP.
diff --git a/Engine/Source/Editor/PropertyEditor/Private/Presentation/PropertyEditor/PropertyEditor.cpp b/Engine/Source/Editor/PropertyEditor/Private/Presentation/PropertyEditor/PropertyEditor.cpp
--- a/Engine/Source/Editor/PropertyEditor/Private/Presentation/PropertyEditor/PropertyEditor.cpp
+++ b/Engine/Source/Editor/PropertyEditor/Private/Presentation/PropertyEditor/PropertyEditor.cpp
@@ -6,7 +6,7 @@
 #include "CategoryPropertyNode.h"
 #include "ItemPropertyNode.h"
 #include "ObjectPropertyNode.h"
-#include "Toolkits/AssetEditorManager.h"
+
 #include "Editor/SceneOutliner/Public/SceneOutlinerFilters.h"
 #include "IDetailPropertyRow.h"
 #include "PropertyCustomizationHelpers.h"
@@ -19,6 +19,7 @@
 #include "Kismet2/BlueprintEditorUtils.h"
 #include "EditConditionParser.h"
 #include "EditConditionContext.h"
+#include "Subsystems/AssetEditorSubsystem.h"
 
 #define LOCTEXT_NAMESPACE "PropertyEditor"
 
@@ -297,7 +298,7 @@ void FPropertyEditor::MakeNewBlueprint()
 			
 			PropertyHandle->SetValueFromFormattedString(Blueprint->GeneratedClass->GetPathName());
 
-			FAssetEditorManager::Get().OpenEditorForAsset(Blueprint);
+			GEditor->GetEditorSubsystem<UAssetEditorSubsystem>()->OpenEditorForAsset(Blueprint);
 		}
 	}
 }
diff --git a/Engine/Source/Editor/PropertyEditor/Private/PropertyEditorToolkit.cpp b/Engine/Source/Editor/PropertyEditor/Private/PropertyEditorToolkit.cpp
--- a/Engine/Source/Editor/PropertyEditor/Private/PropertyEditorToolkit.cpp
+++ b/Engine/Source/Editor/PropertyEditor/Private/PropertyEditorToolkit.cpp
@@ -13,6 +13,7 @@
 #include "IPropertyTableRow.h"
 
 #include "Widgets/Docking/SDockTab.h"
+#include "Subsystems/AssetEditorSubsystem.h"
 
 #define LOCTEXT_NAMESPACE "PropertyEditorToolkit"
 
@@ -55,7 +56,7 @@ FPropertyEditorToolkit::FPropertyEditorToolkit()
 TSharedPtr<FPropertyEditorToolkit> FPropertyEditorToolkit::FindExistingEditor( UObject* Object )
 {
 	// Find any existing property editor instances for this asset
-	const TArray<IAssetEditorInstance*> Editors = FAssetEditorManager::Get().FindEditorsForAsset( Object );
+	const TArray<IAssetEditorInstance*> Editors = GEditor->GetEditorSubsystem<UAssetEditorSubsystem>()->FindEditorsForAsset( Object );
 
 	IAssetEditorInstance* const * ExistingInstance = Editors.FindByPredicate( [&]( IAssetEditorInstance* Editor ){
 		return Editor->GetEditorName() == ToolkitFName;
diff --git a/Engine/Source/Editor/SequenceRecorder/Private/ActorRecording.cpp b/Engine/Source/Editor/SequenceRecorder/Private/ActorRecording.cpp
--- a/Engine/Source/Editor/SequenceRecorder/Private/ActorRecording.cpp
+++ b/Engine/Source/Editor/SequenceRecorder/Private/ActorRecording.cpp
@@ -20,9 +20,10 @@
 #include "CameraRig_Rail.h"
 #include "SequenceRecorder.h"
 #include "Features/IModularFeatures.h"
-#include "Toolkits/AssetEditorManager.h"
+
 #include "IAssetTools.h"
 #include "AssetToolsModule.h"
+#include "Subsystems/AssetEditorSubsystem.h"
 
 static const FName SequencerActorTag(TEXT("SequencerActor"));
 static const FName MovieSceneSectionRecorderFactoryName("MovieSceneSectionRecorderFactory");
@@ -129,7 +130,7 @@ bool UActorRecording::StartRecording(ULevelSequence* CurrentSequence, float Curr
 	{
 		if (TargetAnimation != nullptr)
 		{
-			IAssetEditorInstance* EditorInstance = FAssetEditorManager::Get().FindEditorForAsset(TargetAnimation, false);
+			IAssetEditorInstance* EditorInstance = GEditor->GetEditorSubsystem<UAssetEditorSubsystem>()->FindEditorForAsset(TargetAnimation, false);
 			if (EditorInstance)
 			{
 				UE_LOG(LogAnimation, Log, TEXT("Closing '%s' so we don't invalidate the open version when unloading it."), *TargetAnimation->GetName());
diff --git a/Engine/Source/Editor/SequenceRecorder/Private/AnimationRecorder.cpp b/Engine/Source/Editor/SequenceRecorder/Private/AnimationRecorder.cpp
--- a/Engine/Source/Editor/SequenceRecorder/Private/AnimationRecorder.cpp
+++ b/Engine/Source/Editor/SequenceRecorder/Private/AnimationRecorder.cpp
@@ -7,7 +7,7 @@
 #include "UObject/Package.h"
 #include "Misc/PackageName.h"
 #include "Editor.h"
-#include "Toolkits/AssetEditorManager.h"
+
 #include "Animation/AnimInstance.h"
 #include "Animation/AnimCompress.h"
 #include "Animation/AnimCompress_BitwiseCompressOnly.h"
@@ -19,6 +19,7 @@
 #include "Animation/AnimationRecordingSettings.h"
 #include "Animation/AnimNotifies/AnimNotify.h"
 #include "Animation/AnimNotifies/AnimNotifyState.h"
+#include "Subsystems/AssetEditorSubsystem.h"
 
 #define LOCTEXT_NAMESPACE "FAnimationRecorder"
 
@@ -408,7 +409,7 @@ UAnimSequence* FAnimationRecorder::StopRecord(bool bShowMessage)
 				{
 					TArray<UObject*> Assets;
 					Assets.Add(ReturnObject);
-					FAssetEditorManager::Get().OpenEditorForAssets(Assets);
+					GEditor->GetEditorSubsystem<UAssetEditorSubsystem>()->OpenEditorForAssets(Assets);
 				});
 				Info.HyperlinkText = FText::Format(LOCTEXT("OpenNewAnimationHyperlink", "Open {0}"), FText::FromString(AnimationObject->GetName()));
 				TSharedPtr<SNotificationItem> Notification = FSlateNotificationManager::Get().AddNotification(Info);
diff --git a/Engine/Source/Editor/Sequencer/Private/Sequencer.cpp b/Engine/Source/Editor/Sequencer/Private/Sequencer.cpp
--- a/Engine/Source/Editor/Sequencer/Private/Sequencer.cpp
+++ b/Engine/Source/Editor/Sequencer/Private/Sequencer.cpp
@@ -130,6 +130,7 @@
 #include "ISerializedRecorder.h"
 #include "Features/IModularFeatures.h"
 #include "SequencerContextMenus.h"
+#include "Subsystems/AssetEditorSubsystem.h"
 
 #define LOCTEXT_NAMESPACE "Sequencer"
 
@@ -5048,9 +5049,9 @@ void FSequencer::SaveCurrentMovieSceneAs()
 
 	if ((SavedAssets[0] != AssetsToSave[0]) && (SavedAssets[0] != nullptr))
 	{
-		FAssetEditorManager& AssetEditorManager = FAssetEditorManager::Get();
-		AssetEditorManager.CloseAllEditorsForAsset(AssetsToSave[0]);
-		AssetEditorManager.OpenEditorForAssets(SavedAssets, EToolkitMode::Standalone, MyToolkitHost.ToSharedRef());
+		UAssetEditorSubsystem* AssetEditorSubsystem = GEditor->GetEditorSubsystem<UAssetEditorSubsystem>();
+		AssetEditorSubsystem->CloseAllEditorsForAsset(AssetsToSave[0]);
+		AssetEditorSubsystem->OpenEditorForAssets(SavedAssets, EToolkitMode::Standalone, MyToolkitHost.ToSharedRef());
 	}
 }
 
@@ -9144,13 +9145,13 @@ void FSequencer::DiscardChanges()
 		return;
 	}
 
-	FAssetEditorManager& AssetEditorManager = FAssetEditorManager::Get();
+	UAssetEditorSubsystem* AssetEditorSubsystem = GEditor->GetEditorSubsystem<UAssetEditorSubsystem>();
 	UClass* SequenceClass = EditedSequence->GetClass();
 	FString SequencePath = EditedSequence->GetPathName();
 	UPackage* SequencePackage = EditedSequence->GetOutermost();
 
 	// close asset editor
-	AssetEditorManager.CloseAllEditorsForAsset(EditedSequence);
+	AssetEditorSubsystem->CloseAllEditorsForAsset(EditedSequence);
 
 	// collect objects to be unloaded
 	TMap<FString, UObject*> MovedObjects;
@@ -9234,7 +9235,7 @@ void FSequencer::DiscardChanges()
 	TArray<UObject*> AssetsToReopen;
 	AssetsToReopen.Add(ReloadedSequence);
 
-	AssetEditorManager.OpenEditorForAssets(AssetsToReopen, EToolkitMode::Standalone, MyToolkitHost.ToSharedRef());
+	AssetEditorSubsystem->OpenEditorForAssets(AssetsToReopen, EToolkitMode::Standalone, MyToolkitHost.ToSharedRef());
 }
 
 
diff --git a/Engine/Source/Editor/StaticMeshEditor/Private/StaticMeshAutomationTests.cpp b/Engine/Source/Editor/StaticMeshEditor/Private/StaticMeshAutomationTests.cpp
--- a/Engine/Source/Editor/StaticMeshEditor/Private/StaticMeshAutomationTests.cpp
+++ b/Engine/Source/Editor/StaticMeshEditor/Private/StaticMeshAutomationTests.cpp
@@ -7,10 +7,11 @@
 #include "Framework/Application/SlateApplication.h"
 #include "Engine/StaticMesh.h"
 #include "Editor.h"
-#include "Toolkits/AssetEditorManager.h"
+
 
 #include "Tests/AutomationCommon.h"
 #include "StaticMeshEditorViewportClient.h"
+#include "Subsystems/AssetEditorSubsystem.h"
 
 namespace EditorViewButtonHelper
 {
@@ -147,7 +148,7 @@ namespace EditorViewButtonHelper
 
 	bool FCloseAllAssetEditorsCommand::Update()
 	{
-		FAssetEditorManager::Get().CloseAllAssetEditors();
+		GEditor->GetEditorSubsystem<UAssetEditorSubsystem>()->CloseAllAssetEditors();
 
 		return true;
 	}
@@ -179,7 +180,7 @@ bool FStaticMeshEditorTest::RunTest(const FString& Parameters)
 		NULL,
 		LOAD_None,
 		NULL);
-	FAssetEditorManager::Get().OpenEditorForAsset(EditorMesh);
+	GEditor->GetEditorSubsystem<UAssetEditorSubsystem>()->OpenEditorForAsset(EditorMesh);
 
 	{
 		TArray<TSharedRef<SWindow>> AllWindows;
diff --git a/Engine/Source/Editor/StructViewer/Private/SStructViewer.cpp b/Engine/Source/Editor/StructViewer/Private/SStructViewer.cpp
--- a/Engine/Source/Editor/StructViewer/Private/SStructViewer.cpp
+++ b/Engine/Source/Editor/StructViewer/Private/SStructViewer.cpp
@@ -34,7 +34,7 @@
 #include "ARFilter.h"
 #include "AssetData.h"
 #include "AssetRegistryModule.h"
-#include "Toolkits/AssetEditorManager.h"
+
 #include "DragAndDrop/AssetDragDropOp.h"
 
 #include "Editor/UnrealEdEngine.h"
@@ -46,6 +46,7 @@
 #include "PropertyHandle.h"
 #include "Logging/LogMacros.h"
 #include "SourceCodeNavigation.h"
+#include "Subsystems/AssetEditorSubsystem.h"
 
 #define LOCTEXT_NAMESPACE "SStructViewer"
 
@@ -537,7 +538,7 @@ namespace StructViewer
 		{
 			if (InStruct)
 			{
-				FAssetEditorManager::Get().OpenEditorForAsset(const_cast<UUserDefinedStruct*>(InStruct));
+				GEditor->GetEditorSubsystem<UAssetEditorSubsystem>()->OpenEditorForAsset(const_cast<UUserDefinedStruct*>(InStruct));
 			}
 		}
 
diff --git a/Engine/Source/Editor/UMGEditor/Private/Templates/WidgetTemplateBlueprintClass.cpp b/Engine/Source/Editor/UMGEditor/Private/Templates/WidgetTemplateBlueprintClass.cpp
--- a/Engine/Source/Editor/UMGEditor/Private/Templates/WidgetTemplateBlueprintClass.cpp
+++ b/Engine/Source/Editor/UMGEditor/Private/Templates/WidgetTemplateBlueprintClass.cpp
@@ -5,10 +5,11 @@
 #include "Widgets/SToolTip.h"
 #include "IDocumentation.h"
 #include "WidgetBlueprint.h"
-#include "Toolkits/AssetEditorManager.h"
+
 
 #include "Kismet2/BlueprintEditorUtils.h"
 #include "Styling/SlateIconFinder.h"
+#include "Subsystems/AssetEditorSubsystem.h"
 
 #define LOCTEXT_NAMESPACE "UMGEditor"
 
@@ -83,7 +84,7 @@ TSharedRef<IToolTip> FWidgetTemplateBlueprintClass::GetToolTip() const
 
 FReply FWidgetTemplateBlueprintClass::OnDoubleClicked()
 {
-	FAssetEditorManager::Get().OpenEditorForAsset( WidgetAssetData.GetAsset() );
+	GEditor->GetEditorSubsystem<UAssetEditorSubsystem>()->OpenEditorForAsset( WidgetAssetData.GetAsset() );
 	return FReply::Handled();
 }
 
diff --git a/Engine/Source/Editor/UMGEditor/Private/WidgetBlueprintEditorUtils.cpp b/Engine/Source/Editor/UMGEditor/Private/WidgetBlueprintEditorUtils.cpp
--- a/Engine/Source/Editor/UMGEditor/Private/WidgetBlueprintEditorUtils.cpp
+++ b/Engine/Source/Editor/UMGEditor/Private/WidgetBlueprintEditorUtils.cpp
@@ -34,6 +34,7 @@
 #include "Widgets/Notifications/SNotificationList.h"
 #include "Components/Widget.h"
 #include "Blueprint/WidgetNavigation.h"
+#include "Subsystems/AssetEditorSubsystem.h"
 
 #define LOCTEXT_NAMESPACE "UMG"
 
@@ -343,7 +344,7 @@ void FWidgetBlueprintEditorUtils::ExecuteOpenSelectedWidgetsForEdit( TSet<FWidge
 {
 	for ( auto& Widget : SelectedWidgets )
 	{
-		FAssetEditorManager::Get().OpenEditorForAsset( Widget.GetTemplate()->GetClass()->ClassGeneratedBy );
+		GEditor->GetEditorSubsystem<UAssetEditorSubsystem>()->OpenEditorForAsset( Widget.GetTemplate()->GetClass()->ClassGeneratedBy );
 	}
 }
 
diff --git a/Engine/Source/Editor/UnrealEd/Private/Dialogs/DlgReferenceTree.cpp b/Engine/Source/Editor/UnrealEd/Private/Dialogs/DlgReferenceTree.cpp
--- a/Engine/Source/Editor/UnrealEd/Private/Dialogs/DlgReferenceTree.cpp
+++ b/Engine/Source/Editor/UnrealEd/Private/Dialogs/DlgReferenceTree.cpp
@@ -14,8 +14,9 @@
 #include "Editor.h"
 #include "UnrealEdGlobals.h"
 
-#include "Toolkits/AssetEditorManager.h"
+
 #include "ObjectTools.h"
+#include "Subsystems/AssetEditorSubsystem.h"
 
 FArchiveGenerateReferenceGraph::FArchiveGenerateReferenceGraph( FReferenceGraph& OutGraph ) 
 	: CurrentObject(NULL),
@@ -538,7 +539,7 @@ void SReferenceTree::OnMenuViewProperties( UObject* InObject )
 void SReferenceTree::OnMenuShowEditor( UObject* InObject )
 {
 	// Show the editor for this object
-	FAssetEditorManager::Get().OpenEditorForAsset(InObject);
+	GEditor->GetEditorSubsystem<UAssetEditorSubsystem>()->OpenEditorForAsset(InObject);
 }
 
 void SReferenceTree::OnShowScriptReferences()
diff --git a/Engine/Source/Editor/UnrealEd/Private/Dialogs/SDeleteAssetsDialog.cpp b/Engine/Source/Editor/UnrealEd/Private/Dialogs/SDeleteAssetsDialog.cpp
--- a/Engine/Source/Editor/UnrealEd/Private/Dialogs/SDeleteAssetsDialog.cpp
+++ b/Engine/Source/Editor/UnrealEd/Private/Dialogs/SDeleteAssetsDialog.cpp
@@ -18,8 +18,9 @@
 #include "IContentBrowserSingleton.h"
 #include "ContentBrowserModule.h"
 #include "Editor.h"
-#include "Toolkits/AssetEditorManager.h"
+
 #include "Framework/Commands/GenericCommands.h"
+#include "Subsystems/AssetEditorSubsystem.h"
 
 #define LOCTEXT_NAMESPACE "SDeleteAssetsDialog"
 
@@ -850,7 +851,7 @@ void SDeleteAssetsDialog::OnAssetsActivated(const TArray<FAssetData>& ActivatedA
 			}
 			else
 			{
-				FAssetEditorManager::Get().OpenEditorForAsset(ActivatedAsset.GetAsset());
+				GEditor->GetEditorSubsystem<UAssetEditorSubsystem>()->OpenEditorForAsset(ActivatedAsset.GetAsset());
 
 				// @todo ndarnell select in content browser maybe as well?
 			}
diff --git a/Engine/Source/Editor/UnrealEd/Private/EditorActor.cpp b/Engine/Source/Editor/UnrealEd/Private/EditorActor.cpp
--- a/Engine/Source/Editor/UnrealEd/Private/EditorActor.cpp
+++ b/Engine/Source/Editor/UnrealEd/Private/EditorActor.cpp
@@ -65,6 +65,7 @@
 #include "AssetToolsModule.h"
 #include "AssetSelection.h"
 #include "Framework/Application/SlateApplication.h"
+#include "Subsystems/AssetEditorSubsystem.h"
 
 #define LOCTEXT_NAMESPACE "UnrealEd.EditorActor"
 
@@ -1033,7 +1034,7 @@ bool UUnrealEdEngine::edactDeleteSelected( UWorld* InWorld, bool bVerifyDeletion
 		}
 
 		// Remove actor from all asset editors
-		FAssetEditorManager::Get().RemoveAssetFromAllEditors(Actor);
+		GEditor->GetEditorSubsystem<UAssetEditorSubsystem>()->RemoveAssetFromAllEditors(Actor);
 
 		// Mark the actor's level as dirty.
 		Actor->MarkPackageDirty();
diff --git a/Engine/Source/Editor/UnrealEd/Private/EditorAnimUtils.cpp b/Engine/Source/Editor/UnrealEd/Private/EditorAnimUtils.cpp
--- a/Engine/Source/Editor/UnrealEd/Private/EditorAnimUtils.cpp
+++ b/Engine/Source/Editor/UnrealEd/Private/EditorAnimUtils.cpp
@@ -12,11 +12,12 @@
 #include "IAssetTools.h"
 #include "AssetToolsModule.h"
 #include "Framework/Notifications/NotificationManager.h"
-#include "Toolkits/AssetEditorManager.h"
+
 #include "ObjectEditorUtils.h"
 #include "Widgets/Notifications/SNotificationList.h"
 #include "IContentBrowserSingleton.h"
 #include "ContentBrowserModule.h"
+#include "Subsystems/AssetEditorSubsystem.h"
 
 #define LOCTEXT_NAMESPACE "EditorAnimUtils"
 
@@ -309,7 +310,7 @@ namespace EditorAnimUtils
 
 	void OpenAssetFromNotify(UObject* AssetToOpen)
 	{
-		FAssetEditorManager::Get().OpenEditorForAsset(AssetToOpen);
+		GEditor->GetEditorSubsystem<UAssetEditorSubsystem>()->OpenEditorForAsset(AssetToOpen);
 	}
 
 	//////////////////////////////////////////////////////////////////
diff --git a/Engine/Source/Editor/UnrealEd/Private/EditorEngine.cpp b/Engine/Source/Editor/UnrealEd/Private/EditorEngine.cpp
--- a/Engine/Source/Editor/UnrealEd/Private/EditorEngine.cpp
+++ b/Engine/Source/Editor/UnrealEd/Private/EditorEngine.cpp
@@ -103,7 +103,7 @@
 #include "Layers/LayersSubsystem.h"
 #include "EditorLevelUtils.h"
 
-#include "Toolkits/AssetEditorManager.h"
+
 #include "PropertyEditorModule.h"
 #include "AssetSelection.h"
 
@@ -214,6 +214,8 @@
 #include "RenderTargetPool.h"
 #include "RenderGraphBuilder.h"
 #include "ToolMenus.h"
+#include "Subsystems/AssetEditorSubsystem.h"
+
 
 DEFINE_LOG_CATEGORY_STATIC(LogEditor, Log, All);
 
@@ -793,7 +795,7 @@ void UEditorEngine::InitEditor(IEngineLoop* InEngineLoop)
 
 bool UEditorEngine::HandleOpenAsset(UObject* Asset)
 {
-	return FAssetEditorManager::Get().OpenEditorForAsset(Asset);
+	return GEditor->GetEditorSubsystem<UAssetEditorSubsystem>()->OpenEditorForAsset(Asset);
 }
 
 void UEditorEngine::HandleSettingChanged( FName Name )
@@ -2368,8 +2370,8 @@ void UEditorEngine::CloseEditedWorldAssets(UWorld* InWorld)
 	}
 
 	// Find all assets being edited
-	FAssetEditorManager& EditorManager = FAssetEditorManager::Get();
-	TArray<UObject*> AllAssets = EditorManager.GetAllEditedAssets();
+	UAssetEditorSubsystem* AssetEditorSubsystem = GEditor->GetEditorSubsystem<UAssetEditorSubsystem>();
+	TArray<UObject*> AllAssets = AssetEditorSubsystem->GetAllEditedAssets();
 
 	TSet<UWorld*> ClosingWorlds;
 
@@ -2396,7 +2398,7 @@ void UEditorEngine::CloseEditedWorldAssets(UWorld* InWorld)
 
 		if (AssetWorld && ClosingWorlds.Contains(AssetWorld))
 		{
-			EditorManager.CloseAllEditorsForAsset(Asset);
+			AssetEditorSubsystem->CloseAllEditorsForAsset(Asset);
 		}
 	}
 }
@@ -3108,7 +3110,7 @@ bool UEditorEngine::HasLockedActors()
 void UEditorEngine::EditObject( UObject* ObjectToEdit )
 {
 	// @todo toolkit minor: Needs world-centric support?
-	FAssetEditorManager::Get().OpenEditorForAsset(ObjectToEdit);
+	GEditor->GetEditorSubsystem<UAssetEditorSubsystem>()->OpenEditorForAsset(ObjectToEdit);
 }
 
 void UEditorEngine::SelectLevelInLevelBrowser( bool bDeselectOthers )
diff --git a/Engine/Source/Editor/UnrealEd/Private/Fbx/ReimportFbxSceneFactory.cpp b/Engine/Source/Editor/UnrealEd/Private/Fbx/ReimportFbxSceneFactory.cpp
--- a/Engine/Source/Editor/UnrealEd/Private/Fbx/ReimportFbxSceneFactory.cpp
+++ b/Engine/Source/Editor/UnrealEd/Private/Fbx/ReimportFbxSceneFactory.cpp
@@ -48,13 +48,14 @@
 #include "Kismet2/BlueprintEditorUtils.h"
 #include "EngineGlobals.h"
 #include "Engine/Engine.h"
-#include "Toolkits/AssetEditorManager.h"
+
 #include "IContentBrowserSingleton.h"
 #include "ContentBrowserModule.h"
 
 #include "ObjectTools.h"
 
 #include "AI/Navigation/NavCollisionBase.h"
+#include "Subsystems/AssetEditorSubsystem.h"
 
 #define LOCTEXT_NAMESPACE "FBXSceneReImportFactory"
 
@@ -925,7 +926,7 @@ UBlueprint *UReimportFbxSceneFactory::UpdateOriginalBluePrint(FString &BluePrint
 		return nullptr;
 	}
 	//Close all editor that edit this blueprint
-	FAssetEditorManager::Get().CloseAllEditorsForAsset(BluePrint);
+	GEditor->GetEditorSubsystem<UAssetEditorSubsystem>()->CloseAllEditorsForAsset(BluePrint);
 	//Set the import status for the next reimport
 	for (TSharedPtr<FFbxNodeInfo> NodeInfo : SceneInfoPtr->HierarchyInfo)
 	{
diff --git a/Engine/Source/Editor/UnrealEd/Private/Kismet2/BlueprintEditorUtils.cpp b/Engine/Source/Editor/UnrealEd/Private/Kismet2/BlueprintEditorUtils.cpp
--- a/Engine/Source/Editor/UnrealEd/Private/Kismet2/BlueprintEditorUtils.cpp
+++ b/Engine/Source/Editor/UnrealEd/Private/Kismet2/BlueprintEditorUtils.cpp
@@ -93,7 +93,7 @@
 #include "ScopedTransaction.h"
 #include "ClassViewerFilter.h"
 #include "InstancedReferenceSubobjectHelper.h"
-#include "Toolkits/AssetEditorManager.h"
+
 #include "BlueprintEditorModule.h"
 #include "BlueprintEditor.h"
 #include "Kismet2/Kismet2NameValidators.h"
@@ -118,6 +118,7 @@
 #include "UObject/UObjectThreadContext.h"
 #include "AnimGraphNode_SubInput.h"
 #include "AnimGraphNode_Root.h"
+#include "Subsystems/AssetEditorSubsystem.h"
 
 extern COREUOBJECT_API bool GBlueprintUseCompilationManager;
 
@@ -2157,7 +2158,7 @@ void FBlueprintEditorUtils::MarkBlueprintAsModified(UBlueprint* Blueprint, FProp
 		FBlueprintEditorUtils::ClearMacroCosmeticInfoCache(Blueprint);
 	}
 	
-	IAssetEditorInstance* AssetEditor = FAssetEditorManager::Get().FindEditorForAsset(Blueprint, false);
+	IAssetEditorInstance* AssetEditor = GEditor->GetEditorSubsystem<UAssetEditorSubsystem>()->FindEditorForAsset(Blueprint, false);
 	if (AssetEditor)
 	{
 		FBlueprintEditor* BlueprintEditor = static_cast<FBlueprintEditor*>(AssetEditor);
@@ -8131,8 +8132,8 @@ void FBlueprintEditorUtils::FindAndSetDebuggableBlueprintInstances()
 	TMap< UBlueprint*, TArray< AActor* > > BlueprintsNeedingInstancesToDebug;
 
 	// Find open blueprint editors that have no debug instances
-	FAssetEditorManager& AssetEditorManager = FAssetEditorManager::Get();	
-	TArray<UObject*> EditedAssets = AssetEditorManager.GetAllEditedAssets();		
+	UAssetEditorSubsystem* AssetEditorSubsystem = GEditor->GetEditorSubsystem<UAssetEditorSubsystem>();
+	TArray<UObject*> EditedAssets = AssetEditorSubsystem->GetAllEditedAssets();		
 	for(int32 i=0; i<EditedAssets.Num(); i++)
 	{
 		UBlueprint* Blueprint = Cast<UBlueprint>( EditedAssets[i] );
diff --git a/Engine/Source/Editor/UnrealEd/Private/Kismet2/ComponentEditorUtils.cpp b/Engine/Source/Editor/UnrealEd/Private/Kismet2/ComponentEditorUtils.cpp
--- a/Engine/Source/Editor/UnrealEd/Private/Kismet2/ComponentEditorUtils.cpp
+++ b/Engine/Source/Editor/UnrealEd/Private/Kismet2/ComponentEditorUtils.cpp
@@ -25,10 +25,11 @@
 #include "Styling/SlateIconFinder.h"
 #include "Framework/MultiBox/MultiBoxBuilder.h"
 #include "Editor.h"
-#include "Toolkits/AssetEditorManager.h"
+
 #include "HAL/PlatformApplicationMisc.h"
 #include "ToolMenus.h"
 #include "Kismet2/ComponentEditorContextMenuContex.h"
+#include "Subsystems/AssetEditorSubsystem.h"
 
 #define LOCTEXT_NAMESPACE "ComponentEditorUtils"
 
@@ -1107,7 +1108,7 @@ void FComponentEditorUtils::OnOpenComponentCodeFile(const FString CodeFileName)
 
 void FComponentEditorUtils::OnEditBlueprintComponent(UObject* Blueprint)
 {
-	FAssetEditorManager::Get().OpenEditorForAsset(Blueprint);
+	GEditor->GetEditorSubsystem<UAssetEditorSubsystem>()->OpenEditorForAsset(Blueprint);
 }
 
 #undef LOCTEXT_NAMESPACE
diff --git a/Engine/Source/Editor/UnrealEd/Private/Kismet2/Kismet2.cpp b/Engine/Source/Editor/UnrealEd/Private/Kismet2/Kismet2.cpp
--- a/Engine/Source/Editor/UnrealEd/Private/Kismet2/Kismet2.cpp
+++ b/Engine/Source/Editor/UnrealEd/Private/Kismet2/Kismet2.cpp
@@ -57,7 +57,7 @@
 #include "Kismet2/KismetDebugUtilities.h"
 #include "Framework/MultiBox/MultiBoxBuilder.h"
 #include "Editor.h"
-#include "Toolkits/AssetEditorManager.h"
+
 #include "Editor/Kismet/Public/BlueprintEditorModule.h"
 #include "Editor/Kismet/Public/FindInBlueprintManager.h"
 #include "Toolkits/ToolkitManager.h"
@@ -83,6 +83,7 @@
 #include "Stats/StatsHierarchical.h"
 #include "Classes/EditorStyleSettings.h"
 #include "ToolMenus.h"
+#include "Subsystems/AssetEditorSubsystem.h"
 
 DECLARE_CYCLE_STAT(TEXT("Compile Blueprint"), EKismetCompilerStats_CompileBlueprint, STATGROUP_KismetCompiler);
 DECLARE_CYCLE_STAT(TEXT("Broadcast Precompile"), EKismetCompilerStats_BroadcastPrecompile, STATGROUP_KismetCompiler);
@@ -324,7 +325,7 @@ void FBlueprintUnloader::UnloadBlueprint(const bool bResetPackage)
 		FKismetEditorUtilities::OnBlueprintUnloaded.Broadcast(UnloadingBp);
 
 		// handled in FBlueprintEditor (from the OnBlueprintUnloaded event)
-// 		IAssetEditorInstance* EditorInst = FAssetEditorManager::Get().FindEditorForAsset(UnloadingBp, /*bFocusIfOpen =*/false);
+// 		IAssetEditorInstance* EditorInst = GEditor->GetEditorSubsystem<UAssetEditorSubsystem>()->FindEditorForAsset(UnloadingBp, /*bFocusIfOpen =*/false);
 // 		if (EditorInst != nullptr)
 // 		{
 // 			EditorInst->CloseWindow();
@@ -718,7 +719,7 @@ UBlueprint* FKismetEditorUtilities::ReloadBlueprint(UBlueprint* StaleBlueprint)
 	FSoftObjectPath BlueprintAssetRef(StaleBlueprint);
 
 	// Need to close the editor now, it won't work once it's been garbage collected during the reload
-	FAssetEditorManager::Get().CloseAllEditorsForAsset(StaleBlueprint);
+	GEditor->GetEditorSubsystem<UAssetEditorSubsystem>()->CloseAllEditorsForAsset(StaleBlueprint);
 
 	FBlueprintUnloader Unloader(StaleBlueprint);
 	Unloader.UnloadBlueprint(/*bResetPackage =*/true);
@@ -743,7 +744,7 @@ UBlueprint* FKismetEditorUtilities::ReplaceBlueprint(UBlueprint* Target, UBluepr
 	check(BlueprintPackage != GetTransientPackage());
 	
 	// Need to close the editor now, it won't work once it's been garbage collected during the reload
-	FAssetEditorManager::Get().CloseAllEditorsForAsset(Target);
+	GEditor->GetEditorSubsystem<UAssetEditorSubsystem>()->CloseAllEditorsForAsset(Target);
 
 	FBlueprintUnloader Unloader(Target);
 	Unloader.UnloadBlueprint(/*bResetPackage =*/false);
@@ -1692,7 +1693,7 @@ UBlueprint* FKismetEditorUtilities::CreateBlueprintFromActor(const FName Bluepri
 	if (NewBlueprint)
 	{
 		// Open the editor for the new blueprint
-		FAssetEditorManager::Get().OpenEditorForAsset(NewBlueprint);
+		GEditor->GetEditorSubsystem<UAssetEditorSubsystem>()->OpenEditorForAsset(NewBlueprint);
 	}
 	return NewBlueprint;
 }
@@ -1869,7 +1870,7 @@ public:
 			}
 
 			// Open the editor for the new blueprint
-			FAssetEditorManager::Get().OpenEditorForAsset(Blueprint);
+			GEditor->GetEditorSubsystem<UAssetEditorSubsystem>()->OpenEditorForAsset(Blueprint);
 
 			return Blueprint;
 		}
@@ -2021,7 +2022,7 @@ UBlueprint* FKismetEditorUtilities::CreateBlueprintUsingAsset(UObject* Asset, bo
 			// Open in BP editor if desired
 			if(bOpenInEditor)
 			{
-				FAssetEditorManager::Get().OpenEditorForAsset(NewBP);
+				GEditor->GetEditorSubsystem<UAssetEditorSubsystem>()->OpenEditorForAsset(NewBP);
 			}
 		}
 
@@ -2073,7 +2074,7 @@ TSharedPtr<class IBlueprintEditor> FKismetEditorUtilities::GetIBlueprintEditorFo
 		if (bOpenEditor)
 		{
 			// @todo toolkit major: Needs world-centric support
-			FAssetEditorManager::Get().OpenEditorForAsset(TargetBP);
+			 GEditor->GetEditorSubsystem<UAssetEditorSubsystem>()->OpenEditorForAsset(TargetBP);
 		}
 
 		TSharedPtr< IToolkit > FoundAssetEditor = FToolkitManager::Get().FindEditorForAsset(TargetBP);
@@ -2157,7 +2158,7 @@ void FKismetEditorUtilities::ShowActorReferencesInLevelScript(const AActor* Acto
 		if (LSB != NULL)
 		{
 			// @todo toolkit major: Needs world-centric support.  Other spots, too?
-			FAssetEditorManager::Get().OpenEditorForAsset(LSB);
+			GEditor->GetEditorSubsystem<UAssetEditorSubsystem>()->OpenEditorForAsset(LSB);
 			TSharedPtr<IToolkit> FoundAssetEditor = FToolkitManager::Get().FindEditorForAsset(LSB);
 			if (FoundAssetEditor.IsValid())
 			{
diff --git a/Engine/Source/Editor/UnrealEd/Private/Kismet2/KismetReinstanceUtilities.cpp b/Engine/Source/Editor/UnrealEd/Private/Kismet2/KismetReinstanceUtilities.cpp
--- a/Engine/Source/Editor/UnrealEd/Private/Kismet2/KismetReinstanceUtilities.cpp
+++ b/Engine/Source/Editor/UnrealEd/Private/Kismet2/KismetReinstanceUtilities.cpp
@@ -24,7 +24,7 @@
 #include "Layers/LayersSubsystem.h"
 #include "Editor.h"
 #include "Serialization/ArchiveHasReferences.h"
-#include "Toolkits/AssetEditorManager.h"
+
 #include "UObject/UObjectHash.h"
 #include "UObject/UObjectIterator.h"
 #include "Serialization/FindObjectReferencers.h"
@@ -32,6 +32,7 @@
 #include "BlueprintEditor.h"
 #include "Engine/Selection.h"
 #include "BlueprintEditorSettings.h"
+#include "Subsystems/AssetEditorSubsystem.h"
 
 extern COREUOBJECT_API bool GBlueprintUseCompilationManager;
 
@@ -2204,7 +2205,7 @@ void FBlueprintCompileReinstancer::ReplaceInstancesOfClass_Inner(TMap<UClass*, U
 		// Refresh any editors for objects that we've updated components for
 		for (UObject* BlueprintAsset : PotentialEditorsForRefreshing)
 		{
-			FBlueprintEditor* BlueprintEditor = static_cast<FBlueprintEditor*>(FAssetEditorManager::Get().FindEditorForAsset(BlueprintAsset, /*bFocusIfOpen =*/false));
+			FBlueprintEditor* BlueprintEditor = static_cast<FBlueprintEditor*>(GEditor->GetEditorSubsystem<UAssetEditorSubsystem>()->FindEditorForAsset(BlueprintAsset, /*bFocusIfOpen =*/false));
 			if (BlueprintEditor)
 			{
 				BlueprintEditor->RefreshEditors();
diff --git a/Engine/Source/Editor/UnrealEd/Private/MiniCurveEditor.cpp b/Engine/Source/Editor/UnrealEd/Private/MiniCurveEditor.cpp
--- a/Engine/Source/Editor/UnrealEd/Private/MiniCurveEditor.cpp
+++ b/Engine/Source/Editor/UnrealEd/Private/MiniCurveEditor.cpp
@@ -5,6 +5,8 @@
 #include "Widgets/Layout/SBorder.h"
 #include "EditorStyleSet.h"
 #include "SCurveEditor.h"
+#include "Subsystems/AssetEditorSubsystem.h"
+#include "Editor.h"
 
 
 void SMiniCurveEditor::Construct(const FArguments& InArgs)
@@ -33,12 +35,12 @@ void SMiniCurveEditor::Construct(const FArguments& InArgs)
 
 	WidgetWindow = InArgs._ParentWindow;
 
-	FAssetEditorManager::Get().NotifyAssetOpened(InArgs._OwnerObject, this);
+	GEditor->GetEditorSubsystem<UAssetEditorSubsystem>()->NotifyAssetOpened(InArgs._OwnerObject, this);
 }
 
 SMiniCurveEditor::~SMiniCurveEditor()
 {
-	FAssetEditorManager::Get().NotifyEditorClosed(this);
+	GEditor->GetEditorSubsystem<UAssetEditorSubsystem>()->NotifyEditorClosed(this);
 }
 
 float SMiniCurveEditor::GetTimelineLength() const
diff --git a/Engine/Source/Editor/UnrealEd/Private/ObjectTools.cpp b/Engine/Source/Editor/UnrealEd/Private/ObjectTools.cpp
--- a/Engine/Source/Editor/UnrealEd/Private/ObjectTools.cpp
+++ b/Engine/Source/Editor/UnrealEd/Private/ObjectTools.cpp
@@ -72,7 +72,7 @@
 #include "ReferencedAssetsUtils.h"
 #include "AssetRegistryModule.h"
 #include "PackagesDialog.h"
-#include "Toolkits/AssetEditorManager.h"
+
 #include "PropertyEditorModule.h"
 #include "Kismet2/KismetEditorUtilities.h"
 #include "Kismet2/KismetReinstanceUtilities.h"
@@ -92,6 +92,7 @@
 #include "Engine/MapBuildDataRegistry.h"
 #include "HAL/PlatformApplicationMisc.h"
 #include "Kismet2/BlueprintEditorUtils.h"
+#include "Subsystems/AssetEditorSubsystem.h"
 
 DEFINE_LOG_CATEGORY_STATIC(LogObjectTools, Log, All);
 
@@ -922,7 +923,7 @@ namespace ObjectTools
 		if ( ObjectToConsolidateTo )
 		{
 			// Close all editors to avoid changing references to temporary objects used by the editor
-			if ( !FAssetEditorManager::Get().CloseAllAssetEditors() )
+			if ( !GEditor->GetEditorSubsystem<UAssetEditorSubsystem>()->CloseAllAssetEditors() )
 			{
 				// Failed to close at least one editor. It is possible that this editor has in-memory object references
 				// which are not prepared to be changed dynamically so it is not safe to continue
@@ -2412,7 +2413,7 @@ namespace ObjectTools
 
 		for (UObject* ObjectToDelete : InObjectsToDelete)
 		{
-			const TArray<IAssetEditorInstance*> ObjectEditors = FAssetEditorManager::Get().FindEditorsForAsset(ObjectToDelete);
+			const TArray<IAssetEditorInstance*> ObjectEditors = GEditor->GetEditorSubsystem<UAssetEditorSubsystem>()->FindEditorsForAsset(ObjectToDelete);
 			for (IAssetEditorInstance* ObjectEditorInstance : ObjectEditors)
 			{
 				if (!ObjectEditorInstance->CloseWindow())
diff --git a/Engine/Source/Editor/UnrealEd/Private/PackageTools.cpp b/Engine/Source/Editor/UnrealEd/Private/PackageTools.cpp
--- a/Engine/Source/Editor/UnrealEd/Private/PackageTools.cpp
+++ b/Engine/Source/Editor/UnrealEd/Private/PackageTools.cpp
@@ -27,7 +27,7 @@
 #include "SourceControlHelpers.h"
 #include "Editor.h"
 #include "Dialogs/Dialogs.h"
-#include "Toolkits/AssetEditorManager.h"
+
 
 #include "ObjectTools.h"
 #include "Kismet2/KismetEditorUtilities.h"
@@ -51,6 +51,7 @@
 #include "ShaderCompiler.h"
 #include "DistanceFieldAtlas.h"
 #include "AssetToolsModule.h"
+#include "Subsystems/AssetEditorSubsystem.h"
 
 #define LOCTEXT_NAMESPACE "PackageTools"
 
@@ -379,7 +380,7 @@ UPackageTools::UPackageTools(const FObjectInitializer& ObjectInitializer)
 				{
 					if (Obj->IsAsset())
 					{
-						FAssetEditorManager::Get().CloseAllEditorsForAsset(Obj);
+						GEditor->GetEditorSubsystem<UAssetEditorSubsystem>()->CloseAllEditorsForAsset(Obj);
 					}
 				}, false);
 
@@ -775,7 +776,7 @@ UPackageTools::UPackageTools(const FObjectInitializer& ObjectInitializer)
 			{
 				TArray<FName> WorldNamesToReload;
 				WorldNamesToReload.Add(WorldNameToReload);
-				FAssetEditorManager::Get().OpenEditorsForAssets(WorldNamesToReload);
+				GEditor->GetEditorSubsystem<UAssetEditorSubsystem>()->OpenEditorsForAssets(WorldNamesToReload);
 			}
 			else if (UGameEngine* GameEngine = Cast<UGameEngine>(GEngine))
 			{
diff --git a/Engine/Source/Editor/UnrealEd/Private/PlayLevel.cpp b/Engine/Source/Editor/UnrealEd/Private/PlayLevel.cpp
--- a/Engine/Source/Editor/UnrealEd/Private/PlayLevel.cpp
+++ b/Engine/Source/Editor/UnrealEd/Private/PlayLevel.cpp
@@ -68,7 +68,7 @@
 #include "Slate/SceneViewport.h"
 #include "Kismet2/KismetEditorUtilities.h"
 #include "Kismet2/BlueprintEditorUtils.h"
-#include "Toolkits/AssetEditorManager.h"
+
 #include "LevelEditor.h"
 #include "ILevelViewport.h"
 #include "BlueprintEditorModule.h"
@@ -109,6 +109,7 @@
 #include "CookerSettings.h"
 #include "Widgets/Text/STextBlock.h"
 #include "Widgets/SBoxPanel.h"
+#include "Subsystems/AssetEditorSubsystem.h"
 
 
 DEFINE_LOG_CATEGORY_STATIC(LogPlayLevel, Log, All);
@@ -408,7 +409,7 @@ void UEditorEngine::EndPlayMap()
 				Object->ClearFlags(RF_Standalone);
 			}
 			// Close any asset editors that are currently editing this object
-			FAssetEditorManager::Get().CloseAllEditorsForAsset(Object);
+			GEditor->GetEditorSubsystem<UAssetEditorSubsystem>()->CloseAllEditorsForAsset(Object);
 		}
 	}
 
@@ -731,13 +732,13 @@ void UEditorEngine::TeardownPlaySession(FWorldContext& PieWorldContext)
 
 	// Construct a list of editors that are active for objects being debugged. We will refresh these when we have cleaned up to ensure no invalid objects exist in them
 	TArray< IBlueprintEditor* > Editors;
-	FAssetEditorManager& AssetEditorManager = FAssetEditorManager::Get();
+	UAssetEditorSubsystem* AssetEditorSubsystem = GEditor->GetEditorSubsystem<UAssetEditorSubsystem>();
 	const UWorld::FBlueprintToDebuggedObjectMap& EditDebugObjectsPre = PlayWorld->GetBlueprintObjectsBeingDebugged();
 	for (UWorld::FBlueprintToDebuggedObjectMap::TConstIterator EditIt(EditDebugObjectsPre); EditIt; ++EditIt)
 	{
 		if (UBlueprint* TargetBP = EditIt.Key().Get())
 		{
-			if(IBlueprintEditor* EachEditor = static_cast<IBlueprintEditor*>(AssetEditorManager.FindEditorForAsset(TargetBP, false)))
+			if(IBlueprintEditor* EachEditor = static_cast<IBlueprintEditor*>(AssetEditorSubsystem->FindEditorForAsset(TargetBP, false)))
 			{
 				Editors.AddUnique( EachEditor );
 			}
diff --git a/Engine/Source/Editor/UnrealEd/Private/StaticLightingSystem/StaticLightingSystem.cpp b/Engine/Source/Editor/UnrealEd/Private/StaticLightingSystem/StaticLightingSystem.cpp
--- a/Engine/Source/Editor/UnrealEd/Private/StaticLightingSystem/StaticLightingSystem.cpp
+++ b/Engine/Source/Editor/UnrealEd/Private/StaticLightingSystem/StaticLightingSystem.cpp
@@ -68,7 +68,7 @@ FSwarmDebugOptions GSwarmDebugOptions;
 DEFINE_LOG_CATEGORY(LogStaticLightingSystem);
 
 #include "EngineGlobals.h"
-#include "Toolkits/AssetEditorManager.h"
+
 
 #include "Lightmass/LightmassImportanceVolume.h"
 #include "Components/LightmassPortalComponent.h"
@@ -79,6 +79,7 @@ DEFINE_LOG_CATEGORY(LogStaticLightingSystem);
 #include "Framework/Notifications/NotificationManager.h"
 #include "Widgets/Notifications/SNotificationList.h"
 #include "Misc/UObjectToken.h"
+#include "Subsystems/AssetEditorSubsystem.h"
 
 #define LOCTEXT_NAMESPACE "StaticLightingSystem"
 
@@ -2425,8 +2426,8 @@ bool FStaticLightingSystem::IsAmortizedExporting() const
 void UEditorEngine::BuildLighting(const FLightingBuildOptions& Options)
 {
 	// Forcibly shut down all texture property windows as they become invalid during a light build
-	FAssetEditorManager& AssetEditorManager = FAssetEditorManager::Get();
-	TArray<UObject*> EditedAssets = AssetEditorManager.GetAllEditedAssets();
+	UAssetEditorSubsystem* AssetEditorSubsystem = GEditor->GetEditorSubsystem<UAssetEditorSubsystem>();
+	TArray<UObject*> EditedAssets = AssetEditorSubsystem->GetAllEditedAssets();
 
 	for (int32 AssetIdx = 0; AssetIdx < EditedAssets.Num(); AssetIdx++)
 	{
@@ -2434,7 +2435,7 @@ void UEditorEngine::BuildLighting(const FLightingBuildOptions& Options)
 
 		if (EditedAsset->IsA(UTexture2D::StaticClass()))
 		{
-			IAssetEditorInstance* Editor = AssetEditorManager.FindEditorForAsset(EditedAsset, false);
+			IAssetEditorInstance* Editor = AssetEditorSubsystem->FindEditorForAsset(EditedAsset, false);
 			if (Editor)
 			{
 				Editor->CloseWindow();
