Fix crash when running commandlets that load level packages.

Blueprint search manager was creating missing level blueprints, compiling them then prompting to check out and save.

diff --git a/Engine/Source/Editor/Kismet/Private/BlueprintAssetHandler.cpp b/Engine/Source/Editor/Kismet/Private/BlueprintAssetHandler.cpp
--- a/Engine/Source/Editor/Kismet/Private/BlueprintAssetHandler.cpp
+++ b/Engine/Source/Editor/Kismet/Private/BlueprintAssetHandler.cpp
@@ -11,8 +11,8 @@ class FLevelBlueprintAssetHandler : public IBlueprintAssetHandler
 	{
 		UWorld* World = CastChecked<UWorld>(InObject);
 
-		const bool bCreateLevelScript = false;
-		return World->PersistentLevel ? World->PersistentLevel->GetLevelScriptBlueprint(bCreateLevelScript) : nullptr;
+		const bool bDontCreate = true;
+		return World->PersistentLevel ? World->PersistentLevel->GetLevelScriptBlueprint(bDontCreate) : nullptr;
 	}
 
 	virtual bool AssetContainsBlueprint(const FAssetData& InAssetData) const
