Converted toolbars for Blueprint and Play buttons

diff --git a/Engine/Plugins/Experimental/ControlRig/Source/ControlRigEditor/Private/Editor/ControlRigEditorMode.cpp b/Engine/Plugins/Experimental/ControlRig/Source/ControlRigEditor/Private/Editor/ControlRigEditorMode.cpp
--- a/Engine/Plugins/Experimental/ControlRig/Source/ControlRigEditor/Private/Editor/ControlRigEditorMode.cpp
+++ b/Engine/Plugins/Experimental/ControlRig/Source/ControlRigEditor/Private/Editor/ControlRigEditorMode.cpp
@@ -9,6 +9,7 @@
 #include "RigHierarchyTabSummoner.h"
 #include "RigStackTabSummoner.h"
 #include "RigCurveContainerTabSummoner.h"
+#include "ToolMenus.h"
 
 FControlRigEditorMode::FControlRigEditorMode(const TSharedRef<FControlRigEditor>& InControlRigEditor)
 	: FBlueprintEditorApplicationMode(InControlRigEditor, FControlRigEditorModes::ControlRigEditorMode, FControlRigEditorModes::GetLocalizedMode, false, false)
@@ -116,12 +117,12 @@ FControlRigEditorMode::FControlRigEditorMode(const TSharedRef<FControlRigEditor>
 			)
 		);
 
-	// setup toolbar - clear existing toolbar extender from the BP mode
-	ToolbarExtender = MakeShareable(new FExtender);
-	InControlRigEditor->GetToolbarBuilder()->AddCompileToolbar(ToolbarExtender);
-	InControlRigEditor->GetToolbarBuilder()->AddScriptingToolbar(ToolbarExtender);
-	InControlRigEditor->GetToolbarBuilder()->AddBlueprintGlobalOptionsToolbar(ToolbarExtender);
-//	InControlRigEditor->GetToolbarBuilder()->AddDebuggingToolbar(ToolbarExtender);
+	if (UToolMenu* Toolbar = InControlRigEditor->RegisterModeToolbarIfUnregistered(GetModeName()))
+	{
+		InControlRigEditor->GetToolbarBuilder()->AddCompileToolbar(Toolbar);
+		InControlRigEditor->GetToolbarBuilder()->AddScriptingToolbar(Toolbar);
+		InControlRigEditor->GetToolbarBuilder()->AddBlueprintGlobalOptionsToolbar(Toolbar);
+	}
 }
 
 void FControlRigEditorMode::RegisterTabFactories(TSharedPtr<FTabManager> InTabManager)
diff --git a/Engine/Source/Editor/AnimationBlueprintEditor/AnimationBlueprintEditor.Build.cs b/Engine/Source/Editor/AnimationBlueprintEditor/AnimationBlueprintEditor.Build.cs
--- a/Engine/Source/Editor/AnimationBlueprintEditor/AnimationBlueprintEditor.Build.cs
+++ b/Engine/Source/Editor/AnimationBlueprintEditor/AnimationBlueprintEditor.Build.cs
@@ -39,6 +39,7 @@ public class AnimationBlueprintEditor : ModuleRules
                 "BlueprintGraph",
                 "RHI",
                 "KismetCompiler",
+				"ToolMenus",
             }
 		);
 
diff --git a/Engine/Source/Editor/AnimationBlueprintEditor/Private/AnimationBlueprintEditorMode.cpp b/Engine/Source/Editor/AnimationBlueprintEditor/Private/AnimationBlueprintEditorMode.cpp
--- a/Engine/Source/Editor/AnimationBlueprintEditor/Private/AnimationBlueprintEditorMode.cpp
+++ b/Engine/Source/Editor/AnimationBlueprintEditor/Private/AnimationBlueprintEditorMode.cpp
@@ -14,6 +14,8 @@
 #include "IPersonaPreviewScene.h"
 #include "PersonaUtils.h"
 
+#include "ToolMenus.h"
+
 /////////////////////////////////////////////////////
 // FAnimationBlueprintEditorMode
 
@@ -141,10 +143,14 @@ FAnimationBlueprintEditorMode::FAnimationBlueprintEditorMode(const TSharedRef<FA
 	// setup toolbar - clear existing toolbar extender from the BP mode
 	//@TODO: Keep this in sync with BlueprintEditorModes.cpp
 	ToolbarExtender = MakeShareable(new FExtender);
-	InAnimationBlueprintEditor->GetToolbarBuilder()->AddCompileToolbar(ToolbarExtender);
-	InAnimationBlueprintEditor->GetToolbarBuilder()->AddScriptingToolbar(ToolbarExtender);
-	InAnimationBlueprintEditor->GetToolbarBuilder()->AddBlueprintGlobalOptionsToolbar(ToolbarExtender);
-	InAnimationBlueprintEditor->GetToolbarBuilder()->AddDebuggingToolbar(ToolbarExtender);
+
+	if (UToolMenu* Toolbar = InAnimationBlueprintEditor->RegisterModeToolbarIfUnregistered(GetModeName()))
+	{
+		InAnimationBlueprintEditor->GetToolbarBuilder()->AddCompileToolbar(Toolbar);
+		InAnimationBlueprintEditor->GetToolbarBuilder()->AddScriptingToolbar(Toolbar);
+		InAnimationBlueprintEditor->GetToolbarBuilder()->AddBlueprintGlobalOptionsToolbar(Toolbar);
+		InAnimationBlueprintEditor->GetToolbarBuilder()->AddDebuggingToolbar(Toolbar);
+	}
 
 	PersonaModule.OnRegisterTabs().Broadcast(TabFactories, InAnimationBlueprintEditor);
 	LayoutExtender = MakeShared<FLayoutExtender>();
diff --git a/Engine/Source/Editor/AnimationBlueprintEditor/Private/AnimationBlueprintInterfaceEditorMode.cpp b/Engine/Source/Editor/AnimationBlueprintEditor/Private/AnimationBlueprintInterfaceEditorMode.cpp
--- a/Engine/Source/Editor/AnimationBlueprintEditor/Private/AnimationBlueprintInterfaceEditorMode.cpp
+++ b/Engine/Source/Editor/AnimationBlueprintEditor/Private/AnimationBlueprintInterfaceEditorMode.cpp
@@ -9,6 +9,7 @@
 #include "PersonaModule.h"
 #include "SBlueprintEditorToolbar.h"
 #include "IPersonaPreviewScene.h"
+#include "ToolMenus.h"
 
 FAnimationBlueprintInterfaceEditorMode::FAnimationBlueprintInterfaceEditorMode(const TSharedRef<FAnimationBlueprintEditor>& InAnimationBlueprintEditor)
 	: FBlueprintInterfaceApplicationMode(InAnimationBlueprintEditor, FAnimationBlueprintEditorModes::AnimationBlueprintInterfaceEditorMode, FAnimationBlueprintEditorModes::GetLocalizedMode)
@@ -83,8 +84,12 @@ FAnimationBlueprintInterfaceEditorMode::FAnimationBlueprintInterfaceEditorMode(c
 	FPersonaModule& PersonaModule = FModuleManager::LoadModuleChecked<FPersonaModule>("Persona");
 
 	ToolbarExtender = MakeShareable(new FExtender);
-	InAnimationBlueprintEditor->GetToolbarBuilder()->AddCompileToolbar(ToolbarExtender);
-	InAnimationBlueprintEditor->GetToolbarBuilder()->AddBlueprintGlobalOptionsToolbar(ToolbarExtender);
+
+	if (UToolMenu* Toolbar = InAnimationBlueprintEditor->RegisterModeToolbarIfUnregistered(GetModeName()))
+	{
+		InAnimationBlueprintEditor->GetToolbarBuilder()->AddCompileToolbar(Toolbar);
+		InAnimationBlueprintEditor->GetToolbarBuilder()->AddBlueprintGlobalOptionsToolbar(Toolbar);	
+	}
 
 	PersonaModule.OnRegisterTabs().Broadcast(TabFactories, InAnimationBlueprintEditor);
 	LayoutExtender = MakeShared<FLayoutExtender>();
diff --git a/Engine/Source/Editor/Kismet/Private/BlueprintEditor.cpp b/Engine/Source/Editor/Kismet/Private/BlueprintEditor.cpp
--- a/Engine/Source/Editor/Kismet/Private/BlueprintEditor.cpp
+++ b/Engine/Source/Editor/Kismet/Private/BlueprintEditor.cpp
@@ -99,6 +99,8 @@
 
 #include "BlueprintEditorTabs.h"
 
+#include "ToolMenus.h"
+#include "BlueprintEditorContext.h"
 
 #include "Interfaces/IProjectManager.h"
 
@@ -1895,6 +1897,27 @@ void FBlueprintEditor::InitBlueprintEditor(
 	}
 }
 
+void FBlueprintEditor::InitToolMenuContext(FToolMenuContext& MenuContext)
+{
+	FAssetEditorToolkit::InitToolMenuContext(MenuContext);
+
+	UBlueprintEditorToolMenuContext* Context = NewObject<UBlueprintEditorToolMenuContext>();
+	Context->BlueprintEditor = SharedThis(this);
+	MenuContext.AddObject(Context);
+}
+
+UToolMenu* FBlueprintEditor::RegisterModeToolbarIfUnregistered(const FName InModeName)
+{
+	FName ParentToolbarName;
+	const FName ModeSpecificToolBarName = GetToolMenuToolbarNameForMode(InModeName, ParentToolbarName);
+	if (!UToolMenus::Get()->IsMenuRegistered(ModeSpecificToolBarName))
+	{
+		return UToolMenus::Get()->RegisterMenu(ModeSpecificToolBarName, ParentToolbarName, EMultiBoxType::ToolBar);
+	}
+
+	return nullptr;
+}
+
 void FBlueprintEditor::InitalizeExtenders()
 {
 	TSharedPtr<FExtender> MenuExtender = MakeShareable(new FExtender);
@@ -9387,4 +9410,9 @@ void FBlueprintEditor::ClearAllGraphEditorQuickJumps()
 /////////////////////////////////////////////////////
 /////////////////////////////////////////////////////
 
+UBlueprint* UBlueprintEditorToolMenuContext::GetBlueprintObj() const
+{
+	return BlueprintEditor.IsValid() ? BlueprintEditor.Pin()->GetBlueprintObj() : nullptr;
+}
+
 #undef LOCTEXT_NAMESPACE
diff --git a/Engine/Source/Editor/Kismet/Private/BlueprintEditorModes.cpp b/Engine/Source/Editor/Kismet/Private/BlueprintEditorModes.cpp
--- a/Engine/Source/Editor/Kismet/Private/BlueprintEditorModes.cpp
+++ b/Engine/Source/Editor/Kismet/Private/BlueprintEditorModes.cpp
@@ -2,7 +2,7 @@
 
 #include "BlueprintEditorModes.h"
 #include "Settings/EditorExperimentalSettings.h"
-
+#include "ToolMenus.h"
 
 // Core kismet tabs
 #include "SSCSEditor.h"
@@ -117,11 +117,13 @@ FBlueprintEditorApplicationMode::FBlueprintEditorApplicationMode(TSharedPtr<clas
 	
 	// setup toolbar
 	//@TODO: Keep this in sync with AnimBlueprintMode.cpp
-	InBlueprintEditor->GetToolbarBuilder()->AddBlueprintEditorModesToolbar(ToolbarExtender);
-	InBlueprintEditor->GetToolbarBuilder()->AddCompileToolbar(ToolbarExtender);
-	InBlueprintEditor->GetToolbarBuilder()->AddScriptingToolbar(ToolbarExtender);
-	InBlueprintEditor->GetToolbarBuilder()->AddBlueprintGlobalOptionsToolbar(ToolbarExtender);
-	InBlueprintEditor->GetToolbarBuilder()->AddDebuggingToolbar(ToolbarExtender);
+	if (UToolMenu* Toolbar = InBlueprintEditor->RegisterModeToolbarIfUnregistered(GetModeName()))
+	{
+		InBlueprintEditor->GetToolbarBuilder()->AddCompileToolbar(Toolbar);
+		InBlueprintEditor->GetToolbarBuilder()->AddScriptingToolbar(Toolbar);
+		InBlueprintEditor->GetToolbarBuilder()->AddBlueprintGlobalOptionsToolbar(Toolbar);
+		InBlueprintEditor->GetToolbarBuilder()->AddDebuggingToolbar(Toolbar);
+	}
 
 	FBlueprintEditorModule& BlueprintEditorModule = FModuleManager::LoadModuleChecked<FBlueprintEditorModule>("Kismet");
 	BlueprintEditorModule.OnRegisterTabsForEditor().Broadcast(BlueprintEditorTabFactories, InModeName, InBlueprintEditor);
@@ -194,8 +196,10 @@ FBlueprintDefaultsApplicationMode::FBlueprintDefaultsApplicationMode(TSharedPtr<
 		);
 
 	// setup toolbar
-	InBlueprintEditor->GetToolbarBuilder()->AddCompileToolbar(ToolbarExtender);
-	InBlueprintEditor->GetToolbarBuilder()->AddBlueprintEditorModesToolbar(ToolbarExtender);
+	if (UToolMenu* Toolbar = InBlueprintEditor->RegisterModeToolbarIfUnregistered(GetModeName()))
+	{
+		InBlueprintEditor->GetToolbarBuilder()->AddCompileToolbar(Toolbar);
+	}
 }
 
 void FBlueprintDefaultsApplicationMode::RegisterTabFactories(TSharedPtr<FTabManager> InTabManager)
@@ -276,10 +280,12 @@ FBlueprintComponentsApplicationMode::FBlueprintComponentsApplicationMode(TShared
 		);
 
 	// setup toolbar
-	InBlueprintEditor->GetToolbarBuilder()->AddBlueprintEditorModesToolbar(ToolbarExtender);
-	InBlueprintEditor->GetToolbarBuilder()->AddBlueprintGlobalOptionsToolbar(ToolbarExtender);
-	InBlueprintEditor->GetToolbarBuilder()->AddCompileToolbar(ToolbarExtender);
-	InBlueprintEditor->GetToolbarBuilder()->AddComponentsToolbar(ToolbarExtender);
+	if (UToolMenu* Toolbar = InBlueprintEditor->RegisterModeToolbarIfUnregistered(GetModeName()))
+	{
+		InBlueprintEditor->GetToolbarBuilder()->AddBlueprintGlobalOptionsToolbar(Toolbar);
+		InBlueprintEditor->GetToolbarBuilder()->AddCompileToolbar(Toolbar);
+		InBlueprintEditor->GetToolbarBuilder()->AddComponentsToolbar(Toolbar);
+	}
 }
 
 void FBlueprintComponentsApplicationMode::RegisterTabFactories(TSharedPtr<FTabManager> InTabManager)
@@ -416,8 +422,11 @@ FBlueprintInterfaceApplicationMode::FBlueprintInterfaceApplicationMode(TSharedPt
 		);
 
 	// setup toolbar
-	InBlueprintEditor->GetToolbarBuilder()->AddCompileToolbar(ToolbarExtender);
-	InBlueprintEditor->GetToolbarBuilder()->AddBlueprintGlobalOptionsToolbar(ToolbarExtender);
+	if (UToolMenu* Toolbar = InBlueprintEditor->RegisterModeToolbarIfUnregistered(GetModeName()))
+	{
+		InBlueprintEditor->GetToolbarBuilder()->AddCompileToolbar(Toolbar);
+		InBlueprintEditor->GetToolbarBuilder()->AddBlueprintGlobalOptionsToolbar(Toolbar);
+	}
 }
 
 void FBlueprintInterfaceApplicationMode::RegisterTabFactories(TSharedPtr<FTabManager> InTabManager)
@@ -521,12 +530,14 @@ FBlueprintMacroApplicationMode::FBlueprintMacroApplicationMode(TSharedPtr<class
 		);
 
 	// setup toolbar
-	InBlueprintEditor->GetToolbarBuilder()->AddCompileToolbar(ToolbarExtender);
-	InBlueprintEditor->GetToolbarBuilder()->AddScriptingToolbar(ToolbarExtender);
-	InBlueprintEditor->GetToolbarBuilder()->AddBlueprintGlobalOptionsToolbar(ToolbarExtender);
-	InBlueprintEditor->GetToolbarBuilder()->AddDebuggingToolbar(ToolbarExtender);
+	if (UToolMenu* Toolbar = InBlueprintEditor->RegisterModeToolbarIfUnregistered(GetModeName()))
+	{
+		InBlueprintEditor->GetToolbarBuilder()->AddCompileToolbar(Toolbar);
+		InBlueprintEditor->GetToolbarBuilder()->AddScriptingToolbar(Toolbar);
+		InBlueprintEditor->GetToolbarBuilder()->AddBlueprintGlobalOptionsToolbar(Toolbar);
+		InBlueprintEditor->GetToolbarBuilder()->AddDebuggingToolbar(Toolbar);
+	}
 }
-
 void FBlueprintMacroApplicationMode::RegisterTabFactories(TSharedPtr<FTabManager> InTabManager)
 {
 	TSharedPtr<FBlueprintEditor> BP = MyBlueprintEditor.Pin();
@@ -716,17 +727,19 @@ FBlueprintEditorUnifiedMode::FBlueprintEditorUnifiedMode(TSharedPtr<class FBluep
 	
 	// setup toolbar
 	//@TODO: Keep this in sync with AnimBlueprintMode.cpp
-	//InBlueprintEditor->GetToolbarBuilder()->AddBlueprintEditorModesToolbar(ToolbarExtender);
-	InBlueprintEditor->GetToolbarBuilder()->AddCompileToolbar(ToolbarExtender);
-	InBlueprintEditor->GetToolbarBuilder()->AddScriptingToolbar(ToolbarExtender);
-	InBlueprintEditor->GetToolbarBuilder()->AddBlueprintGlobalOptionsToolbar(ToolbarExtender);
-	
-	if ( bRegisterViewport )
+	if (UToolMenu* Toolbar = InBlueprintEditor->RegisterModeToolbarIfUnregistered(GetModeName()))
 	{
-		InBlueprintEditor->GetToolbarBuilder()->AddComponentsToolbar(ToolbarExtender);
+		InBlueprintEditor->GetToolbarBuilder()->AddCompileToolbar(Toolbar);
+		InBlueprintEditor->GetToolbarBuilder()->AddScriptingToolbar(Toolbar);
+		InBlueprintEditor->GetToolbarBuilder()->AddBlueprintGlobalOptionsToolbar(Toolbar);
+
+		if ( bRegisterViewport )
+		{
+			InBlueprintEditor->GetToolbarBuilder()->AddComponentsToolbar(Toolbar);
+		}
+
+		InBlueprintEditor->GetToolbarBuilder()->AddDebuggingToolbar(Toolbar);
 	}
-	
-	InBlueprintEditor->GetToolbarBuilder()->AddDebuggingToolbar(ToolbarExtender);
 
 	FBlueprintEditorModule& BlueprintEditorModule = FModuleManager::LoadModuleChecked<FBlueprintEditorModule>("Kismet");
 	BlueprintEditorModule.OnRegisterTabsForEditor().Broadcast(BlueprintEditorTabFactories, InModeName, InBlueprintEditor);
diff --git a/Engine/Source/Editor/Kismet/Private/Debugging/SKismetDebuggingView.cpp b/Engine/Source/Editor/Kismet/Private/Debugging/SKismetDebuggingView.cpp
--- a/Engine/Source/Editor/Kismet/Private/Debugging/SKismetDebuggingView.cpp
+++ b/Engine/Source/Editor/Kismet/Private/Debugging/SKismetDebuggingView.cpp
@@ -23,6 +23,7 @@
 #include "Kismet2/DebuggerCommands.h"
 #include "Debugging/KismetDebugCommands.h"
 #include "Widgets/Input/SHyperlink.h"
+#include "ToolMenus.h"
 
 #define LOCTEXT_NAMESPACE "DebugViewUI"
 
@@ -1062,8 +1063,19 @@ void SKismetDebuggingView::Construct(const FArguments& InArgs)
 	BlueprintToWatchPtr = InArgs._BlueprintToWatch;
 
 	// Build the debug toolbar
-	FToolBarBuilder DebugToolbarBuilder(FPlayWorldCommands::GlobalPlayWorldActions, FMultiBoxCustomization::None );
-	FPlayWorldCommands::BuildToolbar(DebugToolbarBuilder);
+	static const FName ToolbarName = "Kismet.DebuggingViewToolBar";
+	if (!UToolMenus::Get()->IsMenuRegistered(ToolbarName))
+	{
+		UToolMenu* ToolBar = UToolMenus::Get()->RegisterMenu(ToolbarName, NAME_None, EMultiBoxType::ToolBar);
+
+		{
+			FToolMenuSection& Section = ToolBar->AddSection("Debug");
+			FPlayWorldCommands::BuildToolbar(Section);
+		}
+	}
+
+	FToolMenuContext MenuContext(FPlayWorldCommands::GlobalPlayWorldActions);
+	TSharedRef<SWidget> ToolbarWidget = UToolMenus::Get()->GenerateWidget(ToolbarName, MenuContext);
 
 	this->ChildSlot
 	[
@@ -1075,7 +1087,7 @@ void SKismetDebuggingView::Construct(const FArguments& InArgs)
 			.Visibility( this, &SKismetDebuggingView::IsDebuggerVisible )
 			.BorderImage( FEditorStyle::GetBrush( TEXT("NoBorder") ) )
 			[
-				DebugToolbarBuilder.MakeWidget()
+				ToolbarWidget
 			]
 		]
 		+SVerticalBox::Slot()
diff --git a/Engine/Source/Editor/Kismet/Private/SBlueprintEditorToolbar.cpp b/Engine/Source/Editor/Kismet/Private/SBlueprintEditorToolbar.cpp
--- a/Engine/Source/Editor/Kismet/Private/SBlueprintEditorToolbar.cpp
+++ b/Engine/Source/Editor/Kismet/Private/SBlueprintEditorToolbar.cpp
@@ -30,6 +30,8 @@
 #include "SBlueprintRevisionMenu.h"
 #include "ISourceCodeAccessor.h"
 #include "ISourceCodeAccessModule.h"
+#include "ToolMenus.h"
+#include "BlueprintEditorContext.h"
 
 #define LOCTEXT_NAMESPACE "KismetToolbar"
 
@@ -397,222 +399,149 @@ static void BlueprintEditorToolbarImpl::MakeCompileDeveloperSubMenu(FMenuBuilder
 //////////////////////////////////////////////////////////////////////////
 // FBlueprintEditorToolbar
 
-void FBlueprintEditorToolbar::AddBlueprintEditorModesToolbar(TSharedPtr<FExtender> Extender)
+void FBlueprintEditorToolbar::AddBlueprintGlobalOptionsToolbar(UToolMenu* InMenu)
 {
-	
-}
-
-void FBlueprintEditorToolbar::AddBlueprintGlobalOptionsToolbar(TSharedPtr<FExtender> Extender)
-{
-	TSharedPtr<FBlueprintEditor> BlueprintEditorPtr = BlueprintEditor.Pin();
-
-	Extender->AddToolBarExtension(
-		"Asset",
-		EExtensionHook::After,
-		BlueprintEditorPtr->GetToolkitCommands(),
-		FToolBarExtensionDelegate::CreateSP( this, &FBlueprintEditorToolbar::FillBlueprintGlobalOptionsToolbar ) );
-}
+	FToolMenuSection& Section = InMenu->AddSection("Settings");
+	Section.InsertPosition = FToolMenuInsert("Asset", EToolMenuInsertType::Before);
 
-void FBlueprintEditorToolbar::AddCompileToolbar(TSharedPtr<FExtender> Extender)
-{
-	TSharedPtr<FBlueprintEditor> BlueprintEditorPtr = BlueprintEditor.Pin();
-
-	Extender->AddToolBarExtension(
-		"Asset",
-		EExtensionHook::Before,
-		BlueprintEditorPtr->GetToolkitCommands(),
-		FToolBarExtensionDelegate::CreateSP( this, &FBlueprintEditorToolbar::FillCompileToolbar ) );
-}
-
-void FBlueprintEditorToolbar::AddNewToolbar(TSharedPtr<FExtender> Extender)
-{
-	TSharedPtr<FBlueprintEditor> BlueprintEditorPtr = BlueprintEditor.Pin();
-
-	Extender->AddToolBarExtension(
-		"MyBlueprint",
-		EExtensionHook::After,
-		BlueprintEditorPtr->GetToolkitCommands(),
-		FToolBarExtensionDelegate::CreateSP( this, &FBlueprintEditorToolbar::FillNewToolbar ) );
-}
-
-void FBlueprintEditorToolbar::AddScriptingToolbar(TSharedPtr<FExtender> Extender)
-{
-	TSharedPtr<FBlueprintEditor> BlueprintEditorPtr = BlueprintEditor.Pin();
-
-	Extender->AddToolBarExtension(
-		"Asset",
-		EExtensionHook::After,
-		BlueprintEditorPtr->GetToolkitCommands(),
-		FToolBarExtensionDelegate::CreateSP( this, &FBlueprintEditorToolbar::FillScriptingToolbar ) );
-}
-
-void FBlueprintEditorToolbar::AddDebuggingToolbar(TSharedPtr<FExtender> Extender)
-{
-	TSharedPtr<FBlueprintEditor> BlueprintEditorPtr = BlueprintEditor.Pin();
-
-	Extender->AddToolBarExtension(
-		"Asset",
-		EExtensionHook::After,
-		BlueprintEditorPtr->GetToolkitCommands(),
-		FToolBarExtensionDelegate::CreateSP( this, &FBlueprintEditorToolbar::FillDebuggingToolbar ) );
-}
-
-void FBlueprintEditorToolbar::AddComponentsToolbar(TSharedPtr<FExtender> Extender)
-{
-	TSharedPtr<FBlueprintEditor> BlueprintEditorPtr = BlueprintEditor.Pin();
-
-	Extender->AddToolBarExtension(
-		"Asset",
-		EExtensionHook::After,
-		BlueprintEditorPtr->GetToolkitCommands(),
-		FToolBarExtensionDelegate::CreateSP( this, &FBlueprintEditorToolbar::FillComponentsToolbar ) );
-}
-
-void FBlueprintEditorToolbar::FillBlueprintEditorModesToolbar(FToolBarBuilder& ToolbarBuilder)
-{
-	TSharedPtr<FBlueprintEditor> BlueprintEditorPtr = BlueprintEditor.Pin();
-	UBlueprint* BlueprintObj = BlueprintEditorPtr->GetBlueprintObj();
-
-	TAttribute<FName> GetActiveMode(BlueprintEditorPtr.ToSharedRef(), &FBlueprintEditor::GetCurrentMode);
-	FOnModeChangeRequested SetActiveMode = FOnModeChangeRequested::CreateSP(BlueprintEditorPtr.ToSharedRef(), &FBlueprintEditor::SetCurrentMode);
-
-	TArray< TSharedPtr< SWidget > > ToolbarWidgets = GenerateToolbarWidgets( BlueprintObj, GetActiveMode, SetActiveMode );
-	
-	for( const auto& Widget : ToolbarWidgets )
+	Section.AddDynamicEntry("BlueprintGlobalOptions", FNewToolMenuSectionDelegate::CreateLambda([](FToolMenuSection& InSection)
 	{
-		BlueprintEditorPtr->AddToolbarWidget( Widget.ToSharedRef() );
-	}
+		UBlueprintEditorToolMenuContext* Context = InSection.FindContext<UBlueprintEditorToolMenuContext>();
+		if (Context && Context->GetBlueprintObj())
+		{
+			const FFullBlueprintEditorCommands& Commands = FFullBlueprintEditorCommands::Get();
+			InSection.AddEntry(FToolMenuEntry::InitToolBarButton(Commands.EditGlobalOptions));
+			InSection.AddEntry(FToolMenuEntry::InitToolBarButton(Commands.EditClassDefaults));
+		}
+	}));
 }
 
-void FBlueprintEditorToolbar::FillBlueprintGlobalOptionsToolbar(FToolBarBuilder& ToolbarBuilder)
+void FBlueprintEditorToolbar::AddCompileToolbar(UToolMenu* InMenu)
 {
-	const FFullBlueprintEditorCommands& Commands = FFullBlueprintEditorCommands::Get();
-	TSharedPtr<FBlueprintEditor> BlueprintEditorPtr = BlueprintEditor.Pin();
-	UBlueprint* BlueprintObj = BlueprintEditorPtr->GetBlueprintObj();
-	
-	ToolbarBuilder.BeginSection("Settings");
-	
-	if(BlueprintObj != NULL)
+	FToolMenuSection& Section = InMenu->AddSection("Compile");
+	Section.InsertPosition = FToolMenuInsert("Asset", EToolMenuInsertType::Before);
+
+	Section.AddDynamicEntry("CompileCommands", FNewToolMenuSectionDelegate::CreateLambda([](FToolMenuSection& InSection)
 	{
-		ToolbarBuilder.AddToolBarButton(Commands.EditGlobalOptions);
-		ToolbarBuilder.AddToolBarButton(Commands.EditClassDefaults);
-	}
-	
-	ToolbarBuilder.EndSection();
+		UBlueprintEditorToolMenuContext* Context = InSection.FindContext<UBlueprintEditorToolMenuContext>();
+		if (Context && Context->BlueprintEditor.IsValid() && Context->GetBlueprintObj())
+		{
+			TSharedPtr<class FBlueprintEditorToolbar> BlueprintEditorToolbar = Context->BlueprintEditor.Pin()->GetToolbarBuilder();
+			if (BlueprintEditorToolbar.IsValid())
+			{
+				const FFullBlueprintEditorCommands& Commands = FFullBlueprintEditorCommands::Get();
+
+				InSection.AddEntry(FToolMenuEntry::InitToolBarButton(
+					Commands.Compile,
+					TAttribute<FText>(),
+					TAttribute<FText>(BlueprintEditorToolbar.ToSharedRef(), &FBlueprintEditorToolbar::GetStatusTooltip),
+					TAttribute<FSlateIcon>(BlueprintEditorToolbar.ToSharedRef(), &FBlueprintEditorToolbar::GetStatusImage),
+					"CompileBlueprint"
+				));
+
+				InSection.AddEntry(FToolMenuEntry::InitComboButton(
+					"BlueprintCompileOptions",
+					FUIAction(),
+					FOnGetContent::CreateStatic(&BlueprintEditorToolbarImpl::GenerateCompileOptionsWidget, Context->BlueprintEditor.Pin()->GetToolkitCommands()),
+					LOCTEXT("BlupeintCompileOptions_ToolbarName",    "Compile Options"),
+					LOCTEXT("BlupeintCompileOptions_ToolbarTooltip", "Options to customize how Blueprints compile"),
+					TAttribute<FSlateIcon>(),
+					/*bSimpleComboBox =*/true
+				));
+			}
+		}
+	}));
 }
 
-void FBlueprintEditorToolbar::FillCompileToolbar(FToolBarBuilder& ToolbarBuilder)
+void FBlueprintEditorToolbar::AddNewToolbar(UToolMenu* InMenu)
 {
-	const FFullBlueprintEditorCommands& Commands = FFullBlueprintEditorCommands::Get();
-	TSharedPtr<FBlueprintEditor> BlueprintEditorPtr = BlueprintEditor.Pin();
-	UBlueprint* BlueprintObj = BlueprintEditorPtr->GetBlueprintObj();
+	FToolMenuSection& Section = InMenu->AddSection("Add");
+	Section.InsertPosition = FToolMenuInsert("MyBlueprint", EToolMenuInsertType::After);
 
-	ToolbarBuilder.BeginSection("Compile");
-	if (BlueprintObj != NULL)
+	Section.AddDynamicEntry("AddCommands", FNewToolMenuSectionDelegate::CreateLambda([](FToolMenuSection& InSection)
 	{
-		ToolbarBuilder.AddToolBarButton( Commands.Compile,
-									 NAME_None, 
-									 TAttribute<FText>(),
-									 TAttribute<FText>(this, &FBlueprintEditorToolbar::GetStatusTooltip),
-									 TAttribute<FSlateIcon>(this, &FBlueprintEditorToolbar::GetStatusImage),
-									 FName(TEXT("CompileBlueprint")));
-
-		FUIAction TempCompileOptionsCommand;
-		ToolbarBuilder.AddComboButton(
-			TempCompileOptionsCommand,
-			FOnGetContent::CreateStatic(&BlueprintEditorToolbarImpl::GenerateCompileOptionsWidget, BlueprintEditorPtr->GetToolkitCommands()),
-			LOCTEXT("BlupeintCompileOptions_ToolbarName",    "Compile Options"),
-			LOCTEXT("BlupeintCompileOptions_ToolbarTooltip", "Options to customize how Blueprints compile"),
-			TAttribute<FSlateIcon>(),
-			/*bSimpleComboBox =*/true
-		);
-	}
-	ToolbarBuilder.EndSection();
+		UBlueprintEditorToolMenuContext* Context = InSection.FindContext<UBlueprintEditorToolMenuContext>();
+		if (Context && Context->BlueprintEditor.IsValid() && Context->GetBlueprintObj())
+		{
+			const FBlueprintEditorCommands& Commands = FBlueprintEditorCommands::Get();
+			InSection.AddEntry(FToolMenuEntry::InitToolBarButton(Commands.AddNewVariable, TAttribute<FText>(), TAttribute<FText>(), TAttribute<FSlateIcon>(), FName(TEXT("BPEAddNewVariable"))));
+			InSection.AddEntry(FToolMenuEntry::InitToolBarButton(Commands.AddNewFunction, TAttribute<FText>(), TAttribute<FText>(), TAttribute<FSlateIcon>(), FName(TEXT("BPEAddNewFunction"))));
+			InSection.AddEntry(FToolMenuEntry::InitToolBarButton(Commands.AddNewMacroDeclaration, TAttribute<FText>(), TAttribute<FText>(), TAttribute<FSlateIcon>(), FName(TEXT("BPEAddNewMacro"))));
+			// Add New Animation Graph isn't supported right now.
+			InSection.AddEntry(FToolMenuEntry::InitToolBarButton(Commands.AddNewEventGraph, TAttribute<FText>(), TAttribute<FText>(), TAttribute<FSlateIcon>(), FName(TEXT("BPEAddNewEventGraph"))));
+			InSection.AddEntry(FToolMenuEntry::InitToolBarButton(Commands.AddNewDelegate, TAttribute<FText>(), TAttribute<FText>(), TAttribute<FSlateIcon>(), FName(TEXT("BPEAddNewDelegate"))));
+		}
+	}));
 }
 
-void FBlueprintEditorToolbar::FillNewToolbar(FToolBarBuilder& ToolbarBuilder)
+void FBlueprintEditorToolbar::AddScriptingToolbar(UToolMenu* InMenu)
 {
-	const FBlueprintEditorCommands& Commands = FBlueprintEditorCommands::Get();
-
-	TSharedPtr<FBlueprintEditor> BlueprintEditorPtr = BlueprintEditor.Pin();
-	UBlueprint* BlueprintObj = BlueprintEditorPtr->GetBlueprintObj();
+	FToolMenuSection& Section = InMenu->AddSection("Script");
+	Section.InsertPosition = FToolMenuInsert("Asset", EToolMenuInsertType::After);
 
-	ToolbarBuilder.BeginSection("AddNew");
-	if (BlueprintObj != NULL)
+	Section.AddDynamicEntry("ScriptCommands", FNewToolMenuSectionDelegate::CreateLambda([](FToolMenuSection& InSection)
 	{
-		ToolbarBuilder.AddToolBarButton(Commands.AddNewVariable, NAME_None, TAttribute<FText>(), TAttribute<FText>(), TAttribute<FSlateIcon>(), FName(TEXT("BPEAddNewVariable")));
-		ToolbarBuilder.AddToolBarButton(Commands.AddNewFunction, NAME_None, TAttribute<FText>(), TAttribute<FText>(), TAttribute<FSlateIcon>(), FName(TEXT("BPEAddNewFunction")));
-		ToolbarBuilder.AddToolBarButton(Commands.AddNewMacroDeclaration, NAME_None, TAttribute<FText>(), TAttribute<FText>(), TAttribute<FSlateIcon>(), FName(TEXT("BPEAddNewMacro")));
-		// Add New Animation Graph isn't supported right now.
-		ToolbarBuilder.AddToolBarButton(Commands.AddNewEventGraph, NAME_None, TAttribute<FText>(), TAttribute<FText>(), TAttribute<FSlateIcon>(), FName(TEXT("BPEAddNewEventGraph")));
-		ToolbarBuilder.AddToolBarButton(Commands.AddNewDelegate, NAME_None, TAttribute<FText>(), TAttribute<FText>(), TAttribute<FSlateIcon>(), FName(TEXT("BPEAddNewDelegate")));
-	}
-	ToolbarBuilder.EndSection();
-}
-
-void FBlueprintEditorToolbar::FillScriptingToolbar(FToolBarBuilder& ToolbarBuilder)
-{
-	const FBlueprintEditorCommands& Commands = FBlueprintEditorCommands::Get();
-
-	TSharedPtr<FBlueprintEditor> BlueprintEditorPtr = BlueprintEditor.Pin();
-	UBlueprint* BlueprintObj = BlueprintEditorPtr->GetBlueprintObj();
-
-	ToolbarBuilder.BeginSection("Script");
-
-	ToolbarBuilder.AddToolBarButton(FBlueprintEditorCommands::Get().FindInBlueprint);
-
-	ToolbarBuilder.AddToolBarButton(
-		FBlueprintEditorCommands::Get().ToggleHideUnrelatedNodes,
-		NAME_None,
-		TAttribute<FText>(),
-		TAttribute<FText>(),
-		FSlateIcon(FEditorStyle::GetStyleSetName(), "GraphEditor.ToggleHideUnrelatedNodes")
-	);
-	ToolbarBuilder.AddComboButton(
-		FUIAction(),
-		FOnGetContent::CreateSP(BlueprintEditorPtr.Get(), &FBlueprintEditor::MakeHideUnrelatedNodesOptionsMenu),
-		LOCTEXT("HideUnrelatedNodesOptions", "Focus Related Nodes Options"),
-		LOCTEXT("HideUnrelatedNodesOptionsMenu", "Focus Related Nodes options menu"),
-		TAttribute<FSlateIcon>(),
-		true
-	);
-
-	ToolbarBuilder.EndSection();
+		UBlueprintEditorToolMenuContext* Context = InSection.FindContext<UBlueprintEditorToolMenuContext>();
+		if (Context && Context->BlueprintEditor.IsValid() && Context->GetBlueprintObj())
+		{
+			InSection.AddEntry(FToolMenuEntry::InitToolBarButton(
+				FBlueprintEditorCommands::Get().ToggleHideUnrelatedNodes,
+				TAttribute<FText>(),
+				TAttribute<FText>(),
+				FSlateIcon(FEditorStyle::GetStyleSetName(), "GraphEditor.ToggleHideUnrelatedNodes")
+			));
+
+			InSection.AddEntry(FToolMenuEntry::InitComboButton(
+				"HideUnrelatedNodesOptions",
+				FUIAction(),
+				FOnGetContent::CreateSP(Context->BlueprintEditor.Pin().Get(), &FBlueprintEditor::MakeHideUnrelatedNodesOptionsMenu),
+				LOCTEXT("HideUnrelatedNodesOptions", "Focus Related Nodes Options"),
+				LOCTEXT("HideUnrelatedNodesOptionsMenu", "Focus Related Nodes options menu"),
+				TAttribute<FSlateIcon>(),
+				true
+			));
+		}
+	}));
 }
 
-void FBlueprintEditorToolbar::FillDebuggingToolbar(FToolBarBuilder& ToolbarBuilder)
+void FBlueprintEditorToolbar::AddDebuggingToolbar(UToolMenu* InMenu)
 {
-	TSharedPtr<FBlueprintEditor> BlueprintEditorPtr = BlueprintEditor.Pin();
-	UBlueprint* BlueprintObj = BlueprintEditorPtr->GetBlueprintObj();
+	FToolMenuSection& Section = InMenu->AddSection("Debugging");
+	Section.InsertPosition = FToolMenuInsert("Asset", EToolMenuInsertType::After);
 
-	ToolbarBuilder.BeginSection("Debugging");
-	if (BlueprintObj)
+	Section.AddDynamicEntry("DebuggingCommands", FNewToolMenuSectionDelegate::CreateLambda([](FToolMenuSection& InSection)
 	{
-		FPlayWorldCommands::BuildToolbar(ToolbarBuilder);
-
-		if (BlueprintObj->BlueprintType != BPTYPE_MacroLibrary)
+		UBlueprintEditorToolMenuContext* Context = InSection.FindContext<UBlueprintEditorToolMenuContext>();
+		if (Context && Context->BlueprintEditor.IsValid() && Context->GetBlueprintObj())
 		{
-			// Selected debug actor button
-			ToolbarBuilder.AddWidget(SNew(SBlueprintEditorSelectedDebugObjectWidget, BlueprintEditorPtr));
+			FPlayWorldCommands::BuildToolbar(InSection);
+
+			if (Context->GetBlueprintObj()->BlueprintType != BPTYPE_MacroLibrary)
+			{
+				// Selected debug actor button
+				InSection.AddEntry(FToolMenuEntry::InitWidget("SelectedDebugObjectWidget", SNew(SBlueprintEditorSelectedDebugObjectWidget, Context->BlueprintEditor.Pin()), FText::GetEmpty()));
+			}
 		}
-	}
-	ToolbarBuilder.EndSection();
+	}));
 }
 
-void FBlueprintEditorToolbar::FillComponentsToolbar(FToolBarBuilder& ToolbarBuilder)
+void FBlueprintEditorToolbar::AddComponentsToolbar(UToolMenu* InMenu)
 {
-	TSharedPtr<FBlueprintEditor> BlueprintEditorPtr = BlueprintEditor.Pin();
-	UBlueprint* BlueprintObj = BlueprintEditorPtr->GetBlueprintObj();
+	FToolMenuInsert InsertPosition("Asset", EToolMenuInsertType::After);
 
 #if 0 // restore this if we ever need the ability to toggle component editing on/off
-	ToolbarBuilder.BeginSection("Components");
-		ToolbarBuilder.AddToolBarButton(FSCSCommands::Get().ToggleComponentEditing);
-	ToolbarBuilder.EndSection();
+	{
+		FToolMenuSection& Section = InMenu->AddSection("Components");
+		Section.InsertPosition = InsertPosition;
+		Section.AddEntry(FToolMenuEntry::InitToolBarButton(FSCSCommands::Get().ToggleComponentEditing));
+	}
 #endif
 
-	ToolbarBuilder.BeginSection("ComponentsViewport");
-		ToolbarBuilder.AddToolBarButton(FBlueprintEditorCommands::Get().EnableSimulation);
-	ToolbarBuilder.EndSection();
+	{
+		FToolMenuSection& Section = InMenu->AddSection("ComponentsViewport");
+		Section.InsertPosition = InsertPosition;
+		Section.AddEntry(FToolMenuEntry::InitToolBarButton(FBlueprintEditorCommands::Get().EnableSimulation));
+	}
 }
 
 FSlateIcon FBlueprintEditorToolbar::GetStatusImage() const
@@ -668,72 +597,4 @@ FText FBlueprintEditorToolbar::GetStatusTooltip() const
 	}
 }
 
-TArray< TSharedPtr< SWidget> > FBlueprintEditorToolbar::GenerateToolbarWidgets(const UBlueprint* BlueprintObj, TAttribute<FName> ActiveModeGetter, FOnModeChangeRequested ActiveModeSetter)
-{
-	TArray< TSharedPtr< SWidget> > Ret;
-	if (!BlueprintObj ||
-		(!FBlueprintEditorUtils::IsLevelScriptBlueprint(BlueprintObj)
-		&& !FBlueprintEditorUtils::IsInterfaceBlueprint(BlueprintObj)
-		&& !BlueprintObj->bIsNewlyCreated)
-		)
-	{
-		// Left side padding
-		Ret.Add(SNew(SSpacer).Size(FVector2D(4.0f, 1.0f)));
-
-		Ret.Add(
-			SNew(SModeWidget, FBlueprintEditorApplicationModes::GetLocalizedMode(FBlueprintEditorApplicationModes::BlueprintDefaultsMode), FBlueprintEditorApplicationModes::BlueprintDefaultsMode)
-			.OnGetActiveMode(ActiveModeGetter)
-			.OnSetActiveMode(ActiveModeSetter)
-			.CanBeSelected(BlueprintObj ? FBlueprintEditorUtils::DoesSupportDefaults(BlueprintObj) : false)
-			.ToolTip(IDocumentation::Get()->CreateToolTip(
-				LOCTEXT("BlueprintDefaultsModeButtonTooltip", "Switch to Class Defaults Mode"),
-				NULL,
-				TEXT("Shared/Editors/BlueprintEditor"),
-				TEXT("DefaultsMode")))
-			.IconImage(FEditorStyle::GetBrush("FullBlueprintEditor.SwitchToBlueprintDefaultsMode"))
-			.SmallIconImage(FEditorStyle::GetBrush("FullBlueprintEditor.SwitchToBlueprintDefaultsMode.Small"))
-			.AddMetaData<FTagMetaData>(FTagMetaData(TEXT("DefaultsMode")))
-		);
-
-		Ret.Add(SNew(SBlueprintModeSeparator));
-
-		Ret.Add(
-			SNew(SModeWidget, FBlueprintEditorApplicationModes::GetLocalizedMode(FBlueprintEditorApplicationModes::BlueprintComponentsMode), FBlueprintEditorApplicationModes::BlueprintComponentsMode)
-			.OnGetActiveMode(ActiveModeGetter)
-			.OnSetActiveMode(ActiveModeSetter)
-			.CanBeSelected(BlueprintObj ? FBlueprintEditorUtils::DoesSupportComponents(BlueprintObj) : false)
-			.ToolTip(IDocumentation::Get()->CreateToolTip(
-				LOCTEXT("ComponentsModeButtonTooltip", "Switch to Components Mode"),
-				NULL,
-				TEXT("Shared/Editors/BlueprintEditor"),
-				TEXT("ComponentsMode")))
-			.IconImage(FEditorStyle::GetBrush("FullBlueprintEditor.SwitchToComponentsMode"))
-			.SmallIconImage(FEditorStyle::GetBrush("FullBlueprintEditor.SwitchToComponentsMode.Small"))
-			.AddMetaData<FTagMetaData>(FTagMetaData(TEXT("ComponentsMode")))
-		);
-
-		Ret.Add(SNew(SBlueprintModeSeparator));
-
-		Ret.Add(
-			SNew(SModeWidget, FBlueprintEditorApplicationModes::GetLocalizedMode(FBlueprintEditorApplicationModes::StandardBlueprintEditorMode), FBlueprintEditorApplicationModes::StandardBlueprintEditorMode)
-			.OnGetActiveMode(ActiveModeGetter)
-			.OnSetActiveMode(ActiveModeSetter)
-			.CanBeSelected(BlueprintObj != NULL)
-			.ToolTip(IDocumentation::Get()->CreateToolTip(
-				LOCTEXT("GraphModeButtonTooltip", "Switch to Graph Editing Mode"),
-				NULL,
-				TEXT("Shared/Editors/BlueprintEditor"),
-				TEXT("GraphMode")))
-			.ToolTipText(LOCTEXT("GraphModeButtonTooltip", "Switch to Graph Editing Mode"))
-			.IconImage(FEditorStyle::GetBrush("FullBlueprintEditor.SwitchToScriptingMode"))
-			.SmallIconImage(FEditorStyle::GetBrush("FullBlueprintEditor.SwitchToScriptingMode.Small"))
-			.AddMetaData<FTagMetaData>(FTagMetaData(TEXT("GraphMode")))
-		);
-
-		// Right side padding
-		Ret.Add(SNew(SSpacer).Size(FVector2D(4.0f, 1.0f)));
-	}
-	return Ret;
-}
-
 #undef LOCTEXT_NAMESPACE
diff --git a/Engine/Source/Editor/Kismet/Private/SMyBlueprint.cpp b/Engine/Source/Editor/Kismet/Private/SMyBlueprint.cpp
--- a/Engine/Source/Editor/Kismet/Private/SMyBlueprint.cpp
+++ b/Engine/Source/Editor/Kismet/Private/SMyBlueprint.cpp
@@ -300,10 +300,6 @@ void SMyBlueprint::Construct(const FArguments& InArgs, TWeakPtr<FBlueprintEditor
 			FCanExecuteAction(),
 			FIsActionChecked(),
 			FIsActionButtonVisible::CreateSP(this, &SMyBlueprint::IsNativeVariable) );
-
-		TSharedPtr<FBlueprintEditorToolbar> Toolbar = MakeShareable(new FBlueprintEditorToolbar(InBlueprintEditor.Pin()));
-		TSharedPtr<FExtender> Extender = MakeShareable(new FExtender);
-		Toolbar->AddNewToolbar(Extender);
 		ToolbarBuilderWidget = SNullWidget::NullWidget;
 	
 		ToolKitCommandList->MapAction(FGenericCommands::Get().Rename,
diff --git a/Engine/Source/Editor/Kismet/Private/WorkflowCentricApplication.cpp b/Engine/Source/Editor/Kismet/Private/WorkflowCentricApplication.cpp
--- a/Engine/Source/Editor/Kismet/Private/WorkflowCentricApplication.cpp
+++ b/Engine/Source/Editor/Kismet/Private/WorkflowCentricApplication.cpp
@@ -26,6 +26,22 @@ void FWorkflowCentricApplication::UnregisterTabSpawners(const TSharedRef<class F
 	InTabManager->UnregisterAllTabSpawners();
 }
 
+FName FWorkflowCentricApplication::GetToolMenuToolbarName(FName& OutParentName) const
+{
+	return GetToolMenuToolbarNameForMode(GetCurrentMode(), OutParentName);
+}
+
+FName FWorkflowCentricApplication::GetToolMenuToolbarNameForMode(const FName InModeName, FName& OutParentName) const
+{
+	const FName BaseMenuName = FAssetEditorToolkit::GetToolMenuToolbarName(OutParentName);
+	if (InModeName != NAME_None)
+	{
+		OutParentName = BaseMenuName;
+		return *(BaseMenuName.ToString() + TEXT(".") + InModeName.ToString());
+	}
+
+	return BaseMenuName;
+}
 
 FName FWorkflowCentricApplication::GetCurrentMode() const
 {
diff --git a/Engine/Source/Editor/Kismet/Public/BlueprintEditor.h b/Engine/Source/Editor/Kismet/Public/BlueprintEditor.h
--- a/Engine/Source/Editor/Kismet/Public/BlueprintEditor.h
+++ b/Engine/Source/Editor/Kismet/Public/BlueprintEditor.h
@@ -47,6 +47,7 @@ class UBlueprintEditorOptions;
 struct Rect;
 class UK2Node_FunctionEntry;
 class UK2Node_Event;
+class UToolMenu;
 
 /* Enums to use when grouping the blueprint members in the list panel. The order here will determine the order in the list */
 namespace NodeSectionID
@@ -238,6 +239,17 @@ public:
 
 	virtual ~FBlueprintEditor();
 
+	/** Add context objects for menus and toolbars */
+	virtual void InitToolMenuContext(FToolMenuContext& MenuContext) override;
+
+	/**
+	* Registers toolbar for a mode if not already registered
+	*
+	* @param InModeName - Name of the mode to register toolbar for
+	* @return nullptr if toolbar already registered.
+	*/
+	UToolMenu* RegisterModeToolbarIfUnregistered(const FName InModeName);
+
 	/** Check the Node Title is visible */
 	bool IsNodeTitleVisible(const UEdGraphNode* Node, bool bRequestRename);
 
diff --git a/Engine/Source/Editor/Kismet/Public/BlueprintEditorContext.h b/Engine/Source/Editor/Kismet/Public/BlueprintEditorContext.h
new file mode 100644
--- /dev/null
+++ b/Engine/Source/Editor/Kismet/Public/BlueprintEditorContext.h
@@ -0,0 +1,21 @@
+// Copyright 1998-2019 Epic Games, Inc. All Rights Reserved.
+
+#pragma once
+
+#include "CoreMinimal.h"
+
+#include "BlueprintEditorContext.generated.h"
+
+class UBlueprint;
+class FBlueprintEditor;
+
+UCLASS()
+class KISMET_API UBlueprintEditorToolMenuContext : public UObject
+{
+	GENERATED_BODY()
+public:
+
+	UBlueprint* GetBlueprintObj() const;
+
+	TWeakPtr<FBlueprintEditor> BlueprintEditor;
+};
diff --git a/Engine/Source/Editor/Kismet/Public/SBlueprintEditorToolbar.h b/Engine/Source/Editor/Kismet/Public/SBlueprintEditorToolbar.h
--- a/Engine/Source/Editor/Kismet/Public/SBlueprintEditorToolbar.h
+++ b/Engine/Source/Editor/Kismet/Public/SBlueprintEditorToolbar.h
@@ -14,6 +14,7 @@ class FBlueprintEditor;
 class FExtender;
 class FMenuBuilder;
 class FToolBarBuilder;
+class UToolMenu;
 
 /**
  * Kismet menu
@@ -79,13 +80,12 @@ public:
 	FBlueprintEditorToolbar(TSharedPtr<FBlueprintEditor> InBlueprintEditor)
 		: BlueprintEditor(InBlueprintEditor) {}
 
-	void AddBlueprintEditorModesToolbar(TSharedPtr<FExtender> Extender);
-	void AddBlueprintGlobalOptionsToolbar(TSharedPtr<FExtender> Extender);
-	void AddCompileToolbar(TSharedPtr<FExtender> Extender);
-	void AddNewToolbar(TSharedPtr<FExtender> Extender);
-	void AddScriptingToolbar(TSharedPtr<FExtender> Extender);
-	void AddDebuggingToolbar(TSharedPtr<FExtender> Extender);
-	void AddComponentsToolbar(TSharedPtr<FExtender> Extender);
+	void AddBlueprintGlobalOptionsToolbar(UToolMenu* InMenu);
+	void AddCompileToolbar(UToolMenu* InMenu);
+	void AddNewToolbar(UToolMenu* InMenu);
+	void AddScriptingToolbar(UToolMenu* InMenu);
+	void AddDebuggingToolbar(UToolMenu* InMenu);
+	void AddComponentsToolbar(UToolMenu* InMenu);
 
 	/** Returns the current status icon for the blueprint being edited */
 	FSlateIcon GetStatusImage() const;
@@ -93,18 +93,6 @@ public:
 	/** Returns the current status as text for the blueprint being edited */
 	FText GetStatusTooltip() const;
 
-	/** Helper function for generating the buttons in the toolbar, reused by merge and diff tools */
-	static TArray< TSharedPtr< class SWidget> > GenerateToolbarWidgets(const class UBlueprint* BlueprintObj, TAttribute<FName> ActiveModeGetter, FOnModeChangeRequested ActiveModeSetter);
-
-private:
-	void FillBlueprintEditorModesToolbar(FToolBarBuilder& ToolbarBuilder);
-	void FillBlueprintGlobalOptionsToolbar(FToolBarBuilder& ToolBarBuilder);
-	void FillCompileToolbar(FToolBarBuilder& ToolbarBuilder);
-	void FillNewToolbar(FToolBarBuilder& ToolbarBuilder);
-	void FillScriptingToolbar(FToolBarBuilder& ToolbarBuilder);
-	void FillDebuggingToolbar(FToolBarBuilder& ToolbarBuilder);
-	void FillComponentsToolbar(FToolBarBuilder& ToolbarBuilder);
-
 protected:
 	/** Pointer back to the blueprint editor tool that owns us */
 	TWeakPtr<FBlueprintEditor> BlueprintEditor;
diff --git a/Engine/Source/Editor/Kismet/Public/WorkflowOrientedApp/WorkflowCentricApplication.h b/Engine/Source/Editor/Kismet/Public/WorkflowOrientedApp/WorkflowCentricApplication.h
--- a/Engine/Source/Editor/Kismet/Public/WorkflowOrientedApp/WorkflowCentricApplication.h
+++ b/Engine/Source/Editor/Kismet/Public/WorkflowOrientedApp/WorkflowCentricApplication.h
@@ -25,8 +25,12 @@ public:
 	
 	// FAssetEditorToolkit interface
 	virtual bool OnRequestClose() override;
+	virtual FName GetToolMenuToolbarName(FName& OutParentName) const override;
 	// End of FAssetEditorToolkit interface
 
+	// Returns the name of the toolbar for the given mode along with the name of the toolbar's parent.
+	FName GetToolMenuToolbarNameForMode(const FName InModeName, FName& OutParentName) const;
+
 	// Returns the current mode of this application
 	FName GetCurrentMode() const;
 	TSharedPtr<FApplicationMode> GetCurrentModePtr() const { return CurrentAppModePtr; }
diff --git a/Engine/Source/Editor/LevelEditor/Private/LevelEditorToolBar.cpp b/Engine/Source/Editor/LevelEditor/Private/LevelEditorToolBar.cpp
--- a/Engine/Source/Editor/LevelEditor/Private/LevelEditorToolBar.cpp
+++ b/Engine/Source/Editor/LevelEditor/Private/LevelEditorToolBar.cpp
@@ -1494,15 +1494,12 @@ void FLevelEditorToolBar::RegisterLevelEditorToolBar( const TSharedRef<FUIComman
 		}));
 	}
 
-	Toolbar->AddDynamicSection("Game", FNewToolBarDelegateLegacy::CreateLambda([](FToolBarBuilder& InBuilder, UToolMenu* InData)
 	{
-		InBuilder.BeginSection("Game");
+		FToolMenuSection& Section = Toolbar->AddSection("Game");
 
 		// Add the shared play-world commands that will be shown on the Kismet toolbar as well
-		FPlayWorldCommands::BuildToolbar(InBuilder, true);
-
-		InBuilder.EndSection();
-	}));
+		FPlayWorldCommands::BuildToolbar(Section, true);
+	}
 
 #undef LOCTEXT_NAMESPACE
 }
diff --git a/Engine/Source/Editor/UMGEditor/Private/BlueprintModes/WidgetDesignerApplicationMode.cpp b/Engine/Source/Editor/UMGEditor/Private/BlueprintModes/WidgetDesignerApplicationMode.cpp
--- a/Engine/Source/Editor/UMGEditor/Private/BlueprintModes/WidgetDesignerApplicationMode.cpp
+++ b/Engine/Source/Editor/UMGEditor/Private/BlueprintModes/WidgetDesignerApplicationMode.cpp
@@ -8,7 +8,7 @@
 
 #include "WidgetBlueprintEditorToolbar.h"
 #include "UMGEditorModule.h"
-
+#include "ToolMenus.h"
 
 #include "TabFactory/PaletteTabSummoner.h"
 #include "TabFactory/HierarchyTabSummoner.h"
@@ -123,9 +123,13 @@ FWidgetDesignerApplicationMode::FWidgetDesignerApplicationMode(TSharedPtr<FWidge
 	ToolbarExtender = UMGEditorModule.GetToolBarExtensibilityManager()->GetAllExtenders();
 	
 	InWidgetEditor->GetWidgetToolbarBuilder()->AddWidgetBlueprintEditorModesToolbar(ToolbarExtender);
-	InWidgetEditor->GetWidgetToolbarBuilder()->AddWidgetReflector(ToolbarExtender);
-	InWidgetEditor->GetToolbarBuilder()->AddCompileToolbar(ToolbarExtender);
-	InWidgetEditor->GetToolbarBuilder()->AddDebuggingToolbar(ToolbarExtender);
+
+	if (UToolMenu* Toolbar = InWidgetEditor->RegisterModeToolbarIfUnregistered(GetModeName()))
+	{
+		InWidgetEditor->GetWidgetToolbarBuilder()->AddWidgetReflector(Toolbar);
+		InWidgetEditor->GetToolbarBuilder()->AddCompileToolbar(Toolbar);
+		InWidgetEditor->GetToolbarBuilder()->AddDebuggingToolbar(Toolbar);
+	}
 }
 
 void FWidgetDesignerApplicationMode::RegisterTabFactories(TSharedPtr<FTabManager> InTabManager)
diff --git a/Engine/Source/Editor/UMGEditor/Private/BlueprintModes/WidgetGraphApplicationMode.cpp b/Engine/Source/Editor/UMGEditor/Private/BlueprintModes/WidgetGraphApplicationMode.cpp
--- a/Engine/Source/Editor/UMGEditor/Private/BlueprintModes/WidgetGraphApplicationMode.cpp
+++ b/Engine/Source/Editor/UMGEditor/Private/BlueprintModes/WidgetGraphApplicationMode.cpp
@@ -7,6 +7,7 @@
 
 #include "WidgetBlueprintEditorToolbar.h"
 #include "BlueprintModes/WidgetBlueprintApplicationModes.h"
+#include "ToolMenus.h"
 
 /////////////////////////////////////////////////////
 // FWidgetGraphApplicationMode
@@ -78,10 +79,14 @@ FWidgetGraphApplicationMode::FWidgetGraphApplicationMode(TSharedPtr<FWidgetBluep
 	//@TODO: Keep this in sync with AnimBlueprintMode.cpp
 	ToolbarExtender = MakeShareable(new FExtender);
 	InWidgetEditor->GetWidgetToolbarBuilder()->AddWidgetBlueprintEditorModesToolbar(ToolbarExtender);
-	InWidgetEditor->GetToolbarBuilder()->AddCompileToolbar(ToolbarExtender);
-	InWidgetEditor->GetToolbarBuilder()->AddScriptingToolbar(ToolbarExtender);
-	InWidgetEditor->GetToolbarBuilder()->AddBlueprintGlobalOptionsToolbar(ToolbarExtender);
-	InWidgetEditor->GetToolbarBuilder()->AddDebuggingToolbar(ToolbarExtender);
+
+	if (UToolMenu* Toolbar = InWidgetEditor->RegisterModeToolbarIfUnregistered(GetModeName()))
+	{
+		InWidgetEditor->GetToolbarBuilder()->AddCompileToolbar(Toolbar);
+		InWidgetEditor->GetToolbarBuilder()->AddScriptingToolbar(Toolbar);
+		InWidgetEditor->GetToolbarBuilder()->AddBlueprintGlobalOptionsToolbar(Toolbar);
+		InWidgetEditor->GetToolbarBuilder()->AddDebuggingToolbar(Toolbar);
+	}
 }
 
 void FWidgetGraphApplicationMode::RegisterTabFactories(TSharedPtr<FTabManager> InTabManager)
diff --git a/Engine/Source/Editor/UMGEditor/Private/WidgetBlueprintEditorToolbar.cpp b/Engine/Source/Editor/UMGEditor/Private/WidgetBlueprintEditorToolbar.cpp
--- a/Engine/Source/Editor/UMGEditor/Private/WidgetBlueprintEditorToolbar.cpp
+++ b/Engine/Source/Editor/UMGEditor/Private/WidgetBlueprintEditorToolbar.cpp
@@ -14,6 +14,8 @@
 #include "IDocumentation.h"
 #include "BlueprintEditor.h"
 #include "Framework/MultiBox/MultiBoxBuilder.h"
+#include "ToolMenus.h"
+#include "BlueprintEditorContext.h"
 
 #include "WidgetBlueprintEditor.h"
 #include "WorkflowOrientedApp/SModeWidget.h"
@@ -123,33 +125,22 @@ void FWidgetBlueprintEditorToolbar::FillWidgetBlueprintEditorModesToolbar(FToolB
 	}
 }
 
-void FWidgetBlueprintEditorToolbar::AddWidgetReflector(TSharedPtr<FExtender> Extender)
-{
-	TSharedPtr<FWidgetBlueprintEditor> BlueprintEditorPtr = WidgetEditor.Pin();
 
-	Extender->AddToolBarExtension(
-		"Asset",
-		EExtensionHook::After,
-		BlueprintEditorPtr->GetToolkitCommands(),
-		FToolBarExtensionDelegate::CreateSP(this, &FWidgetBlueprintEditorToolbar::FillWidgetReflectorToolbar));
-}
-
-void FWidgetBlueprintEditorToolbar::FillWidgetReflectorToolbar(FToolBarBuilder& ToolbarBuilder)
+void FWidgetBlueprintEditorToolbar::AddWidgetReflector(UToolMenu* InMenu)
 {
-	ToolbarBuilder.BeginSection("WidgetTools");
+	FToolMenuSection& Section = InMenu->AddSection("WidgetTools");
+	Section.InsertPosition = FToolMenuInsert("Asset", EToolMenuInsertType::After);
 
-	ToolbarBuilder.AddToolBarButton(
+	Section.AddEntry(FToolMenuEntry::InitToolBarButton(
+		"OpenWidgetReflector",
 		FUIAction(
 			FExecuteAction::CreateLambda([=] { FGlobalTabmanager::Get()->InvokeTab(FTabId("WidgetReflector")); }),
 			FCanExecuteAction()
 		)
-		, NAME_None
 		, LOCTEXT("OpenWidgetReflector", "Widget Reflector")
 		, LOCTEXT("OpenWidgetReflectorToolTip", "Opens the Widget Reflector, a handy tool for diagnosing problems with live widgets.")
 		, FSlateIcon(FCoreStyle::Get().GetStyleSetName(), "WidgetReflector.Icon")
-	);
-
-	ToolbarBuilder.EndSection();
+	));
 }
 
 #undef LOCTEXT_NAMESPACE
diff --git a/Engine/Source/Editor/UMGEditor/Private/WidgetBlueprintEditorToolbar.h b/Engine/Source/Editor/UMGEditor/Private/WidgetBlueprintEditorToolbar.h
--- a/Engine/Source/Editor/UMGEditor/Private/WidgetBlueprintEditorToolbar.h
+++ b/Engine/Source/Editor/UMGEditor/Private/WidgetBlueprintEditorToolbar.h
@@ -7,6 +7,7 @@
 class FExtender;
 class FToolBarBuilder;
 class FWidgetBlueprintEditor;
+class UToolMenu;
 
 /**
  * Handles all of the toolbar related construction for the widget blueprint editor.
@@ -23,14 +24,11 @@ public:
 	 */
 	void AddWidgetBlueprintEditorModesToolbar(TSharedPtr<FExtender> Extender);
 
-	void AddWidgetReflector(TSharedPtr<FExtender> Extender);
+	void AddWidgetReflector(UToolMenu* InMenu);
 
 public:
 	/**  */
 	void FillWidgetBlueprintEditorModesToolbar(FToolBarBuilder& ToolbarBuilder);
 
-	/**  */
-	void FillWidgetReflectorToolbar(FToolBarBuilder& ToolbarBuilder);
-
 	TWeakPtr<FWidgetBlueprintEditor> WidgetEditor;
 };
diff --git a/Engine/Source/Editor/UMGEditor/UMGEditor.Build.cs b/Engine/Source/Editor/UMGEditor/UMGEditor.Build.cs
--- a/Engine/Source/Editor/UMGEditor/UMGEditor.Build.cs
+++ b/Engine/Source/Editor/UMGEditor/UMGEditor.Build.cs
@@ -64,7 +64,8 @@ public class UMGEditor : ModuleRules
                 "TargetPlatform",
 				"TimeManagement",
 				"GameProjectGeneration",
-				"PropertyPath"
+				"PropertyPath",
+				"ToolMenus",
 			}
 			);
 	}
diff --git a/Engine/Source/Editor/UnrealEd/Private/Kismet2/DebuggerCommands.cpp b/Engine/Source/Editor/UnrealEd/Private/Kismet2/DebuggerCommands.cpp
--- a/Engine/Source/Editor/UnrealEd/Private/Kismet2/DebuggerCommands.cpp
+++ b/Engine/Source/Editor/UnrealEd/Private/Kismet2/DebuggerCommands.cpp
@@ -69,6 +69,8 @@
 #include "Framework/Notifications/NotificationManager.h"
 #include "Widgets/Notifications/SNotificationList.h"
 
+#include "ToolMenus.h"
+#include "SBlueprintEditorToolbar.h"
 
 #define LOCTEXT_NAMESPACE "DebuggerCommands"
 
@@ -612,86 +614,92 @@ void FPlayWorldCommands::AddPIEPreviewDeviceActions(const FPlayWorldCommands &Co
 	}
 }
 
-void FPlayWorldCommands::BuildToolbar( FToolBarBuilder& ToolbarBuilder, bool bIncludeLaunchButtonAndOptions )
+void FPlayWorldCommands::BuildToolbar(FToolMenuSection& InSection, bool bIncludeLaunchButtonAndOptions)
 {
 	// Play
-	ToolbarBuilder.AddToolBarButton( 
+	InSection.AddEntry(FToolMenuEntry::InitToolBarButton( 
 		FPlayWorldCommands::Get().RepeatLastPlay, 
-		NAME_None, 
 		LOCTEXT("RepeatLastPlay", "Play"),
 		TAttribute< FText >::Create( TAttribute< FText >::FGetter::CreateStatic( &FInternalPlayWorldCommandCallbacks::GetRepeatLastPlayToolTip ) ),
 		TAttribute< FSlateIcon >::Create( TAttribute< FSlateIcon >::FGetter::CreateStatic( &FInternalPlayWorldCommandCallbacks::GetRepeatLastPlayIcon ) ),
 		FName(TEXT("LevelToolbarPlay"))
-	);
+	));
 
 	// Play combo box
 	FUIAction SpecialPIEOptionsMenuAction;
 	SpecialPIEOptionsMenuAction.IsActionVisibleDelegate = FIsActionButtonVisible::CreateStatic( &FInternalPlayWorldCommandCallbacks::CanShowNonPlayWorldOnlyActions );
 
-	ToolbarBuilder.AddComboButton(
+	InSection.AddEntry(FToolMenuEntry::InitComboButton(
+		"PlayCombo",
 		SpecialPIEOptionsMenuAction,
 		FOnGetContent::CreateStatic( &GeneratePlayMenuContent, GlobalPlayWorldActions.ToSharedRef() ),
 		LOCTEXT( "PlayCombo_Label", "Active Play Mode" ),
 		LOCTEXT( "PIEComboToolTip", "Change Play Mode and Play Settings" ),
 		FSlateIcon(FEditorStyle::GetStyleSetName(), "LevelEditor.RepeatLastPlay"),
 		true
-	);
+	));
 
-	if (bIncludeLaunchButtonAndOptions && GetDefault<UEditorStyleSettings>()->bShowLaunchMenus)
+	if (bIncludeLaunchButtonAndOptions)
 	{
-		// Launch
-		ToolbarBuilder.AddToolBarButton( 
-			FPlayWorldCommands::Get().RepeatLastLaunch, 
-			NAME_None, 
-			LOCTEXT("RepeatLastLaunch", "Launch"),
-			TAttribute< FText >::Create( TAttribute< FText >::FGetter::CreateStatic( &FInternalPlayWorldCommandCallbacks::GetRepeatLastLaunchToolTip ) ),
-			TAttribute< FSlateIcon >::Create( TAttribute< FSlateIcon >::FGetter::CreateStatic( &FInternalPlayWorldCommandCallbacks::GetRepeatLastLaunchIcon ) ),
-			FName(TEXT("RepeatLastLaunch"))
-		);
-
-		// Launch combo box
-		FUIAction LaunchMenuAction;
-		LaunchMenuAction.IsActionVisibleDelegate = FIsActionButtonVisible::CreateStatic( &FInternalPlayWorldCommandCallbacks::CanShowNonPlayWorldOnlyActions );
-
-		ToolbarBuilder.AddComboButton(
-			LaunchMenuAction,
-			FOnGetContent::CreateStatic( &GenerateLaunchMenuContent, GlobalPlayWorldActions.ToSharedRef() ),
-			LOCTEXT( "LaunchCombo_Label", "Launch Options" ),
-			LOCTEXT( "PODComboToolTip", "Options for launching on a device" ),
-			FSlateIcon(FEditorStyle::GetStyleSetName(), "LevelEditor.RepeatLastLaunch"),
-			true
-		);
+		InSection.AddDynamicEntry("LaunchButtons", FNewToolMenuSectionDelegate::CreateLambda([](FToolMenuSection& InDynamicSection)
+		{
+			if (GetDefault<UEditorStyleSettings>()->bShowLaunchMenus)
+			{
+				// Launch
+				InDynamicSection.AddEntry(FToolMenuEntry::InitToolBarButton(
+					FPlayWorldCommands::Get().RepeatLastLaunch,
+					LOCTEXT("RepeatLastLaunch", "Launch"),
+					TAttribute< FText >::Create(TAttribute< FText >::FGetter::CreateStatic(&FInternalPlayWorldCommandCallbacks::GetRepeatLastLaunchToolTip)),
+					TAttribute< FSlateIcon >::Create(TAttribute< FSlateIcon >::FGetter::CreateStatic(&FInternalPlayWorldCommandCallbacks::GetRepeatLastLaunchIcon)),
+					FName(TEXT("RepeatLastLaunch"))
+				));
+
+				// Launch combo box
+				FUIAction LaunchMenuAction;
+				LaunchMenuAction.IsActionVisibleDelegate = FIsActionButtonVisible::CreateStatic(&FInternalPlayWorldCommandCallbacks::CanShowNonPlayWorldOnlyActions);
+
+				InDynamicSection.AddEntry(FToolMenuEntry::InitComboButton(
+					"LaunchCombo",
+					LaunchMenuAction,
+					FOnGetContent::CreateStatic(&GenerateLaunchMenuContent, GlobalPlayWorldActions.ToSharedRef()),
+					LOCTEXT("LaunchCombo_Label", "Launch Options"),
+					LOCTEXT("PODComboToolTip", "Options for launching on a device"),
+					FSlateIcon(FEditorStyle::GetStyleSetName(), "LevelEditor.RepeatLastLaunch"),
+					true
+				));
+			}
+		}));
 	}
 
 	// Resume/pause toggle (only one will be visible, and only in PIE/SIE)
-	ToolbarBuilder.AddToolBarButton(FPlayWorldCommands::Get().ResumePlaySession, NAME_None, TAttribute<FText>(),
+	InSection.AddEntry(FToolMenuEntry::InitToolBarButton(FPlayWorldCommands::Get().ResumePlaySession, TAttribute<FText>(),
 		TAttribute<FText>::Create(TAttribute<FText>::FGetter::CreateStatic(&FInternalPlayWorldCommandCallbacks::GetResumePlaySessionToolTip)),
 		TAttribute<FSlateIcon>::Create(TAttribute<FSlateIcon>::FGetter::CreateStatic(&FInternalPlayWorldCommandCallbacks::GetResumePlaySessionImage)),
 		FName(TEXT("ResumePlaySession"))
-	);
+	));
 
-	ToolbarBuilder.AddToolBarButton(FPlayWorldCommands::Get().PausePlaySession, NAME_None, TAttribute<FText>(), TAttribute<FText>(), TAttribute<FSlateIcon>(), FName(TEXT("PausePlaySession")));
-	ToolbarBuilder.AddToolBarButton(FPlayWorldCommands::Get().SingleFrameAdvance, NAME_None, TAttribute<FText>(), TAttribute<FText>(), TAttribute<FSlateIcon>(), FName(TEXT("SingleFrameAdvance")));
+	InSection.AddEntry(FToolMenuEntry::InitToolBarButton(FPlayWorldCommands::Get().PausePlaySession, TAttribute<FText>(), TAttribute<FText>(), TAttribute<FSlateIcon>(), FName(TEXT("PausePlaySession"))));
+	InSection.AddEntry(FToolMenuEntry::InitToolBarButton(FPlayWorldCommands::Get().SingleFrameAdvance, TAttribute<FText>(), TAttribute<FText>(), TAttribute<FSlateIcon>(), FName(TEXT("SingleFrameAdvance"))));
 
 	// Stop
-	ToolbarBuilder.AddToolBarButton(FPlayWorldCommands::Get().StopPlaySession, NAME_None, TAttribute<FText>(), TAttribute<FText>(), TAttribute<FSlateIcon>(), FName(TEXT("StopPlaySession")));
+	InSection.AddEntry(FToolMenuEntry::InitToolBarButton(FPlayWorldCommands::Get().StopPlaySession, TAttribute<FText>(), TAttribute<FText>(), TAttribute<FSlateIcon>(), FName(TEXT("StopPlaySession"))));
 
 	// Late Join
-	ToolbarBuilder.AddToolBarButton(FPlayWorldCommands::Get().LateJoinSession, NAME_None, TAttribute<FText>(), TAttribute<FText>(), TAttribute<FSlateIcon>(), FName(TEXT("LateJoinSession")));
+	InSection.AddEntry(FToolMenuEntry::InitToolBarButton(FPlayWorldCommands::Get().LateJoinSession, TAttribute<FText>(), TAttribute<FText>(), TAttribute<FSlateIcon>(), FName(TEXT("LateJoinSession"))));
 
 	// Eject/possess toggle
-	ToolbarBuilder.AddToolBarButton(FPlayWorldCommands::Get().PossessEjectPlayer, NAME_None, 
+	InSection.AddEntry(FToolMenuEntry::InitToolBarButton(FPlayWorldCommands::Get().PossessEjectPlayer,
 		TAttribute<FText>::Create(TAttribute<FText>::FGetter::CreateStatic(&FInternalPlayWorldCommandCallbacks::GetPossessEjectLabel)),
 		TAttribute<FText>::Create(TAttribute<FText>::FGetter::CreateStatic(&FInternalPlayWorldCommandCallbacks::GetPossessEjectTooltip)),
 		TAttribute<FSlateIcon>::Create(TAttribute<FSlateIcon>::FGetter::CreateStatic(&FInternalPlayWorldCommandCallbacks::GetPossessEjectImage)),
 		FName(TEXT("PossessEjectPlayer"))
-	);
+	));
 
 	// Single-stepping only buttons
-	ToolbarBuilder.AddToolBarButton(FPlayWorldCommands::Get().ShowCurrentStatement, NAME_None, TAttribute<FText>(), TAttribute<FText>(), TAttribute<FSlateIcon>(), FName(TEXT("ShowCurrentStatement")));
-	ToolbarBuilder.AddToolBarButton(FPlayWorldCommands::Get().StepInto, NAME_None, TAttribute<FText>(), TAttribute<FText>(), TAttribute<FSlateIcon>(), FName(TEXT("StepInto")));
-	ToolbarBuilder.AddToolBarButton(FPlayWorldCommands::Get().StepOver, NAME_None, TAttribute<FText>(), TAttribute<FText>(), TAttribute<FSlateIcon>(), FName(TEXT("StepOver")));
-	ToolbarBuilder.AddToolBarButton(FPlayWorldCommands::Get().StepOut, NAME_None, TAttribute<FText>(), TAttribute<FText>(), TAttribute<FSlateIcon>(), FName(TEXT("StepOut")));
+	InSection.AddEntry(FToolMenuEntry::InitToolBarButton(FPlayWorldCommands::Get().ShowCurrentStatement, TAttribute<FText>(), TAttribute<FText>(), TAttribute<FSlateIcon>(), FName(TEXT("ShowCurrentStatement"))));
+	InSection.AddEntry(FToolMenuEntry::InitToolBarButton(FPlayWorldCommands::Get().StepInto, TAttribute<FText>(), TAttribute<FText>(), TAttribute<FSlateIcon>(), FName(TEXT("StepInto"))));
+	InSection.AddEntry(FToolMenuEntry::InitToolBarButton(FPlayWorldCommands::Get().StepOver, TAttribute<FText>(), TAttribute<FText>(), TAttribute<FSlateIcon>(), FName(TEXT("StepOver"))));
+	InSection.AddEntry(FToolMenuEntry::InitToolBarButton(FPlayWorldCommands::Get().StepOut, TAttribute<FText>(), TAttribute<FText>(), TAttribute<FSlateIcon>(), FName(TEXT("StepOut"))));
 }
 
 // function will enumerate available Android devices that can export their profile to a json file
diff --git a/Engine/Source/Editor/UnrealEd/Private/Toolkits/AssetEditorToolkit.cpp b/Engine/Source/Editor/UnrealEd/Private/Toolkits/AssetEditorToolkit.cpp
--- a/Engine/Source/Editor/UnrealEd/Private/Toolkits/AssetEditorToolkit.cpp
+++ b/Engine/Source/Editor/UnrealEd/Private/Toolkits/AssetEditorToolkit.cpp
@@ -975,6 +975,14 @@ FName FAssetEditorToolkit::GetToolMenuAppName() const
 
 FName FAssetEditorToolkit::GetToolMenuToolbarName() const
 {
+	FName ParentName;
+	return GetToolMenuToolbarName(ParentName);
+}
+
+FName FAssetEditorToolkit::GetToolMenuToolbarName(FName& OutParentName) const
+{
+	static const FName DefaultToolbarName = "AssetEditor.DefaultToolBar";
+	OutParentName = DefaultToolbarName;
 	return *(TEXT("AssetEditor.") + GetToolMenuAppName().ToString() + TEXT(".ToolBar"));
 }
 
@@ -993,18 +1001,24 @@ void FAssetEditorToolkit::RegisterMenus()
 	}
 }
 
+void FAssetEditorToolkit::InitToolMenuContext(FToolMenuContext& MenuContext)
+{
+
+}
+
 void FAssetEditorToolkit::GenerateToolbar()
 {
 	TSharedPtr<FExtender> Extender = FExtender::Combine(ToolbarExtenders);
 
 	RegisterMenus();
 
-	const FName ToolBarName = GetToolMenuToolbarName();
+	FName ParentToolbarName;
+	const FName ToolBarName = GetToolMenuToolbarName(ParentToolbarName);
 	UToolMenus* ToolMenus = UToolMenus::Get();
 	UToolMenu* FoundMenu = ToolMenus->FindMenu(ToolBarName);
 	if (!FoundMenu || !FoundMenu->IsRegistered())
 	{
-		FoundMenu = ToolMenus->RegisterMenu(ToolBarName, "AssetEditor.DefaultToolBar", EMultiBoxType::ToolBar);
+		FoundMenu = ToolMenus->RegisterMenu(ToolBarName, ParentToolbarName, EMultiBoxType::ToolBar);
 	}
 
 	FToolMenuContext MenuContext(GetToolkitCommands(), Extender);
@@ -1013,6 +1027,8 @@ void FAssetEditorToolkit::GenerateToolbar()
 	ToolkitMenuContext->Toolkit = AsShared();
 	MenuContext.AddObject(ToolkitMenuContext);
 
+	InitToolMenuContext(MenuContext);
+
 	UToolMenu* GeneratedToolbar = ToolMenus->GenerateMenu(ToolBarName, MenuContext);
 	GeneratedToolbar->bToolBarIsFocusable = bIsToolbarFocusable;
 	GeneratedToolbar->bToolBarForceSmallIcons = bIsToolbarUsingSmallIcons;
diff --git a/Engine/Source/Editor/UnrealEd/Public/Kismet2/DebuggerCommands.h b/Engine/Source/Editor/UnrealEd/Public/Kismet2/DebuggerCommands.h
--- a/Engine/Source/Editor/UnrealEd/Public/Kismet2/DebuggerCommands.h
+++ b/Engine/Source/Editor/UnrealEd/Public/Kismet2/DebuggerCommands.h
@@ -11,6 +11,7 @@
 #include "Framework/Commands/UICommandList.h"
 
 class FToolBarBuilder;
+struct FToolMenuSection;
 
 /** This class acts as a generic widget that listens to and process global play world actions */
 class UNREALED_API SGlobalPlayWorldActions : public SCompoundWidget
@@ -52,7 +53,7 @@ public:
 	static void BindGlobalPlayWorldCommands();
 
 	/** Populates a toolbar with the menu commands for play-world control (pause/resume/stop/possess/eject/step/show current loc) */
-	UNREALED_API static void BuildToolbar( FToolBarBuilder& ToolBarBuilder, bool bIncludeLaunchButtonAndOptions = false );
+	UNREALED_API static void BuildToolbar( FToolMenuSection& InSection, bool bIncludeLaunchButtonAndOptions = false );
 
 	/**
 	* Return the active widget that processes play world actions for PIE
diff --git a/Engine/Source/Editor/UnrealEd/Public/Toolkits/AssetEditorToolkit.h b/Engine/Source/Editor/UnrealEd/Public/Toolkits/AssetEditorToolkit.h
--- a/Engine/Source/Editor/UnrealEd/Public/Toolkits/AssetEditorToolkit.h
+++ b/Engine/Source/Editor/UnrealEd/Public/Toolkits/AssetEditorToolkit.h
@@ -18,6 +18,7 @@
 
 class FAssetEditorModeManager;
 class FMenuBuilder;
+struct FToolMenuContext;
 class SBorder;
 class SStandaloneAssetEditorToolkitHost;
 
@@ -187,6 +188,8 @@ public:
 	void RegenerateMenusAndToolbars();
 
 	/** Get name used by tool menu */
+	virtual FName GetToolMenuToolbarName(FName& OutParentName) const;
+	virtual void InitToolMenuContext(FToolMenuContext& MenuContext);
 	FName GetToolMenuToolbarName() const;
 	FName GetToolMenuAppName() const;
 
