Viewing a blueprint actor in the Blueprint Editor Viewport no longer changes a custom thumbnail zoom level when the zoom level is really close to the actor.

diff --git a/Engine/Source/Editor/Kismet/Private/SCSEditorViewportClient.cpp b/Engine/Source/Editor/Kismet/Private/SCSEditorViewportClient.cpp
--- a/Engine/Source/Editor/Kismet/Private/SCSEditorViewportClient.cpp
+++ b/Engine/Source/Editor/Kismet/Private/SCSEditorViewportClient.cpp
@@ -761,16 +761,16 @@ void FSCSEditorViewportClient::ResetCamera()
 
 	// For now, loosely base default camera positioning on thumbnail preview settings
 	USceneThumbnailInfo* ThumbnailInfo = Cast<USceneThumbnailInfo>(Blueprint->ThumbnailInfo);
-	if(ThumbnailInfo)
+	if(ThumbnailInfo == nullptr)
 	{
-		if(PreviewActorBounds.SphereRadius + ThumbnailInfo->OrbitZoom < 0)
-		{
-			ThumbnailInfo->OrbitZoom = -PreviewActorBounds.SphereRadius;
-		}
+		ThumbnailInfo = USceneThumbnailInfo::StaticClass()->GetDefaultObject<USceneThumbnailInfo>();
 	}
-	else
+
+	// Clamp zoom to the actor's bounding sphere radius
+	float OrbitZoom = ThumbnailInfo->OrbitZoom;
+	if (PreviewActorBounds.SphereRadius + OrbitZoom < 0)
 	{
-		ThumbnailInfo = USceneThumbnailInfo::StaticClass()->GetDefaultObject<USceneThumbnailInfo>();
+		OrbitZoom = -PreviewActorBounds.SphereRadius;
 	}
 
 	ToggleOrbitCamera(true);
@@ -784,9 +784,8 @@ void FSCSEditorViewportClient::ResetCamera()
 		FRotator ThumbnailAngle(ThumbnailInfo->OrbitPitch, ThumbnailInfo->OrbitYaw, 0.0f);
 
 		SetViewLocationForOrbiting(PreviewActorBounds.Origin);
-		SetViewLocation( GetViewLocation() + FVector(0.0f, TargetDistance * 1.5f + ThumbnailInfo->OrbitZoom - AutoViewportOrbitCameraTranslate, 0.0f) );
+		SetViewLocation( GetViewLocation() + FVector(0.0f, TargetDistance * 1.5f + OrbitZoom - AutoViewportOrbitCameraTranslate, 0.0f) );
 		SetViewRotation( ThumbnailAngle );
-	
 	}
 
 	Invalidate();
