Fixed location of blueprint editor toolbar buttons
#rnx

diff --git a/Engine/Source/Editor/Kismet/Private/SBlueprintEditorToolbar.cpp b/Engine/Source/Editor/Kismet/Private/SBlueprintEditorToolbar.cpp
--- a/Engine/Source/Editor/Kismet/Private/SBlueprintEditorToolbar.cpp
+++ b/Engine/Source/Editor/Kismet/Private/SBlueprintEditorToolbar.cpp
@@ -402,7 +402,7 @@ static void BlueprintEditorToolbarImpl::MakeCompileDeveloperSubMenu(FMenuBuilder
 void FBlueprintEditorToolbar::AddBlueprintGlobalOptionsToolbar(UToolMenu* InMenu)
 {
 	FToolMenuSection& Section = InMenu->AddSection("Settings");
-	Section.InsertPosition = FToolMenuInsert("Asset", EToolMenuInsertType::Before);
+	Section.InsertPosition = FToolMenuInsert("Asset", EToolMenuInsertType::After);
 
 	Section.AddDynamicEntry("BlueprintGlobalOptions", FNewToolMenuSectionDelegate::CreateLambda([](FToolMenuSection& InSection)
 	{
