Fix extra toolbar buttons showing up in blueprint editor
#rnx

diff --git a/Engine/Source/Editor/Kismet/Private/BlueprintEditorModes.cpp b/Engine/Source/Editor/Kismet/Private/BlueprintEditorModes.cpp
--- a/Engine/Source/Editor/Kismet/Private/BlueprintEditorModes.cpp
+++ b/Engine/Source/Editor/Kismet/Private/BlueprintEditorModes.cpp
@@ -114,16 +114,6 @@ FBlueprintEditorApplicationMode::FBlueprintEditorApplicationMode(TSharedPtr<clas
 	CoreTabFactories.RegisterFactory(MakeShareable(new FSelectionDetailsSummoner(InBlueprintEditor)));
 
 	TabLayout = GetDefaltEditorLayout(InBlueprintEditor);
-	
-	// setup toolbar
-	//@TODO: Keep this in sync with AnimBlueprintMode.cpp
-	if (UToolMenu* Toolbar = InBlueprintEditor->RegisterModeToolbarIfUnregistered(GetModeName()))
-	{
-		InBlueprintEditor->GetToolbarBuilder()->AddCompileToolbar(Toolbar);
-		InBlueprintEditor->GetToolbarBuilder()->AddScriptingToolbar(Toolbar);
-		InBlueprintEditor->GetToolbarBuilder()->AddBlueprintGlobalOptionsToolbar(Toolbar);
-		InBlueprintEditor->GetToolbarBuilder()->AddDebuggingToolbar(Toolbar);
-	}
 
 	FBlueprintEditorModule& BlueprintEditorModule = FModuleManager::LoadModuleChecked<FBlueprintEditorModule>("Kismet");
 	BlueprintEditorModule.OnRegisterTabsForEditor().Broadcast(BlueprintEditorTabFactories, InModeName, InBlueprintEditor);
