Re-add blueprint editor find toolbar button
#rnx

diff --git a/Engine/Source/Editor/Kismet/Private/SBlueprintEditorToolbar.cpp b/Engine/Source/Editor/Kismet/Private/SBlueprintEditorToolbar.cpp
--- a/Engine/Source/Editor/Kismet/Private/SBlueprintEditorToolbar.cpp
+++ b/Engine/Source/Editor/Kismet/Private/SBlueprintEditorToolbar.cpp
@@ -484,6 +484,8 @@ void FBlueprintEditorToolbar::AddScriptingToolbar(UToolMenu* InMenu)
 		UBlueprintEditorToolMenuContext* Context = InSection.FindContext<UBlueprintEditorToolMenuContext>();
 		if (Context && Context->BlueprintEditor.IsValid() && Context->GetBlueprintObj())
 		{
+			InSection.AddEntry(FToolMenuEntry::InitToolBarButton(FBlueprintEditorCommands::Get().FindInBlueprint));
+
 			InSection.AddEntry(FToolMenuEntry::InitToolBarButton(
 				FBlueprintEditorCommands::Get().ToggleHideUnrelatedNodes,
 				TAttribute<FText>(),
