Converted code to use FObjectEditorUtils::GetCategoryText rather than access the UObjectCategory meta-data directly

This keeps the implementation common and will allow it to be changed in the future

#rnx

diff --git a/Engine/Plugins/FX/Niagara/Source/NiagaraEditor/Private/EdGraphSchema_Niagara.cpp b/Engine/Plugins/FX/Niagara/Source/NiagaraEditor/Private/EdGraphSchema_Niagara.cpp
--- a/Engine/Plugins/FX/Niagara/Source/NiagaraEditor/Private/EdGraphSchema_Niagara.cpp
+++ b/Engine/Plugins/FX/Niagara/Source/NiagaraEditor/Private/EdGraphSchema_Niagara.cpp
@@ -7,6 +7,7 @@
 #include "Framework/Commands/UIAction.h"
 #include "Framework/MultiBox/MultiBoxBuilder.h"
 #include "ToolMenus.h"
+#include "ObjectEditorUtils.h"
 #include "NiagaraCommon.h"
 #include "NiagaraEditorCommon.h"
 #include "INiagaraCompiler.h"
@@ -810,7 +811,7 @@ TArray<TSharedPtr<FNiagaraSchemaAction_NewNode> > UEdGraphSchema_Niagara::GetGra
 					FText MenuCat;
 					if (const UClass* Class = Type.GetClass())
 					{						
-						MenuCat = Class->GetMetaDataText(TEXT("Category"), TEXT("UObjectCategory"), Class->GetFullGroupName(false));
+						MenuCat = FObjectEditorUtils::GetCategoryText(Class);
 					}
 					else
 					{
diff --git a/Engine/Source/Editor/AIGraph/Private/AIGraphTypes.cpp b/Engine/Source/Editor/AIGraph/Private/AIGraphTypes.cpp
--- a/Engine/Source/Editor/AIGraph/Private/AIGraphTypes.cpp
+++ b/Engine/Source/Editor/AIGraph/Private/AIGraphTypes.cpp
@@ -12,6 +12,7 @@
 #include "Engine/Blueprint.h"
 #include "AssetData.h"
 #include "Editor.h"
+#include "ObjectEditorUtils.h"
 #include "Logging/MessageLog.h"
 #include "ARFilter.h"
 #include "AssetRegistryModule.h"
@@ -86,7 +87,7 @@ FString FGraphNodeClassData::GetDisplayName() const
 
 FText FGraphNodeClassData::GetCategory() const
 {
-	return Class.IsValid() ? Class->GetMetaDataText(TEXT("Category"), TEXT("UObjectCategory"), Class->GetFullGroupName(false)) : Category;
+	return Class.IsValid() ? FObjectEditorUtils::GetCategoryText(Class.Get()) : Category;
 }
 
 bool FGraphNodeClassData::IsAbstract() const
diff --git a/Engine/Source/Editor/AnimGraph/Private/AnimGraphDetails.cpp b/Engine/Source/Editor/AnimGraph/Private/AnimGraphDetails.cpp
--- a/Engine/Source/Editor/AnimGraph/Private/AnimGraphDetails.cpp
+++ b/Engine/Source/Editor/AnimGraph/Private/AnimGraphDetails.cpp
@@ -15,6 +15,7 @@
 #include "Widgets/Images/SImage.h"
 #include "AnimGraphNode_Root.h"
 #include "ScopedTransaction.h"
+#include "ObjectEditorUtils.h"
 #include "SKismetInspector.h"
 #include "AnimationGraph.h"
 #include "AnimationGraphSchema.h"
@@ -400,7 +401,7 @@ void FAnimGraphDetails::RefreshGroupSource()
 				
 		if(Function->HasMetaData(FBlueprintMetadata::MD_AnimBlueprintFunction))
 		{
-			FText Group = Function->GetMetaDataText(FBlueprintMetadata::MD_FunctionCategory, TEXT("UObjectCategory"), Function->GetFullGroupName(false));
+			FText Group = FObjectEditorUtils::GetCategoryText(Function);
 
 			if (!Group.IsEmpty())
 			{
diff --git a/Engine/Source/Editor/AnimGraph/Private/AnimationGraphSchema.cpp b/Engine/Source/Editor/AnimGraph/Private/AnimationGraphSchema.cpp
--- a/Engine/Source/Editor/AnimGraph/Private/AnimationGraphSchema.cpp
+++ b/Engine/Source/Editor/AnimGraph/Private/AnimationGraphSchema.cpp
@@ -9,6 +9,7 @@
 #include "Animation/AnimBlueprint.h"
 #include "Framework/MultiBox/MultiBoxBuilder.h"
 #include "ToolMenus.h"
+#include "ObjectEditorUtils.h"
 #include "K2Node.h"
 #include "EdGraphSchema_K2_Actions.h"
 #include "Kismet2/BlueprintEditorUtils.h"
@@ -275,7 +276,7 @@ void UAnimationGraphSchema::CreateFunctionGraphTerminators(UEdGraph& Graph, UCla
 			Graph.GetNodesOfClass<UAnimGraphNode_Root>(RootNodes);
 
 			check(RootNodes.Num() == 1);
-			RootNodes[0]->Node.Group = *InterfaceToImplement->GetMetaDataText(TEXT("Category"), TEXT("UObjectCategory"), InterfaceToImplement->GetFullGroupName(false)).ToString();
+			RootNodes[0]->Node.Group = *FObjectEditorUtils::GetCategoryText(InterfaceToImplement).ToString();
 
 			int32 CurrentPoseIndex = 0;
 			for (TFieldIterator<UProperty> PropIt(InterfaceToImplement); PropIt && (PropIt->PropertyFlags & CPF_Parm); ++PropIt)
@@ -705,7 +706,7 @@ void UAnimationGraphSchema::ConformAnimGraphToInterface(UBlueprint* InBlueprint,
 		InGraph.GetNodesOfClass<UAnimGraphNode_Root>(RootNodes);
 
 		check(RootNodes.Num() == 1);
-		RootNodes[0]->Node.Group = *InFunction->GetMetaDataText(TEXT("Category"), TEXT("UObjectCategory"), InFunction->GetFullGroupName(false)).ToString();
+		RootNodes[0]->Node.Group = *FObjectEditorUtils::GetCategoryText(InFunction).ToString();
 
 		TArray<UAnimGraphNode_LinkedInputPose*> LinkedInputPoseNodes;
 		InGraph.GetNodesOfClass<UAnimGraphNode_LinkedInputPose>(LinkedInputPoseNodes);
diff --git a/Engine/Source/Editor/BehaviorTreeEditor/Private/EdGraphSchema_BehaviorTree.cpp b/Engine/Source/Editor/BehaviorTreeEditor/Private/EdGraphSchema_BehaviorTree.cpp
--- a/Engine/Source/Editor/BehaviorTreeEditor/Private/EdGraphSchema_BehaviorTree.cpp
+++ b/Engine/Source/Editor/BehaviorTreeEditor/Private/EdGraphSchema_BehaviorTree.cpp
@@ -20,6 +20,7 @@
 #include "BehaviorTreeGraphNode_Task.h"
 #include "Framework/MultiBox/MultiBoxBuilder.h"
 #include "ToolMenus.h"
+#include "ObjectEditorUtils.h"
 #include "BehaviorTreeEditorModule.h"
 #include "IBehaviorTreeEditor.h"
 #include "BehaviorTreeDebugger.h"
@@ -72,7 +73,7 @@ void UEdGraphSchema_BehaviorTree::GetGraphNodeContextActions(FGraphContextMenuBu
 
 	if (SubNodeFlags == ESubNode::Decorator)
 	{
-		const FText& Category = UBehaviorTreeGraphNode_CompositeDecorator::StaticClass()->GetMetaDataText(TEXT("Category"), TEXT("UObjectCategory"), UBehaviorTreeGraphNode_CompositeDecorator::StaticClass()->GetFullGroupName(false));
+		const FText Category = FObjectEditorUtils::GetCategoryText(UBehaviorTreeGraphNode_CompositeDecorator::StaticClass());
 		UEdGraph* Graph = (UEdGraph*)ContextMenuBuilder.CurrentGraph;
 		UBehaviorTreeGraphNode_CompositeDecorator* OpNode = NewObject<UBehaviorTreeGraphNode_CompositeDecorator>(Graph);
 		TSharedPtr<FAISchemaAction_NewSubNode> AddOpAction = UAIGraphSchema::AddNewSubNodeAction(ContextMenuBuilder, Category, FText::FromString(OpNode->GetNodeTypeDescription()), FText::GetEmpty());
diff --git a/Engine/Source/Editor/BlueprintGraph/Private/K2Node_CallFunction.cpp b/Engine/Source/Editor/BlueprintGraph/Private/K2Node_CallFunction.cpp
--- a/Engine/Source/Editor/BlueprintGraph/Private/K2Node_CallFunction.cpp
+++ b/Engine/Source/Editor/BlueprintGraph/Private/K2Node_CallFunction.cpp
@@ -36,6 +36,7 @@
 #include "BlueprintActionFilter.h"
 #include "FindInBlueprintManager.h"
 #include "ScopedTransaction.h"
+#include "ObjectEditorUtils.h"
 #include "SPinTypeSelector.h"
 #include "SourceCodeNavigation.h"
 #include "HAL/FileManager.h"
@@ -1656,7 +1657,7 @@ FText UK2Node_CallFunction::GetDefaultCategoryForFunction(const UFunction* Funct
 		else
 		{
 			// Look for localized metadata
-			FuncCategory = Function->GetMetaDataText(FBlueprintMetadata::MD_FunctionCategory, TEXT("UObjectCategory"), Function->GetFullGroupName(false));
+			FuncCategory = FObjectEditorUtils::GetCategoryText(Function);
 
 			// If the result is culture invariant, force it into a display string
 			if (FuncCategory.IsCultureInvariant())
diff --git a/Engine/Source/Editor/Kismet/Private/BlueprintDetailsCustomization.cpp b/Engine/Source/Editor/Kismet/Private/BlueprintDetailsCustomization.cpp
--- a/Engine/Source/Editor/Kismet/Private/BlueprintDetailsCustomization.cpp
+++ b/Engine/Source/Editor/Kismet/Private/BlueprintDetailsCustomization.cpp
@@ -1427,7 +1427,7 @@ void FBlueprintVarActionDetails::PopulateCategories(SMyBlueprint* MyBlueprint, T
 	{
 		if(UFunction* Function = Blueprint->SkeletonGeneratedClass->FindFunctionByName(FunctionGraph->GetFName()))
 		{
-			FText FunctionCategory = Function->GetMetaDataText(FBlueprintMetadata::MD_FunctionCategory, TEXT("UObjectCategory"), Function->GetFullGroupName(false));
+			FText FunctionCategory = FObjectEditorUtils::GetCategoryText(Function);
 
 			if(!FunctionCategory.IsEmpty())
 			{
@@ -1490,7 +1490,7 @@ void FBlueprintVarActionDetails::PopulateCategories(SMyBlueprint* MyBlueprint, T
 
 		if (UEdGraphSchema_K2::CanKismetOverrideFunction(Function) && !UEdGraphSchema_K2::FunctionCanBePlacedAsEvent(Function))
 		{
-			FText FunctionCategory = Function->GetMetaDataText(FBlueprintMetadata::MD_FunctionCategory, TEXT("UObjectCategory"), Function->GetFullGroupName(false));
+			FText FunctionCategory = FObjectEditorUtils::GetCategoryText(Function);
 
 			if (!FunctionCategory.IsEmpty())
 			{
diff --git a/Engine/Source/Editor/Kismet/Private/SMyBlueprint.cpp b/Engine/Source/Editor/Kismet/Private/SMyBlueprint.cpp
--- a/Engine/Source/Editor/Kismet/Private/SMyBlueprint.cpp
+++ b/Engine/Source/Editor/Kismet/Private/SMyBlueprint.cpp
@@ -1238,7 +1238,7 @@ void SMyBlueprint::CollectAllActions(FGraphActionListBuilderBase& OutAllActions)
 			UFunction* Function = BlueprintObj->SkeletonGeneratedClass->FindFunctionByName(Graph->GetFName());
 			if (Function != nullptr)
 			{
-				FunctionCategory = Function->GetMetaDataText(FBlueprintMetadata::MD_FunctionCategory, TEXT("UObjectCategory"), Function->GetFullGroupName(false));
+				FunctionCategory = FObjectEditorUtils::GetCategoryText(Function);
 			}
 		}
 
@@ -1332,7 +1332,7 @@ void SMyBlueprint::CollectAllActions(FGraphActionListBuilderBase& OutAllActions)
 				FunctionDesc = FBlueprintEditorUtils::GetDeprecatedMemberMenuItemName(FunctionDesc);
 			}
 
-			FText FunctionCategory = Function->GetMetaDataText(FBlueprintMetadata::MD_FunctionCategory, TEXT("UObjectCategory"), Function->GetFullGroupName(false));
+			FText FunctionCategory = FObjectEditorUtils::GetCategoryText(Function);
 
 			TSharedPtr<FEdGraphSchemaAction_K2Graph> NewFuncAction = MakeShareable(new FEdGraphSchemaAction_K2Graph(EEdGraphSchemaAction_K2Graph::Function, FunctionCategory, FunctionDesc, FunctionTooltip, 1, NodeSectionID::FUNCTION_OVERRIDABLE));
 			NewFuncAction->FuncName = FunctionName;
@@ -1348,7 +1348,7 @@ void SMyBlueprint::CollectAllActions(FGraphActionListBuilderBase& OutAllActions)
 		{
 			if (UFunction * Function = BlueprintObj->SkeletonGeneratedClass->FindFunctionByName(FunctionName))
 			{
-				FunctionCategory = Function->GetMetaDataText(FBlueprintMetadata::MD_FunctionCategory, TEXT("UObjectCategory"), Function->GetFullGroupName(false));
+				FunctionCategory = FObjectEditorUtils::GetCategoryText(Function);
 
 				if (IAnimClassInterface * AnimClassInterface = IAnimClassInterface::GetFromClass(BlueprintObj->SkeletonGeneratedClass))
 				{
@@ -1431,7 +1431,7 @@ void SMyBlueprint::CollectAllActions(FGraphActionListBuilderBase& OutAllActions)
 					FText FunctionTooltip = Function->GetToolTipText();
 					FText FunctionDesc = K2Schema->GetFriendlySignatureName(Function);
 
-					FText FunctionCategory = Function->GetMetaDataText(FBlueprintMetadata::MD_FunctionCategory, TEXT("UObjectCategory"), Function->GetFullGroupName(false));
+					FText FunctionCategory = FObjectEditorUtils::GetCategoryText(Function);
 					bool bIsAnimFunction = IsInAnimBPLambda(FunctionName, FunctionCategory);
 
 					TSharedPtr<FEdGraphSchemaAction_K2Graph> NewFuncAction = MakeShareable(new FEdGraphSchemaAction_K2Graph(EEdGraphSchemaAction_K2Graph::Interface, FunctionCategory, FunctionDesc, FunctionTooltip, 1, bIsAnimFunction ? NodeSectionID::ANIMLAYER : NodeSectionID::INTERFACE));
diff --git a/Engine/Source/Editor/PropertyEditor/Private/DetailLayoutHelpers.cpp b/Engine/Source/Editor/PropertyEditor/Private/DetailLayoutHelpers.cpp
--- a/Engine/Source/Editor/PropertyEditor/Private/DetailLayoutHelpers.cpp
+++ b/Engine/Source/Editor/PropertyEditor/Private/DetailLayoutHelpers.cpp
@@ -185,17 +185,13 @@ namespace DetailLayoutHelpers
 								FText CategoryDisplayName = CategoryImpl.GetDisplayName();
 								if (CategoryDisplayName.IsEmpty() || CategoryDisplayName.IsCultureInvariant())
 								{
-									const FString CategorySourceText = CategoryName.ToString();
-									const FString CategoryKey = Property->GetFullGroupName(false);
+									CategoryDisplayName = FObjectEditorUtils::GetCategoryText(Property);
 
-									if (FText::FindText(TEXT("UObjectCategory"), CategoryKey, /*OUT*/CategoryDisplayName, &CategorySourceText))
+									// Category names in English are typically gathered in their non-pretty form (eg "UserInterface" rather than "User Interface"), so skip 
+									// applying the localized variant if the text matches the raw category name, as in this case the pretty printer will do a better job
+									if (!CategoryName.ToString().Equals(CategoryDisplayName.ToString(), ESearchCase::CaseSensitive))
 									{
-										// Category names in English are typically gathered in their non-pretty form (eg "UserInterface" rather than "User Interface"), so skip 
-										// applying the localized variant if the text matches the raw category name, as in this case the pretty printer will do a better job
-										if (!CategorySourceText.Equals(CategoryDisplayName.ToString(), ESearchCase::CaseSensitive))
-										{
-											CategoryImpl.SetDisplayName(CategoryName, CategoryDisplayName);
-										}
+										CategoryImpl.SetDisplayName(CategoryName, CategoryDisplayName);
 									}
 								}
 							}
diff --git a/Engine/Source/Runtime/CoreUObject/Private/UObject/Class.cpp b/Engine/Source/Runtime/CoreUObject/Private/UObject/Class.cpp
--- a/Engine/Source/Runtime/CoreUObject/Private/UObject/Class.cpp
+++ b/Engine/Source/Runtime/CoreUObject/Private/UObject/Class.cpp
@@ -440,7 +440,7 @@ const FString& UField::GetMetaData(const FName& Key) const
 	return MetaDataString;
 }
 
-const FText UField::GetMetaDataText(const TCHAR* MetaDataKey, const FString LocalizationNamespace, const FString LocalizationKey) const
+FText UField::GetMetaDataText(const TCHAR* MetaDataKey, const FString LocalizationNamespace, const FString LocalizationKey) const
 {
 	FString DefaultMetaData;
 
@@ -450,12 +450,11 @@ const FText UField::GetMetaDataText(const TCHAR* MetaDataKey, const FString Loca
 	}
 
 	// If attempting to grab the DisplayName metadata, we must correct the source string and output it as a DisplayString for lookup
-	if( DefaultMetaData.IsEmpty() && FString(MetaDataKey) == TEXT("DisplayName") )
+	if( DefaultMetaData.IsEmpty() && FCString::Stricmp(MetaDataKey, TEXT("DisplayName")) == 0 )
 	{
 		DefaultMetaData = FName::NameToDisplayString( GetName(), IsA( UBoolProperty::StaticClass() ) );
 	}
 
-
 	FText LocalizedMetaData;
 	if ( !( FText::FindText( LocalizationNamespace, LocalizationKey, /*OUT*/LocalizedMetaData, &DefaultMetaData ) ) )
 	{
@@ -468,7 +467,7 @@ const FText UField::GetMetaDataText(const TCHAR* MetaDataKey, const FString Loca
 	return LocalizedMetaData;
 }
 
-const FText UField::GetMetaDataText(const FName& MetaDataKey, const FString LocalizationNamespace, const FString LocalizationKey) const
+FText UField::GetMetaDataText(const FName& MetaDataKey, const FString LocalizationNamespace, const FString LocalizationKey) const
 {
 	FString DefaultMetaData;
 
diff --git a/Engine/Source/Runtime/CoreUObject/Public/UObject/Class.h b/Engine/Source/Runtime/CoreUObject/Public/UObject/Class.h
--- a/Engine/Source/Runtime/CoreUObject/Public/UObject/Class.h
+++ b/Engine/Source/Runtime/CoreUObject/Public/UObject/Class.h
@@ -142,8 +142,8 @@ class COREUOBJECT_API UField : public UObject
 	 * @param LocalizationKey			Key to lookup in the localization manager
 	 * @return							Localized metadata if available, defaults to whatever is provided via GetMetaData
 	 */
-	const FText GetMetaDataText(const TCHAR* MetaDataKey, const FString LocalizationNamespace = FString(), const FString LocalizationKey = FString()) const;
-	const FText GetMetaDataText(const FName& MetaDataKey, const FString LocalizationNamespace = FString(), const FString LocalizationKey = FString()) const;
+	FText GetMetaDataText(const TCHAR* MetaDataKey, const FString LocalizationNamespace = FString(), const FString LocalizationKey = FString()) const;
+	FText GetMetaDataText(const FName& MetaDataKey, const FString LocalizationNamespace = FString(), const FString LocalizationKey = FString()) const;
 
 	/**
 	 * Sets the metadata value associated with the key
diff --git a/Engine/Source/Runtime/Engine/Private/Animation/AnimBlueprintGeneratedClass.cpp b/Engine/Source/Runtime/Engine/Private/Animation/AnimBlueprintGeneratedClass.cpp
--- a/Engine/Source/Runtime/Engine/Private/Animation/AnimBlueprintGeneratedClass.cpp
+++ b/Engine/Source/Runtime/Engine/Private/Animation/AnimBlueprintGeneratedClass.cpp
@@ -7,6 +7,7 @@
 #include "Animation/AnimBlueprintGeneratedClass.h"
 #include "Serialization/ObjectWriter.h"
 #include "Serialization/ObjectReader.h"
+#include "ObjectEditorUtils.h"
 #include "BonePose.h"
 #include "Animation/AnimNodeBase.h"
 #include "Animation/AnimInstance.h"
@@ -375,7 +376,7 @@ void UAnimBlueprintGeneratedClass::GenerateAnimationBlueprintFunctions()
 		bool bFoundOutput = false;
 #if WITH_EDITOR
 		// In editor we can grab the group from metadata, otherwise we need to wait until CDO post load (LinkFunctionsToDefaultObjectNodes)
-		FText CategoryText = It->GetMetaDataText(TEXT("Category"), TEXT("UObjectCategory"), It->GetFullGroupName(false));
+		FText CategoryText = FObjectEditorUtils::GetCategoryText(*It);
 		FName Group = CategoryText.IsEmpty() ? NAME_None : FName(*CategoryText.ToString());
 #endif
 		UStructProperty* OutputPoseNodeProperty = nullptr;
diff --git a/Engine/Source/Runtime/Engine/Private/ObjectEditorUtils.cpp b/Engine/Source/Runtime/Engine/Private/ObjectEditorUtils.cpp
--- a/Engine/Source/Runtime/Engine/Private/ObjectEditorUtils.cpp
+++ b/Engine/Source/Runtime/Engine/Private/ObjectEditorUtils.cpp
@@ -10,36 +10,29 @@
 namespace FObjectEditorUtils
 {
 
-	FText GetCategoryText( const UProperty* InProperty )
+	FText GetCategoryText( const UField* InField )
 	{
 		static const FName NAME_Category(TEXT("Category"));
-		if (InProperty && InProperty->HasMetaData(NAME_Category))
+		if (InField && InField->HasMetaData(NAME_Category))
 		{
-			return InProperty->GetMetaDataText(NAME_Category, TEXT("UObjectCategory"), InProperty->GetFullGroupName(false));
-		}
-		else
-		{
-			return FText::GetEmpty();
+			return InField->GetMetaDataText(NAME_Category, TEXT("UObjectCategory"), InField->GetFullGroupName(false));
 		}
+		return FText::GetEmpty();
 	}
 
-	FString GetCategory( const UProperty* InProperty )
+	FString GetCategory( const UField* InField )
 	{
-		return GetCategoryText(InProperty).ToString();
+		return GetCategoryText(InField).ToString();
 	}
 
-
-	FName GetCategoryFName( const UProperty* InProperty )
+	FName GetCategoryFName( const UField* InField )
 	{
-		FName OutCategoryName( NAME_None );
-
-		static const FName CategoryKey( TEXT("Category") );
-		if( InProperty && InProperty->HasMetaData( CategoryKey ) )
+		static const FName CategoryKey(TEXT("Category"));
+		if (InField && InField->HasMetaData(CategoryKey))
 		{
-			OutCategoryName = FName( *InProperty->GetMetaData( CategoryKey ) );
+			return FName(*InField->GetMetaData(CategoryKey));
 		}
-
-		return OutCategoryName;
+		return NAME_None;
 	}
 
 	bool IsFunctionHiddenFromClass( const UFunction* InFunction,const UClass* Class )
diff --git a/Engine/Source/Runtime/Engine/Public/ObjectEditorUtils.h b/Engine/Source/Runtime/Engine/Public/ObjectEditorUtils.h
--- a/Engine/Source/Runtime/Engine/Public/ObjectEditorUtils.h
+++ b/Engine/Source/Runtime/Engine/Public/ObjectEditorUtils.h
@@ -6,36 +6,36 @@
 #include "UObject/UnrealType.h"
 
 #if WITH_EDITOR
+class UField;
 class UProperty;
 
 namespace FObjectEditorUtils
 {
 
 	/**
-	 * Gets the category this property belongs to.
+	 * Gets the category this field belongs to.
 	 *
-	 * @param	InProperty	Property we want the category name of.
-	 * @return	Category name of the given property.
+	 * @param	InField	Field we want the category name of.
+	 * @return	Category name of the given field.
 	 */
-	ENGINE_API FText GetCategoryText( const class UProperty* InProperty );	
+	ENGINE_API FText GetCategoryText( const UField* InField );	
 
 	/**
-	 * Gets the category this property belongs to.
+	 * Gets the category this field belongs to.
 	 *
-	 * @param	InProperty	Property we want the category name of.
-	 * @return	Category name of the given property.
+	 * @param	InField	Field we want the category name of.
+	 * @return	Category name of the given field.
 	 */
-	ENGINE_API FString GetCategory( const class UProperty* InProperty );	
+	ENGINE_API FString GetCategory( const UField* InField );
 
 	/**
-	 * Gets the FName of the category this property belongs to.  
+	 * Gets the FName of the category this field belongs to.  
 	 * Note, this value is suitable for comparison against other categories but NOT suitable as a display name since it is not localized
 	 *
-	 * @param	InProperty	Property we want the category name of.
-	 * @return	Category name of the given property.
+	 * @param	InField	Field we want the category name of.
+	 * @return	Category name of the given field.
 	 */
-	ENGINE_API FName GetCategoryFName( const class UProperty* InProperty );
-
+	ENGINE_API FName GetCategoryFName( const UField* InField );
 
 	/**
 	 * Query if a function is flagged as hidden from the given class either by category or by function name
