Editor - Fix clipping in the SSCSEditor's listview.

diff --git a/Engine/Source/Editor/Kismet/Public/SSCSEditor.h b/Engine/Source/Editor/Kismet/Public/SSCSEditor.h
--- a/Engine/Source/Editor/Kismet/Public/SSCSEditor.h
+++ b/Engine/Source/Editor/Kismet/Public/SSCSEditor.h
@@ -712,7 +712,9 @@ public:
 		, _ClearSelectionOnClick(true)
 		, _ExternalScrollbar()
 		, _OnTableViewBadState()
-		{}
+		{
+			_Clipping = EWidgetClipping::ClipToBounds;
+		}
 
 		SLATE_ARGUMENT( SSCSEditor*, SCSEditor )
 
