Reduce the overall time spent constructing search result nodes during a global Blueprint search.

diff --git a/Engine/Config/BaseEditor.ini b/Engine/Config/BaseEditor.ini
--- a/Engine/Config/BaseEditor.ini
+++ b/Engine/Config/BaseEditor.ini
@@ -433,6 +433,7 @@ ForceUnits=EUnit::Newtons
 bDisableDeferredIndexing=false
 bEnableCsvStatsProfiling=false
 bEnableDeveloperMenuTools=false
+bDisableSearchResultTemplates=false
 
 [BlueprintNativizationSettings]
 bUseStaticVariablesInWrappers=true
diff --git a/Engine/Source/Editor/Kismet/Private/FindInBlueprintManager.cpp b/Engine/Source/Editor/Kismet/Private/FindInBlueprintManager.cpp
--- a/Engine/Source/Editor/Kismet/Private/FindInBlueprintManager.cpp
+++ b/Engine/Source/Editor/Kismet/Private/FindInBlueprintManager.cpp
@@ -1586,6 +1586,7 @@ FFindInBlueprintSearchManager::FFindInBlueprintSearchManager()
 	, bDisableDeferredIndexing(false)
 	, bEnableCSVStatsProfiling(false)
 	, bEnableDeveloperMenuTools(false)
+	, bDisableSearchResultTemplates(false)
 	, bIsPausing(false)
 {
 	for (int32 TabIdx = 0; TabIdx < UE_ARRAY_COUNT(GlobalFindResultsTabIDs); TabIdx++)
@@ -1625,6 +1626,7 @@ void FFindInBlueprintSearchManager::Initialize()
 	GConfig->GetBool(TEXT("BlueprintSearchSettings"), TEXT("bDisableDeferredIndexing"), bDisableDeferredIndexing, GEditorIni);
 	GConfig->GetBool(TEXT("BlueprintSearchSettings"), TEXT("bEnableCsvStatsProfiling"), bEnableCSVStatsProfiling, GEditorIni);
 	GConfig->GetBool(TEXT("BlueprintSearchSettings"), TEXT("bEnableDeveloperMenuTools"), bEnableDeveloperMenuTools, GEditorIni);
+	GConfig->GetBool(TEXT("BlueprintSearchSettings"), TEXT("bDisableSearchResultTemplates"), bDisableSearchResultTemplates, GEditorIni);
 
 	// If profiling has been enabled, turn on the stat category and begin a capture.
 	if (bEnableCSVStatsProfiling)
@@ -1637,7 +1639,6 @@ void FFindInBlueprintSearchManager::Initialize()
 		}
 	}
 
-
 	// Must ensure we do not attempt to load the AssetRegistry Module while saving a package, however, if it is loaded already we can safely obtain it
 	if (!GIsSavingPackage || (GIsSavingPackage && FModuleManager::Get().IsModuleLoaded(TEXT("AssetRegistry"))))
 	{
diff --git a/Engine/Source/Editor/Kismet/Private/FindInBlueprints.cpp b/Engine/Source/Editor/Kismet/Private/FindInBlueprints.cpp
--- a/Engine/Source/Editor/Kismet/Private/FindInBlueprints.cpp
+++ b/Engine/Source/Editor/Kismet/Private/FindInBlueprints.cpp
@@ -130,11 +130,6 @@ FFindInBlueprintsResult::FFindInBlueprintsResult(const FText& InDisplayText )
 {
 }
 
-FFindInBlueprintsResult::FFindInBlueprintsResult( const FText& InDisplayText, TSharedPtr<FFindInBlueprintsResult> InParent) 
-	: Parent(InParent), DisplayText(InDisplayText)
-{
-}
-
 FReply FFindInBlueprintsResult::OnClick()
 {
 	// If there is a parent, handle it using the parent's functionality
@@ -218,9 +213,8 @@ FText FFindInBlueprintsResult::GetDisplayString() const
 //////////////////////////////////////////////////////////
 // FFindInBlueprintsGraphNode
 
-FFindInBlueprintsGraphNode::FFindInBlueprintsGraphNode(const FText& InValue, TSharedPtr<FFindInBlueprintsResult> InParent)
-	: FFindInBlueprintsResult(InValue, InParent)
-	, Glyph("EditorStyle", "")
+FFindInBlueprintsGraphNode::FFindInBlueprintsGraphNode()
+	: Glyph("EditorStyle", "")
 	, Class(nullptr)
 {
 }
@@ -326,9 +320,8 @@ UObject* FFindInBlueprintsGraphNode::GetObject(UBlueprint* InBlueprint) const
 //////////////////////////////////////////////////////////
 // FFindInBlueprintsPin
 
-FFindInBlueprintsPin::FFindInBlueprintsPin(const FText& InValue, TSharedPtr<FFindInBlueprintsResult> InParent, FString InSchemaName)
-	: FFindInBlueprintsResult(InValue, InParent)
-	, SchemaName(InSchemaName)
+FFindInBlueprintsPin::FFindInBlueprintsPin(FString InSchemaName)
+	: SchemaName(InSchemaName)
 	, IconColor(FSlateColor::UseForeground())
 {
 }
@@ -405,9 +398,8 @@ void FFindInBlueprintsPin::FinalizeSearchData()
 //////////////////////////////////////////////////////////
 // FFindInBlueprintsProperty
 
-FFindInBlueprintsProperty::FFindInBlueprintsProperty(const FText& InValue, TSharedPtr<FFindInBlueprintsResult> InParent)
-	: FFindInBlueprintsResult(InValue, InParent)
-	, bIsSCSComponent(false)
+FFindInBlueprintsProperty::FFindInBlueprintsProperty()
+	: bIsSCSComponent(false)
 {
 }
 
@@ -512,9 +504,8 @@ void FFindInBlueprintsProperty::FinalizeSearchData()
 //////////////////////////////////////////////////////////
 // FFindInBlueprintsGraph
 
-FFindInBlueprintsGraph::FFindInBlueprintsGraph(const FText& InValue, TSharedPtr<FFindInBlueprintsResult> InParent, EGraphType InGraphType)
-	: FFindInBlueprintsResult(InValue, InParent)
-	, GraphType(InGraphType)
+FFindInBlueprintsGraph::FFindInBlueprintsGraph(EGraphType InGraphType)
+	: GraphType(InGraphType)
 {
 }
 
diff --git a/Engine/Source/Editor/Kismet/Private/ImaginaryBlueprintData.cpp b/Engine/Source/Editor/Kismet/Private/ImaginaryBlueprintData.cpp
--- a/Engine/Source/Editor/Kismet/Private/ImaginaryBlueprintData.cpp
+++ b/Engine/Source/Editor/Kismet/Private/ImaginaryBlueprintData.cpp
@@ -75,17 +75,24 @@ FSearchResult FImaginaryFiBData::CreateSearchResult(FSearchResult InParent) cons
 {
 	CSV_SCOPED_TIMING_STAT(FindInBlueprint, CreateSearchResult);
 
-	FSearchResult ReturnSearchResult = CreateSearchResult_Internal(InParent);
+	FSearchResult ReturnSearchResult = CreateSearchResult_Internal(SearchResultTemplate);
+	if (ReturnSearchResult.IsValid())
+	{	
+		ReturnSearchResult->Parent = InParent;
 
-	FText DisplayName;
-	for( const TPair<FindInBlueprintsHelpers::FSimpleFTextKeyStorage, FSearchableValueInfo>& TagsAndValues : ParsedTagsAndValues )
-	{
-		if (TagsAndValues.Value.IsCoreDisplay() || !TagsAndValues.Value.IsSearchable())
+		if (!FFindInBlueprintSearchManager::Get().ShouldEnableSearchResultTemplates())
 		{
-			FText Value = TagsAndValues.Value.GetDisplayText(*LookupTablePtr);
-			ReturnSearchResult->ParseSearchInfo(TagsAndValues.Key.Text, Value);
+			for (const TPair<FindInBlueprintsHelpers::FSimpleFTextKeyStorage, FSearchableValueInfo>& TagsAndValues : ParsedTagsAndValues)
+			{
+				if (TagsAndValues.Value.IsCoreDisplay() || !TagsAndValues.Value.IsSearchable())
+				{
+					FText Value = TagsAndValues.Value.GetDisplayText(*LookupTablePtr);
+					ReturnSearchResult->ParseSearchInfo(TagsAndValues.Key.Text, Value);
+				}
+			}
 		}
 	}
+	
 	return ReturnSearchResult;
 }
 
@@ -195,7 +202,28 @@ void FImaginaryFiBData::ParseAllChildData_Internal(ESearchableValueStatus InSear
 			}
 			if (!TrySpecialHandleJsonValue(KeyText, JsonValue))
 			{
-				ParseJsonValue(KeyText, KeyText, JsonValue, false, InSearchabilityOverride);
+				TArray<FSearchableValueInfo> ParsedValues;
+				ParseJsonValue(KeyText, KeyText, JsonValue, ParsedValues, false, InSearchabilityOverride);
+
+				if (FFindInBlueprintSearchManager::Get().ShouldEnableSearchResultTemplates())
+				{
+					for (const FSearchableValueInfo& ParsedValue : ParsedValues)
+					{
+						if (ParsedValue.IsCoreDisplay() || !ParsedValue.IsSearchable())
+						{
+							// If necessary, create the search result template.
+							if (!SearchResultTemplate.IsValid())
+							{
+								FSearchResult NullTemplate;
+								SearchResultTemplate = CreateSearchResult_Internal(NullTemplate);
+								check(SearchResultTemplate.IsValid());
+							}
+
+							// Parse out meta values used for display and cache them in the template.
+							SearchResultTemplate->ParseSearchInfo(KeyText, ParsedValue.GetDisplayText(*LookupTablePtr));
+						}
+					}
+				}
 			}
 		}
 	}
@@ -212,18 +240,10 @@ void FImaginaryFiBData::ParseAllChildData(ESearchableValueStatus InSearchability
 	ParseAllChildData_Internal(InSearchabilityOverride);
 }
 
-void FImaginaryFiBData::ParseJsonValue(FText InKey, FText InDisplayKey, TSharedPtr< FJsonValue > InJsonValue, bool bIsInArray/*=false*/, ESearchableValueStatus InSearchabilityOverride/* = ESearchableValueStatus::Searchable*/)
+void FImaginaryFiBData::ParseJsonValue(FText InKey, FText InDisplayKey, TSharedPtr< FJsonValue > InJsonValue, TArray<FSearchableValueInfo>& OutParsedValues, bool bIsInArray/*=false*/, ESearchableValueStatus InSearchabilityOverride/* = ESearchableValueStatus::Searchable*/)
 {
 	ESearchableValueStatus SearchabilityStatus = (InSearchabilityOverride == ESearchableValueStatus::Searchable)? GetSearchabilityStatus(InKey.ToString()) : InSearchabilityOverride;
-	if( InJsonValue->Type == EJson::String)
-	{
-		ParsedTagsAndValues.Add(FindInBlueprintsHelpers::FSimpleFTextKeyStorage(InKey), FSearchableValueInfo(InDisplayKey, FCString::Atoi(*InJsonValue->AsString()), SearchabilityStatus));
-	}
-	else if (InJsonValue->Type == EJson::Boolean)
-	{
-		ParsedTagsAndValues.Add(FindInBlueprintsHelpers::FSimpleFTextKeyStorage(InKey), FSearchableValueInfo(InDisplayKey, FText::FromString(InJsonValue->AsString()), SearchabilityStatus));
-	}
-	else if( InJsonValue->Type == EJson::Array)
+	if( InJsonValue->Type == EJson::Array)
 	{
 		TSharedPtr< FCategorySectionHelper, ESPMode::ThreadSafe > ArrayCategory = MakeShareable(new FCategorySectionHelper(AsShared(), LookupTablePtr, InKey, true));
 		ParsedChildData.Add(ArrayCategory);
@@ -231,19 +251,29 @@ void FImaginaryFiBData::ParseJsonValue(FText InKey, FText InDisplayKey, TSharedP
 		for( int32 ArrayIdx = 0; ArrayIdx < ArrayList.Num(); ++ArrayIdx)
 		{
 			TSharedPtr< FJsonValue > ArrayValue = ArrayList[ArrayIdx];
-			ArrayCategory->ParseJsonValue(InKey, FText::FromString(FString::FromInt(ArrayIdx)), ArrayValue,/*bIsInArray=*/true, SearchabilityStatus);
+			ArrayCategory->ParseJsonValue(InKey, FText::FromString(FString::FromInt(ArrayIdx)), ArrayValue, OutParsedValues, /*bIsInArray=*/true, SearchabilityStatus);
 		}		
 	}
 	else if (InJsonValue->Type == EJson::Object)
 	{
 		TSharedPtr< FCategorySectionHelper, ESPMode::ThreadSafe > SubObjectCategory = MakeShareable(new FCategorySectionHelper(AsShared(), InJsonValue->AsObject(), LookupTablePtr, InDisplayKey, bIsInArray));
-		SubObjectCategory->ParseAllChildData_Internal(SearchabilityStatus);
+		SubObjectCategory->ParseAllChildData(SearchabilityStatus);
 		ParsedChildData.Add(SubObjectCategory);
 	}
 	else
 	{
-		// For everything else, there's this. Numbers come here and will be treated as strings
-		ParsedTagsAndValues.Add(FindInBlueprintsHelpers::FSimpleFTextKeyStorage(InKey), FSearchableValueInfo(InDisplayKey, FText::FromString(InJsonValue->AsString()), SearchabilityStatus));
+		FSearchableValueInfo& ParsedValue = OutParsedValues.AddDefaulted_GetRef();
+		if (InJsonValue->Type == EJson::String)
+		{
+			ParsedValue = FSearchableValueInfo(InDisplayKey, FCString::Atoi(*InJsonValue->AsString()), SearchabilityStatus);
+		}
+		else
+		{
+			// For everything else, there's this. Numbers come here and will be treated as strings
+			ParsedValue = FSearchableValueInfo(InDisplayKey, FText::FromString(InJsonValue->AsString()), SearchabilityStatus);
+		}
+
+		ParsedTagsAndValues.Add(FindInBlueprintsHelpers::FSimpleFTextKeyStorage(InKey), ParsedValue);
 	}
 }
 
@@ -265,11 +295,11 @@ bool FImaginaryFiBData::TestBasicStringExpression(const FTextFilterString& InVal
 			FText Value = ParsedValues.Value.GetDisplayText(*LookupTablePtr);
 			FString ValueAsString = Value.ToString();
 			ValueAsString.ReplaceInline(TEXT(" "), TEXT(""));
-			bool bMatchesSearch = TextFilterUtils::TestBasicStringExpression(ValueAsString, InValue, InTextComparisonMode) || TextFilterUtils::TestBasicStringExpression(Value.BuildSourceString(), InValue, InTextComparisonMode);
+			bool bMatchesSearch = TextFilterUtils::TestBasicStringExpression(MoveTemp(ValueAsString), InValue, InTextComparisonMode) || TextFilterUtils::TestBasicStringExpression(Value.BuildSourceString(), InValue, InTextComparisonMode);
 			
 			if (bMatchesSearch && !ParsedValues.Value.IsCoreDisplay())
 			{
-				FSearchResult SearchResult(new FFindInBlueprintsResult(CreateSearchComponentDisplayText(ParsedValues.Value.GetDisplayKey(), Value), nullptr ));
+				FSearchResult SearchResult = MakeShared<FFindInBlueprintsResult>(CreateSearchComponentDisplayText(ParsedValues.Value.GetDisplayKey(), Value));
 				InOutMatchingSearchComponents.Add(this, FComponentUniqueDisplay(SearchResult));
 			}
 
@@ -300,11 +330,11 @@ bool FImaginaryFiBData::TestComplexExpression(const FName& InKey, const FTextFil
 				FText Value = TagsValuePair.Value.GetDisplayText(*LookupTablePtr);
 				FString ValueAsString = Value.ToString();
 				ValueAsString.ReplaceInline(TEXT(" "), TEXT(""));
-				bool bMatchesSearch = TextFilterUtils::TestComplexExpression(ValueAsString, InValue, InComparisonOperation, InTextComparisonMode) || TextFilterUtils::TestBasicStringExpression(Value.BuildSourceString(), InValue, InTextComparisonMode);
+				bool bMatchesSearch = TextFilterUtils::TestComplexExpression(MoveTemp(ValueAsString), InValue, InComparisonOperation, InTextComparisonMode) || TextFilterUtils::TestComplexExpression(Value.BuildSourceString(), InValue, InComparisonOperation, InTextComparisonMode);
 
 				if (bMatchesSearch && !TagsValuePair.Value.IsCoreDisplay())
 				{
-					FSearchResult SearchResult(new FFindInBlueprintsResult(CreateSearchComponentDisplayText(TagsValuePair.Value.GetDisplayKey(), Value), nullptr));
+					FSearchResult SearchResult = MakeShared<FFindInBlueprintsResult>(CreateSearchComponentDisplayText(TagsValuePair.Value.GetDisplayKey(), Value));
 					InOutMatchingSearchComponents.Add(this, FComponentUniqueDisplay(SearchResult));
 				}
 				bMatchesSearchQuery |= bMatchesSearch;
@@ -421,9 +451,16 @@ bool FCategorySectionHelper::CanCallFilter(ESearchQueryFilter InSearchQueryFilte
 	return true;
 }
 
-FSearchResult FCategorySectionHelper::CreateSearchResult_Internal(FSearchResult InParent) const
+FSearchResult FCategorySectionHelper::CreateSearchResult_Internal(FSearchResult InTemplate) const
 {
-	return FSearchResult(new FFindInBlueprintsResult(CategoryName, InParent));
+	if (InTemplate.IsValid())
+	{
+		return MakeShared<FFindInBlueprintsResult>(*InTemplate);
+	}
+	else
+	{
+		return MakeShared<FFindInBlueprintsResult>(CategoryName);
+	}
 }
 
 void FCategorySectionHelper::ParseAllChildData_Internal(ESearchableValueStatus InSearchabilityOverride/* = ESearchableValueStatus::Searchable*/)
@@ -491,9 +528,16 @@ FImaginaryBlueprint::FImaginaryBlueprint(FString InBlueprintName, FString InBlue
 	ParsedChildData.Add(InterfaceCategory);
 }
 
-FSearchResult FImaginaryBlueprint::CreateSearchResult_Internal(FSearchResult InParent) const
+FSearchResult FImaginaryBlueprint::CreateSearchResult_Internal(FSearchResult InTemplate) const
 {
-	return FSearchResult(new FFindInBlueprintsResult(ParsedTagsAndValues.Find(FindInBlueprintsHelpers::FSimpleFTextKeyStorage(FFindInBlueprintSearchTags::FiB_Path))->GetDisplayText(LookupTable)));
+	if (InTemplate.IsValid())
+	{
+		return MakeShared<FFindInBlueprintsResult>(*InTemplate);
+	}
+	else
+	{
+		return MakeShared<FFindInBlueprintsResult>(ParsedTagsAndValues.Find(FindInBlueprintsHelpers::FSimpleFTextKeyStorage(FFindInBlueprintSearchTags::FiB_Path))->GetDisplayText(LookupTable));
+	}
 }
 
 void FImaginaryBlueprint::DumpParsedObject_Internal(FArchive& Ar) const
@@ -609,9 +653,16 @@ FImaginaryGraph::FImaginaryGraph(FImaginaryFiBDataWeakPtr InOuter, TSharedPtr< F
 {
 }
 
-FSearchResult FImaginaryGraph::CreateSearchResult_Internal(FSearchResult InParent) const
+FSearchResult FImaginaryGraph::CreateSearchResult_Internal(FSearchResult InTemplate) const
 {
-	return FSearchResult(new FFindInBlueprintsGraph(FText::GetEmpty(), InParent, GraphType));
+	if (InTemplate.IsValid())
+	{
+		return MakeShared<FFindInBlueprintsGraph>(*StaticCastSharedPtr<FFindInBlueprintsGraph>(InTemplate));
+	}
+	else
+	{
+		return MakeShared<FFindInBlueprintsGraph>(GraphType);
+	}
 }
 
 void FImaginaryGraph::DumpParsedObject_Internal(FArchive& Ar) const
@@ -684,9 +735,16 @@ FImaginaryGraphNode::FImaginaryGraphNode(FImaginaryFiBDataWeakPtr InOuter, TShar
 {
 }
 
-FSearchResult FImaginaryGraphNode::CreateSearchResult_Internal(FSearchResult InParent) const
+FSearchResult FImaginaryGraphNode::CreateSearchResult_Internal(FSearchResult InTemplate) const
 {
-	return FSearchResult(new FFindInBlueprintsGraphNode(FText::GetEmpty(), InParent));
+	if (InTemplate.IsValid())
+	{
+		return MakeShared<FFindInBlueprintsGraphNode>(*StaticCastSharedPtr<FFindInBlueprintsGraphNode>(InTemplate));
+	}
+	else
+	{
+		return MakeShared<FFindInBlueprintsGraphNode>();
+	}
 }
 
 void FImaginaryGraphNode::DumpParsedObject_Internal(FArchive& Ar) const
@@ -788,9 +846,16 @@ bool FImaginaryProperty::IsCompatibleWithFilter(ESearchQueryFilter InSearchQuery
 		InSearchQueryFilter == ESearchQueryFilter::VariablesFilter;
 }
 
-FSearchResult FImaginaryProperty::CreateSearchResult_Internal(FSearchResult InParent) const
+FSearchResult FImaginaryProperty::CreateSearchResult_Internal(FSearchResult InTemplate) const
 {
-	return FSearchResult(new FFindInBlueprintsProperty(FText::GetEmpty(), InParent));
+	if (InTemplate.IsValid())
+	{
+		return MakeShared<FFindInBlueprintsProperty>(*StaticCastSharedPtr<FFindInBlueprintsProperty>(InTemplate));
+	}
+	else
+	{
+		return MakeShared<FFindInBlueprintsProperty>();
+	}
 }
 
 void FImaginaryProperty::DumpParsedObject_Internal(FArchive& Ar) const
@@ -848,9 +913,16 @@ bool FImaginaryPin::IsCompatibleWithFilter(ESearchQueryFilter InSearchQueryFilte
 	return InSearchQueryFilter == ESearchQueryFilter::AllFilter || InSearchQueryFilter == ESearchQueryFilter::PinsFilter;
 }
 
-FSearchResult FImaginaryPin::CreateSearchResult_Internal(FSearchResult InParent) const
+FSearchResult FImaginaryPin::CreateSearchResult_Internal(FSearchResult InTemplate) const
 {
-	return FSearchResult(new FFindInBlueprintsPin(FText::GetEmpty(), InParent, SchemaName));
+	if (InTemplate.IsValid())
+	{
+		return MakeShared<FFindInBlueprintsPin>(*StaticCastSharedPtr<FFindInBlueprintsPin>(InTemplate));
+	}
+	else
+	{
+		return MakeShared<FFindInBlueprintsPin>(SchemaName);
+	}
 }
 
 void FImaginaryPin::DumpParsedObject_Internal(FArchive& Ar) const
diff --git a/Engine/Source/Editor/Kismet/Public/FindInBlueprintManager.h b/Engine/Source/Editor/Kismet/Public/FindInBlueprintManager.h
--- a/Engine/Source/Editor/Kismet/Public/FindInBlueprintManager.h
+++ b/Engine/Source/Editor/Kismet/Public/FindInBlueprintManager.h
@@ -253,13 +253,11 @@ struct FFindInBlueprintCachingOptions
 class FFindInBlueprintsResult : public TSharedFromThis< FFindInBlueprintsResult >
 {
 public:
-	/* Create a root */
-	FFindInBlueprintsResult(const FText& InDisplayText);
-
-	/* Create a listing for a search result*/
-	FFindInBlueprintsResult(const FText& InDisplayText, TSharedPtr<FFindInBlueprintsResult> InParent);
+	FFindInBlueprintsResult() = default;
+	virtual ~FFindInBlueprintsResult() = default;
 
-	virtual ~FFindInBlueprintsResult() {}
+	/* Create a root */
+	explicit FFindInBlueprintsResult(const FText& InDisplayText);
 
 	/* Called when user clicks on the search item */
 	virtual FReply OnClick();
@@ -292,15 +290,6 @@ public:
 	/** Returns the Object represented by this search information give the Blueprint it can be found in */
 	virtual UObject* GetObject(UBlueprint* InBlueprint) const;
 
-	/**
-	* Adds extra search info, anything that does not have a predestined place in the search result. Adds a sub-item to the searches and formats its description so the tag displays
-	*
-	* @param	InKey			This is the tag for the data, describing what it is so special handling can occur if needed
-	* @param	InValue			Compared against search query to see if it passes the filter, sometimes data is rejected because it is deemed unsearchable
-	* @param	InParent		The parent search result
-	*/
-	void AddExtraSearchInfo(FText InKey, FText InValue, TSharedPtr< FFindInBlueprintsResult > InParent);
-
 	/** Returns the display string for the row */
 	FText GetDisplayString() const;
 
@@ -561,8 +550,12 @@ public:
 
 	bool IsGatheringDataEnabled() const { return bEnableGatheringData; }
 
+	/** If TRUE, the developer menu tool commands will be shown in the 'Developer' section of the Blueprint Editor's menu bar */
 	bool ShouldEnableDeveloperMenuTools() const { return bEnableDeveloperMenuTools; }
 
+	/** If TRUE, search result meta will be gathered once and stored in a template. Avoids doing this work redundantly at search time. */
+	bool ShouldEnableSearchResultTemplates() const { return !bDisableSearchResultTemplates; }
+
 	/** Find or create the global find results widget */
 	TSharedPtr<SFindInBlueprints> GetGlobalFindResults();
 
@@ -717,6 +710,9 @@ protected:
 	/** Whether to enable Blueprint editor developer menu tools */
 	bool bEnableDeveloperMenuTools;
 
+	/** Disable the use of search result templates. Setting this to TRUE will slightly decrease overall memory usage, but will also increase global search times */
+	bool bDisableSearchResultTemplates;
+
 	/** TRUE when the the FiB manager wants to pause all searches, helps manage the pausing procedure */
 	volatile bool bIsPausing;
 };
diff --git a/Engine/Source/Editor/Kismet/Public/FindInBlueprints.h b/Engine/Source/Editor/Kismet/Public/FindInBlueprints.h
--- a/Engine/Source/Editor/Kismet/Public/FindInBlueprints.h
+++ b/Engine/Source/Editor/Kismet/Public/FindInBlueprints.h
@@ -90,22 +90,24 @@ class FFindInBlueprintsNoResult : public FFindInBlueprintsResult
 {
 public:
 	FFindInBlueprintsNoResult(const FText& InDisplayText)
-		:FFindInBlueprintsResult(InDisplayText, nullptr)
+		:FFindInBlueprintsResult(InDisplayText)
 	{
 	}
 
+	/** FFindInBlueprintsResult Interface */
 	virtual FReply OnClick() override
 	{
 		// Do nothing on click.
 		return FReply::Handled();
 	}
+	/** End FFindInBlueprintsResult Interface */
 };
 
 /** Graph nodes use this class to store their data */
 class FFindInBlueprintsGraphNode : public FFindInBlueprintsResult
 {
 public:
-	FFindInBlueprintsGraphNode(const FText& InValue, TSharedPtr<FFindInBlueprintsResult> InParent);
+	FFindInBlueprintsGraphNode();
 	virtual ~FFindInBlueprintsGraphNode() {}
 
 	/** FFindInBlueprintsResult Interface */
@@ -138,7 +140,7 @@ private:
 class FFindInBlueprintsPin : public FFindInBlueprintsResult
 {
 public:
-	FFindInBlueprintsPin(const FText& InValue, TSharedPtr<FFindInBlueprintsResult> InParent, FString InSchemaName);
+	FFindInBlueprintsPin(FString InSchemaName);
 	virtual ~FFindInBlueprintsPin() {}
 
 	/** FFindInBlueprintsResult Interface */
@@ -147,7 +149,8 @@ public:
 	virtual FText GetCategory() const override;
 	virtual void FinalizeSearchData() override;
 	/** End FFindInBlueprintsResult Interface */
-protected:
+
+private:
 	/** The name of the schema this pin exists under */
 	FString SchemaName;
 
@@ -162,7 +165,7 @@ protected:
 class FFindInBlueprintsProperty : public FFindInBlueprintsResult
 {
 public:
-	FFindInBlueprintsProperty(const FText& InValue, TSharedPtr<FFindInBlueprintsResult> InParent);
+	FFindInBlueprintsProperty();
 	virtual ~FFindInBlueprintsProperty() {}
 
 	/** FFindInBlueprintsResult Interface */
@@ -172,7 +175,8 @@ public:
 	virtual FText GetCategory() const override;
 	virtual void FinalizeSearchData() override;
 	/** End FFindInBlueprintsResult Interface */
-protected:
+
+private:
 	/** The pin that this search result refers to */
 	FEdGraphPinType PinType;
 
@@ -187,7 +191,7 @@ protected:
 class FFindInBlueprintsGraph : public FFindInBlueprintsResult
 {
 public:
-	FFindInBlueprintsGraph(const FText& InValue, TSharedPtr<FFindInBlueprintsResult> InParent, EGraphType InGraphType);
+	FFindInBlueprintsGraph(EGraphType InGraphType);
 	virtual ~FFindInBlueprintsGraph() {}
 
 	/** FFindInBlueprintsResult Interface */
@@ -196,8 +200,8 @@ public:
 	virtual void ParseSearchInfo(FText InKey, FText InValue) override;
 	virtual FText GetCategory() const override;
 	/** End FFindInBlueprintsResult Interface */
-protected:
 
+private:
 	/** The type of graph this represents */
 	EGraphType GraphType;
 };
diff --git a/Engine/Source/Editor/Kismet/Public/ImaginaryBlueprintData.h b/Engine/Source/Editor/Kismet/Public/ImaginaryBlueprintData.h
--- a/Engine/Source/Editor/Kismet/Public/ImaginaryBlueprintData.h
+++ b/Engine/Source/Editor/Kismet/Public/ImaginaryBlueprintData.h
@@ -22,6 +22,12 @@ enum ESearchableValueStatus
 class FSearchableValueInfo
 {
 public:
+	FSearchableValueInfo()
+		: SearchableValueStatus(ESearchableValueStatus::NotSearchable)
+		, LookupTableKey(-1)
+	{
+	}
+
 	FSearchableValueInfo(FText InDisplayKey, int32 InLookUpTableKey)
 		: SearchableValueStatus(ESearchableValueStatus::Searchable)
 		, DisplayKey(InDisplayKey)
@@ -35,7 +41,6 @@ public:
 		, LookupTableKey(-1)
 		, DisplayText(InDisplayText)
 	{
-
 	}
 
 	FSearchableValueInfo(FText InDisplayKey, int32 InLookUpTableKey, ESearchableValueStatus InSearchableValueStatus)
@@ -43,7 +48,6 @@ public:
 		, DisplayKey(InDisplayKey)
 		, LookupTableKey(InLookUpTableKey)
 	{
-
 	}
 
 	FSearchableValueInfo(FText InDisplayKey, FText InDisplayText, ESearchableValueStatus InSearchableValueStatus)
@@ -52,7 +56,6 @@ public:
 		, LookupTableKey(-1)
 		, DisplayText(InDisplayText)
 	{
-
 	}
 
 	/** Returns TRUE if the data is searchable */
@@ -186,13 +189,13 @@ protected:
 	virtual ESearchableValueStatus GetSearchabilityStatus(FString InKey) { return ESearchableValueStatus::Searchable; };
 
 	/** Protected internal function which builds the search result for this item */
-	virtual FSearchResult CreateSearchResult_Internal(FSearchResult InParent) const = 0;
+	virtual FSearchResult CreateSearchResult_Internal(FSearchResult InTemplate) const = 0;
 
 	/** Creates a display string for this item in search results */
 	FText CreateSearchComponentDisplayText(FText InKey, FText InValue) const;
 
 	/** Helper function for parsing Json values into usable properties */
-	void ParseJsonValue(FText InKey, FText InDisplayKey, TSharedPtr< FJsonValue > InJsonValue, bool bIsInArray=false, ESearchableValueStatus InSearchabilityOverride = ESearchableValueStatus::Searchable);
+	void ParseJsonValue(FText InKey, FText InDisplayKey, TSharedPtr< FJsonValue > InJsonValue, TArray<FSearchableValueInfo>& OutParsedValues, bool bIsInArray = false, ESearchableValueStatus InSearchabilityOverride = ESearchableValueStatus::Searchable);
 
 	/** Internal version of the ParseAllChildData function, handles the bulk of the work */
 	virtual void ParseAllChildData_Internal(ESearchableValueStatus InSearchabilityOverride = ESearchableValueStatus::Searchable);
@@ -218,20 +221,16 @@ protected:
 
 	/** Allows for thread-safe parsing of the imaginary data. Only a single Imaginary data can be parsed at a time. */
 	static FCriticalSection ParseChildDataCriticalSection;
+
+private:
+	/** If display meta is present, this will cache those values and is then used as a basis when constructing a search result tree */
+	FSearchResult SearchResultTemplate;
 };
 
 class FFiBMetaData : public FImaginaryFiBData
 {
 public:
 	FFiBMetaData(FImaginaryFiBDataWeakPtr InOuter, TSharedPtr< FJsonObject > InUnparsedJsonObject, TMap<int32, FText>* InLookupTablePtr);
-	
-	/** FImaginaryFiBData Interface */
-	virtual bool TrySpecialHandleJsonValue(FText InKey, TSharedPtr< FJsonValue > InJsonValue) override;
-	virtual FSearchResult CreateSearchResult_Internal(FSearchResult InParent) const override
-	{
-		return nullptr;
-	}
-	/** End FImaginaryFiBData Interface */
 
 	/** Returns TRUE if the metadata is informing that the FProperty and children should be hidden */
 	bool IsHidden() const
@@ -246,7 +245,17 @@ public:
 	{
 		return bIsExplicit;
 	}
+
 protected:
+	/** FImaginaryFiBData Interface */
+	virtual bool TrySpecialHandleJsonValue(FText InKey, TSharedPtr< FJsonValue > InJsonValue) override;
+	virtual FSearchResult CreateSearchResult_Internal(FSearchResult InTemplate) const override
+	{
+		return nullptr;
+	}
+	/** End FImaginaryFiBData Interface */
+
+private:
 	/** TRUE if the FProperty this metadata represents is hidden */
 	bool bIsHidden;
 
@@ -265,9 +274,9 @@ public:
 	FCategorySectionHelper(FImaginaryFiBDataWeakPtr InOuter, TSharedPtr< FJsonObject > InUnparsedJsonObject, TMap<int32, FText>* InLookupTablePtr, FText InCategoryName, bool bInTagAndValueCategory, FCategorySectionHelperCallback InSpecialHandlingCallback);
 
 	/** FImaginaryFiBData Interface */
-	virtual void ParseAllChildData_Internal(ESearchableValueStatus InSearchabilityOverride/* = ESearchableValueStatus::ESearchableValueStatus::Searchable*/) override;
 	virtual bool IsCategory() const override { return true; }
 	virtual bool IsTagAndValueCategory() const override { return bIsTagAndValue; }
+	virtual bool CanCallFilter(ESearchQueryFilter InSearchQueryFilter) const override;
 	/** End FImaginaryFiBData Interface */
 
 	/** Returns the category name prepared for checking as a function */
@@ -278,8 +287,8 @@ public:
 
 protected:
 	/** FImaginaryFiBData Interface */
-	virtual bool CanCallFilter(ESearchQueryFilter InSearchQueryFilter) const override;
-	virtual FSearchResult CreateSearchResult_Internal(FSearchResult InParent) const override;
+	virtual FSearchResult CreateSearchResult_Internal(FSearchResult InTemplate) const override;
+	virtual void ParseAllChildData_Internal(ESearchableValueStatus InSearchabilityOverride/* = ESearchableValueStatus::ESearchableValueStatus::Searchable*/) override;
 	virtual void DumpParsedObject_Internal(FArchive& Ar) const override;
 	/** End FImaginaryFiBData Interface */
 
@@ -309,7 +318,7 @@ public:
 protected:
 	/** FImaginaryFiBData Interface */
 	virtual bool TrySpecialHandleJsonValue(FText InKey, TSharedPtr< FJsonValue > InJsonValue) override;
-	virtual FSearchResult CreateSearchResult_Internal(FSearchResult InParent) const override;
+	virtual FSearchResult CreateSearchResult_Internal(FSearchResult InTemplate) const override;
 	virtual void DumpParsedObject_Internal(FArchive& Ar) const override;
 	/** End FImaginaryFiBData Interface */
 
@@ -348,7 +357,7 @@ protected:
 	/** FImaginaryFiBData Interface */
 	virtual bool TrySpecialHandleJsonValue(FText InKey, TSharedPtr< FJsonValue > InJsonValue) override;
 	virtual ESearchableValueStatus GetSearchabilityStatus(FString InKey) override;
-	virtual FSearchResult CreateSearchResult_Internal(FSearchResult InParent) const override;
+	virtual FSearchResult CreateSearchResult_Internal(FSearchResult InTemplate) const override;
 	virtual void DumpParsedObject_Internal(FArchive& Ar) const override;
 	/** End FImaginaryFiBData Interface */
 
@@ -366,18 +375,18 @@ public:
 	/** FImaginaryFiBData Interface */
 	virtual bool IsCompatibleWithFilter(ESearchQueryFilter InSearchQueryFilter) const override;
 	virtual bool CanCallFilter(ESearchQueryFilter InSearchQueryFilter) const override;
-	virtual void ParseAllChildData_Internal(ESearchableValueStatus InSearchabilityOverride/* = ESearchableValueStatus::Searchable*/) override;
-	virtual void DumpParsedObject_Internal(FArchive& Ar) const override;
 	/** End FImaginaryFiBData Interface */
 
 protected:
 	/** FImaginaryFiBData Interface */
 	virtual bool TrySpecialHandleJsonValue(FText InKey, TSharedPtr< FJsonValue > InJsonValue) override;
 	virtual ESearchableValueStatus GetSearchabilityStatus(FString InKey) override;
-	virtual FSearchResult CreateSearchResult_Internal(FSearchResult InParent) const override;
+	virtual FSearchResult CreateSearchResult_Internal(FSearchResult InTemplate) const override;
+	virtual void ParseAllChildData_Internal(ESearchableValueStatus InSearchabilityOverride/* = ESearchableValueStatus::Searchable*/) override;
+	virtual void DumpParsedObject_Internal(FArchive& Ar) const override;
 	/** End FImaginaryFiBData Interface */
 
-protected:
+private:
 	/** Schema name that manages this node */
 	FString SchemaName;
 };
@@ -395,7 +404,7 @@ public:
 protected:
 	/** FImaginaryFiBData Interface */
 	virtual ESearchableValueStatus GetSearchabilityStatus(FString InKey) override;
-	virtual FSearchResult CreateSearchResult_Internal(FSearchResult InParent) const override;
+	virtual FSearchResult CreateSearchResult_Internal(FSearchResult InTemplate) const override;
 	virtual void DumpParsedObject_Internal(FArchive& Ar) const override;
 	/** End FImaginaryFiBData Interface */
 };
@@ -425,11 +434,11 @@ protected:
 	/** FImaginaryFiBData Interface */
 	virtual bool TrySpecialHandleJsonValue(FText InKey, TSharedPtr< FJsonValue > InJsonValue);
 	virtual ESearchableValueStatus GetSearchabilityStatus(FString InKey) override;
-	virtual FSearchResult CreateSearchResult_Internal(FSearchResult InParent) const override;
+	virtual FSearchResult CreateSearchResult_Internal(FSearchResult InTemplate) const override;
 	virtual void DumpParsedObject_Internal(FArchive& Ar) const override;
 	/** End FImaginaryFiBData Interface */
 
-protected:
+private:
 	/** Schema name that manages this pin */
 	FString SchemaName;
 };
