Fix crash introduced by CL# 8130582 since Function can be null. Use the Event's name in that case instead

diff --git a/Engine/Source/Editor/Kismet/Private/SMyBlueprint.cpp b/Engine/Source/Editor/Kismet/Private/SMyBlueprint.cpp
--- a/Engine/Source/Editor/Kismet/Private/SMyBlueprint.cpp
+++ b/Engine/Source/Editor/Kismet/Private/SMyBlueprint.cpp
@@ -1717,7 +1717,7 @@ FReply SMyBlueprint::OnActionDragged( const TArray< TSharedPtr<FEdGraphSchemaAct
 
 				TSharedRef< FKismetFunctionDragDropAction> DragOperation =
 					FKismetFunctionDragDropAction::New(
-						InAction, Function->GetFName(),
+						InAction, (Function ? Function->GetFName() : Event->GetFName()),
 						GetBlueprintObj()->SkeletonGeneratedClass,
 						FMemberReference(),
 						AnalyticsDelegate,
