Reintroduce invariant that child types must be compiled when a parent is - primarily for animation blueprints

diff --git a/Engine/Source/Editor/Kismet/Private/BlueprintCompilationManager.cpp b/Engine/Source/Editor/Kismet/Private/BlueprintCompilationManager.cpp
--- a/Engine/Source/Editor/Kismet/Private/BlueprintCompilationManager.cpp
+++ b/Engine/Source/Editor/Kismet/Private/BlueprintCompilationManager.cpp
@@ -937,6 +937,20 @@ void FBlueprintCompilationManagerImpl::FlushCompilationQueueImpl(bool bSuppressB
 						return false;
 					}
 
+					// if our parent is still being compiled, then we still need to be compiled:
+					UClass* Iter = Data.BP->ParentClass;
+					while(Iter)
+					{
+						if(UBlueprint* BP = Cast<UBlueprint>(Iter->ClassGeneratedBy))
+						{
+							if(BP->bBeingCompiled)
+							{
+								return false;
+							}
+						}
+						Iter = Iter->GetSuperClass();
+					}
+
 					// look for references to a function with a signature change
 					// in the old class, if it has none, we can skip bytecode recompile:
 					bool bHasNoReferencesToChangedFunctions = true;
