Decoupled BP indexing from the search thread in the editor in order to take advantage of multiprocessing, and improved the overall UX.

diff --git a/Engine/Config/BaseEditor.ini b/Engine/Config/BaseEditor.ini
--- a/Engine/Config/BaseEditor.ini
+++ b/Engine/Config/BaseEditor.ini
@@ -430,7 +430,9 @@ TemperatureUnits=EUnit::Celsius
 ForceUnits=EUnit::Newtons
 
 [BlueprintSearchSettings]
+AsyncTaskBatchSize=2048
 bDisableDeferredIndexing=false
+bDisableThreadedIndexing=false
 bEnableCsvStatsProfiling=false
 bEnableDeveloperMenuTools=false
 bDisableSearchResultTemplates=false
diff --git a/Engine/Source/Editor/Kismet/Private/FindInBlueprintManager.cpp b/Engine/Source/Editor/Kismet/Private/FindInBlueprintManager.cpp
--- a/Engine/Source/Editor/Kismet/Private/FindInBlueprintManager.cpp
+++ b/Engine/Source/Editor/Kismet/Private/FindInBlueprintManager.cpp
@@ -5,7 +5,7 @@
 #include "Misc/MessageDialog.h"
 #include "HAL/FileManager.h"
 #include "Misc/Paths.h"
-#include "HAL/RunnableThread.h"
+#include "Async/ParallelFor.h"
 #include "Misc/ScopeLock.h"
 #include "Serialization/MemoryWriter.h"
 #include "Serialization/MemoryReader.h"
@@ -53,6 +53,7 @@
 #include "UObject/EditorObjectVersion.h"
 #include "Framework/Notifications/NotificationManager.h"
 #include "Widgets/Notifications/SNotificationList.h"
+#include "Application/ThrottleManager.h"
 
 DEFINE_LOG_CATEGORY(LogFindInBlueprint);
 CSV_DEFINE_CATEGORY(FindInBlueprint, false);
@@ -140,7 +141,7 @@ FStreamSearch::FStreamSearch(const FString& InSearchValue, const FStreamSearchOp
 	SearchId = GlobalSearchCounter++;
 
 	// Create a uniquely-named thread to ensure disambiguation in the thread view and assist with debugging multiple instances running in parallel.
-	Thread = FRunnableThread::Create( this, *FString::Printf(TEXT("FStreamSearch_%d"), SearchId), 0, TPri_BelowNormal );
+	Thread = TUniquePtr<FRunnableThread>(FRunnableThread::Create( this, *FString::Printf(TEXT("FStreamSearch_%d"), SearchId), 0, TPri_BelowNormal ));
 }
 
 bool FStreamSearch::Init()
@@ -226,8 +227,6 @@ void FStreamSearch::EnsureCompletion()
 
 	Stop();
 	Thread->WaitForCompletion();
-	delete Thread;
-	Thread = NULL;
 }
 
 bool FStreamSearch::IsComplete() const
@@ -521,7 +520,7 @@ namespace BlueprintSearchMetaDataHelpers
 	public:
 		FFindInBlueprintJsonWriter(FString* const InOutString)
 			:TFindInBlueprintJsonStringWriter<TCondensedJsonPrintPolicy<TCHAR>>(InOutString)
-			, JsonOutput(InOutString)
+			,JsonOutput(InOutString)
 				{
 		}
 
@@ -1252,13 +1251,142 @@ namespace BlueprintSearchMetaDataHelpers
 	}
 }
 
-class FCacheAllBlueprintsTickableObject
+/** Interface for controlling the async indexing task. */
+class IAsyncSearchIndexTaskController
+{
+public:
+	/** Return true if there is still work available. */
+	virtual bool IsWorkPending() const = 0;
+
+	/** Return the next batch of asset paths for indexing. */
+	virtual void GetAssetPathsToIndex(TArray<FName>& OutAssetPaths) = 0;
+
+	/** Return true if assets should be fully indexed. */
+	virtual bool ShouldFullyIndexAssets() const = 0;
+
+	/** Determine if multiprocessing should be enabled for indexing the next batch. */
+	virtual bool ShouldEnableMultiprocessing() const = 0;
+
+	/** Add the given asset path to the queue for gathering search data from a Blueprint. */
+	virtual void AddAssetPathToGatherQueue(FName InAssetPath) = 0;
+
+	/** Called when indexing has been completed for the given asset path. */
+	virtual void IndexCompletedForAssetPath(FName InAssetPath) = 0;
+};
+
+/** Asynchronous indexing thread. Can spawn additional worker threads to index multiple assets in parallel. */
+class FAsyncSearchIndexTaskRunnable : public FRunnable
+{
+public:
+	FAsyncSearchIndexTaskRunnable(IAsyncSearchIndexTaskController* InController)
+		:Controller(InController)
+	{
+		Thread = TUniquePtr<FRunnableThread>(FRunnableThread::Create(this, TEXT("FAsyncSearchIndexTaskRunnable"), 0, TPri_BelowNormal));
+	}
+
+	virtual uint32 Run() override
+	{
+		while (Controller->IsWorkPending())
+		{
+			// Get the next batch of asset paths to be indexed.
+			TArray<FName> AssetPathsToIndex;
+			Controller->GetAssetPathsToIndex(AssetPathsToIndex);
+
+			// Determine whether this is a full or partial indexing operation.
+			const bool bEnableFullIndexingPass = Controller->ShouldFullyIndexAssets();
+
+			// Determine whether to disable multiprocessing for this batch. If true, this batch will be processed only on this thread.
+			const bool bDisableMultiprocessing = !Controller->ShouldEnableMultiprocessing();
+
+			ParallelFor(AssetPathsToIndex.Num(), [&AssetPathsToIndex, Controller = this->Controller, bEnableFullIndexingPass, bDisableMultiprocessing](int32 ArrayIdx)
+			{
+				CSV_CUSTOM_STAT(FindInBlueprint, IndexedAssetCountThisFrame, 1, ECsvCustomStatOp::Accumulate);
+
+				FName AssetPath = AssetPathsToIndex[ArrayIdx];
+				FSearchData SearchData = FFindInBlueprintSearchManager::Get().GetSearchDataForAssetPath(AssetPath);
+				if (SearchData.IsValid() && !SearchData.IsMarkedForDeletion() && !SearchData.IsIndexingCompleted())
+				{
+					// Generate the metadata tag value if it was not previously cached or loaded.
+					if (SearchData.Value.Len() == 0)
+					{
+						// This must be done on the main thread, so enqueue it and continue.
+						Controller->AddAssetPathToGatherQueue(AssetPath);
+					}
+					else
+					{
+						if (bEnableFullIndexingPass)
+						{
+							// Unpack the metadata tag and rebuild the index for this asset.
+							const FString AssetPathAsString = AssetPath.ToString();
+							if (FiBSerializationHelpers::ValidateSearchDataVersionInfo(AssetPathAsString, SearchData.Value, SearchData.VersionInfo))
+							{
+								SearchData.ImaginaryBlueprint = MakeShareable(new FImaginaryBlueprint(FPaths::GetBaseFilename(AssetPathAsString), AssetPathAsString, SearchData.ParentClass, SearchData.Interfaces, SearchData.Value, SearchData.VersionInfo));
+								SearchData.Value.Empty();
+
+								// Build the full index using a BFS traversal.
+								TArray<FImaginaryFiBDataSharedPtr> IndexNodes = { SearchData.ImaginaryBlueprint };
+								while (IndexNodes.Num() > 0)
+								{
+									ParallelFor(IndexNodes.Num(), [&IndexNodes](int32 NodeIdx)
+									{
+										IndexNodes[NodeIdx]->ParseAllChildData();
+									}, bDisableMultiprocessing);
+
+									TArray<FImaginaryFiBDataSharedPtr> ChildNodes;
+									for (FImaginaryFiBDataSharedPtr& NodePtr : IndexNodes)
+									{
+										ChildNodes.Append(NodePtr->GetAllParsedChildData());
+									}
+
+									IndexNodes = MoveTemp(ChildNodes);
+								}
+							}
+
+							// Signal that this asset has now been fully indexed.
+							SearchData.StateFlags |= ESearchDataStateFlags::IsIndexed;
+
+							// Update this entry in the search database (thread-safe).
+							FFindInBlueprintSearchManager::Get().ApplySearchDataToDatabase(MoveTemp(SearchData));
+						}
+
+						// Signal that indexing has been completed for this asset path.
+						Controller->IndexCompletedForAssetPath(AssetPath);
+					}
+				}
+				else
+				{
+					// Signal that indexing has been completed for this asset path, even though its cache entry is invalid.
+					Controller->IndexCompletedForAssetPath(AssetPath);
+				}
+			}, bDisableMultiprocessing);
+		}
+
+		return 0;
+	}
+
+	void EnsureCompletion()
+	{
+		if (Thread.IsValid())
+		{
+			Thread->WaitForCompletion();
+		}
+	}
+
+private:
+	TUniquePtr<FRunnableThread> Thread;
+	IAsyncSearchIndexTaskController* Controller;
+};
+
+class FCacheAllBlueprintsTickableObject : public IAsyncSearchIndexTaskController
 {
 public:
 	DECLARE_DELEGATE_OneParam(FOnAssetCached, FName);
 
 	struct FCacheParams
 	{
+		/** Cache type */
+		EFiBCacheOpType OpType;
+
 		/** Control flags */
 		EFiBCacheOpFlags OpFlags;
 
@@ -1268,19 +1396,30 @@ public:
 		/** Callback for when caching is finished */
 		FSimpleDelegate OnFinished;
 
+		/** Size of each batch (for async processing) */
+		int32 AsyncTaskBatchSize;
+
 		FCacheParams()
 			:OpFlags(EFiBCacheOpFlags::None)
+			,AsyncTaskBatchSize(0)
 		{
 		}
 	};
 
-	FCacheAllBlueprintsTickableObject(TSet<FName> InAssets, const FCacheParams& InParams)
-		: TickCacheIndex(0)
-		, UncachedAssets(InAssets.Array())
+	FCacheAllBlueprintsTickableObject(const TSet<FName>& InAssets, const FCacheParams& InParams)
+		: UncachedAssets(InAssets.Array())
+		, CacheParams(InParams)
+		, TickCacheIndex(0)
+		, AsyncTaskBatchIndex(0)
+		, bIsGatheringSearchMetadata(false)
 		, bIsStarted(false)
 		, bIsCancelled(false)
-		, CacheParams(InParams)
 	{
+		if (CacheParams.AsyncTaskBatchSize <= 0)
+		{
+			CacheParams.AsyncTaskBatchSize = UncachedAssets.Num();
+		}
+
 		if (EnumHasAnyFlags(CacheParams.OpFlags, EFiBCacheOpFlags::ShowProgress)
 			&& !EnumHasAnyFlags(CacheParams.OpFlags, EFiBCacheOpFlags::HideNotifications))
 		{
@@ -1302,7 +1441,7 @@ public:
 		}
 	}
 
-	~FCacheAllBlueprintsTickableObject()
+	virtual ~FCacheAllBlueprintsTickableObject()
 	{
 
 	}
@@ -1313,6 +1452,18 @@ public:
 		return TickCacheIndex + 1;
 	}
 
+	/** Returns the current cache operation type */
+	EFiBCacheOpType GetCurrentCacheOpType() const
+	{
+		return CacheParams.OpType;
+	}
+
+	/** Returns the current cache operation control flags */
+	EFiBCacheOpFlags GetCurrentCacheOpFlags() const
+	{
+		return CacheParams.OpFlags;
+	}
+
 	/** Returns the name of the current Blueprint being cached */
 	FName GetCurrentCacheBlueprintName() const
 	{
@@ -1364,7 +1515,7 @@ public:
 			if (bIsImmediate)
 			{
 				// Note: This will effectively delete this instance. It should not be used after this!
-				FFindInBlueprintSearchManager::Get().FinishedCachingBlueprints(TickCacheIndex, FailedToCacheList);
+				FFindInBlueprintSearchManager::Get().FinishedCachingBlueprints(CacheParams.OpType, CacheParams.OpFlags, TickCacheIndex, FailedToCacheList);
 			}
 			else
 			{
@@ -1373,13 +1524,87 @@ public:
 		}
 	}
 
+	virtual bool IsWorkPending() const override
+	{
+		const int32 AsyncTaskBatchSize = CacheParams.AsyncTaskBatchSize;
+		const int32 StartIndex = AsyncTaskBatchIndex * AsyncTaskBatchSize;
+		return StartIndex < UncachedAssets.Num() || !AssetsPendingGatherQueue.IsEmpty() || !AssetsPendingAsyncIndexing.IsEmpty() || bIsGatheringSearchMetadata;
+	}
+
+	virtual bool ShouldFullyIndexAssets() const override
+	{
+		return !EnumHasAnyFlags(CacheParams.OpFlags, EFiBCacheOpFlags::ExecuteGatherPhaseOnly);
+	}
+
+	virtual bool ShouldEnableMultiprocessing() const override
+	{
+		if (CacheParams.OpType == EFiBCacheOpType::CachePendingAssets)
+		{
+			// Don't utilize the task graph if any of the following conditions hold TRUE:
+			// a) The application has throttled the tick rate.
+			// b) No global search tabs are currently open and visible.
+			// c) The initial asset discovery phase has not yet been completed.
+			// d) Multiprocessing has been explicitly disabled for this operation.
+			return FSlateThrottleManager::Get().IsAllowingExpensiveTasks()
+				&& FFindInBlueprintSearchManager::Get().IsGlobalFindResultsOpen()
+				&& !FFindInBlueprintSearchManager::Get().IsAssetDiscoveryInProgress()
+				&& !EnumHasAnyFlags(CacheParams.OpFlags, EFiBCacheOpFlags::ExecuteOnSingleThread);
+		}
+
+		return false;
+	}
+
+	virtual void GetAssetPathsToIndex(TArray<FName>& OutAssetPaths) override
+	{
+		OutAssetPaths.Empty();
+
+		const int32 AsyncTaskBatchSize = CacheParams.AsyncTaskBatchSize;
+		const int32 StartIndex = AsyncTaskBatchIndex * AsyncTaskBatchSize;
+		if (StartIndex < UncachedAssets.Num())
+		{
+			for (int32 i = StartIndex; i < FMath::Min(StartIndex + AsyncTaskBatchSize, UncachedAssets.Num()); ++i)
+			{
+				OutAssetPaths.Add(UncachedAssets[i]);
+			}
+
+			AsyncTaskBatchIndex++;
+		}
+		else
+		{
+			FName AssetPath;
+			int32 Count = 0;
+			while (Count < AsyncTaskBatchSize && AssetsPendingAsyncIndexing.Dequeue(AssetPath))
+			{
+				OutAssetPaths.Add(AssetPath);
+				++Count;
+			}
+		}
+	}
+
+	virtual void AddAssetPathToGatherQueue(FName InAssetPath) override
+	{
+		AssetsPendingGatherQueue.Enqueue(InAssetPath);
+	}
+
+	virtual void IndexCompletedForAssetPath(FName InAssetPath) override
+	{
+		FScopeLock Lock(&AsyncTaskCompletionMutex);
+
+		CompletedAsyncTaskAssets.Add(InAssetPath);
+	}
+
 	/** Enables the caching process */
 	void Start()
 	{
 		if (!bIsStarted)
 		{
 			bIsStarted = true;
-			FFindInBlueprintSearchManager::Get().StartedCachingBlueprints(CacheParams.OpFlags);
+			FFindInBlueprintSearchManager::Get().StartedCachingBlueprints(CacheParams.OpType, CacheParams.OpFlags);
+
+			if (CacheParams.OpType == EFiBCacheOpType::CachePendingAssets && !EnumHasAnyFlags(CacheParams.OpFlags, EFiBCacheOpFlags::ExecuteOnMainThread))
+			{
+				AsyncIndexBuilderTask = MakeUnique<FAsyncSearchIndexTaskRunnable>(this);
+			}
 		}
 	}
 
@@ -1401,48 +1626,136 @@ public:
 		if (bIsCancelled || GWarn->ReceivedUserCancel())
 		{
 			// Note: This will effectively delete this instance. It should not be used after this!
-			FFindInBlueprintSearchManager::Get().FinishedCachingBlueprints(TickCacheIndex, FailedToCacheList);
+			FFindInBlueprintSearchManager::Get().FinishedCachingBlueprints(CacheParams.OpType, CacheParams.OpFlags, TickCacheIndex, FailedToCacheList);
 		}
 		else
 		{
-			FAssetRegistryModule* AssetRegistryModule = &FModuleManager::LoadModuleChecked<FAssetRegistryModule>(TEXT("AssetRegistry"));
-			FAssetData AssetData = AssetRegistryModule->Get().GetAssetByObjectPath(UncachedAssets[TickCacheIndex]);
-			if (AssetData.IsValid())
+			if (CacheParams.OpType == EFiBCacheOpType::CachePendingAssets)
 			{
-				const bool bIsWorldAsset = AssetData.AssetClass == UWorld::StaticClass()->GetFName();
-
-				// Construct a full package filename with path so we can query the read only status and save to disk
-				FString FinalPackageFilename = FPackageName::LongPackageNameToFilename(AssetData.PackageName.ToString());
-				if (FinalPackageFilename.Len() > 0 && FPaths::GetExtension(FinalPackageFilename).Len() == 0)
+				if (AsyncIndexBuilderTask.IsValid())
 				{
-					FinalPackageFilename += bIsWorldAsset ? FPackageName::GetMapPackageExtension() : FPackageName::GetAssetPackageExtension();
-				}
-				FText ErrorMessage;
-				bool bValidFilename = FFileHelper::IsFilenameValidForSaving(FinalPackageFilename, ErrorMessage);
-				if (bValidFilename)
-				{
-					bValidFilename = bIsWorldAsset ? FEditorFileUtils::IsValidMapFilename(FinalPackageFilename, ErrorMessage) : FPackageName::IsValidLongPackageName(FinalPackageFilename, false, &ErrorMessage);
-				}
+					// Process assets that require search metadata regeneration on the main thread (one per tick).
+					if (!AssetsPendingGatherQueue.IsEmpty())
+					{
+						// Since we may empty the queue below, this flag is used to indicate that work is still pending.
+						bIsGatheringSearchMetadata = true;
 
-				const bool bCheckOutAndSave = EnumHasAnyFlags(CacheParams.OpFlags, EFiBCacheOpFlags::CheckOutAndSave);
+						FName AssetPath;
+						if (AssetsPendingGatherQueue.Dequeue(AssetPath))
+						{
+							bool bEnqueueForAsyncIndexing = false;
 
-				bool bIsAssetReadOnlyOnDisk = IFileManager::Get().IsReadOnly(*FinalPackageFilename);
-				bool bFailedToCache = bCheckOutAndSave;
+							FSearchData SearchData = FFindInBlueprintSearchManager::Get().GetSearchDataForAssetPath(AssetPath);
+							if (SearchData.IsValid() && !SearchData.IsMarkedForDeletion())
+							{
+								if (UBlueprint* Blueprint = SearchData.Blueprint.Get())
+								{
+									using namespace BlueprintSearchMetaDataHelpers;
+									SearchData.Value = GatherBlueprintSearchMetadata<FFindInBlueprintJsonWriter>(Blueprint);
+									SearchData.VersionInfo = FSearchDataVersionInfo::Current;
 
-				if (!bIsAssetReadOnlyOnDisk || !bCheckOutAndSave)
+									if (SearchData.Value.Len() > 0)
+									{
+										bEnqueueForAsyncIndexing = ShouldFullyIndexAssets();
+									}
+									else
+									{
+										SearchData.StateFlags |= ESearchDataStateFlags::IsIndexed;
+									}
+								}
+
+								// Update the entry in the search database. Don't use search data after this line.
+								FFindInBlueprintSearchManager::Get().ApplySearchDataToDatabase(MoveTemp(SearchData));
+							}
+
+							if (bEnqueueForAsyncIndexing)
+							{
+								// Enqueue this asset path to restart async indexing with the updated search metadata.
+								AssetsPendingAsyncIndexing.Enqueue(AssetPath);
+							}
+							else
+							{
+								// We're done indexing this asset path, so invoke the completion callback.
+								IndexCompletedForAssetPath(AssetPath);
+							}
+						}
+
+						// Indicate that gather work is no longer in progress.
+						bIsGatheringSearchMetadata = false;
+					}
+
+					// Block additional task completions until we're done processing.
+					FScopeLock Lock(&AsyncTaskCompletionMutex);
+
+					// Process each completed asset path.
+					for (const FName& AssetPath : CompletedAsyncTaskAssets)
+					{
+						// Execute the completion callback, if bound.
+						CacheParams.OnCached.ExecuteIfBound(AssetPath);
+
+						// Increment the counter for progress/UI display.
+						++TickCacheIndex;
+					}
+
+					// Reset for the next tick.
+					CompletedAsyncTaskAssets.Empty();
+				}
+				else
 				{
-					if (!FFindInBlueprintSearchManager::Get().IsUnindexedCacheInProgress())
+					// Generate the metadata tag value if it was not previously cached or loaded.
+					FName AssetPath = UncachedAssets[TickCacheIndex];
+					FSearchData SearchData = FFindInBlueprintSearchManager::Get().GetSearchDataForAssetPath(AssetPath);
+					if (SearchData.IsValid() && !SearchData.IsMarkedForDeletion() && SearchData.Value.Len() == 0)
 					{
-						// Re-index and update the cached value for loaded Blueprint assets only
-						if (AssetData.IsAssetLoaded())
+						if (UBlueprint* Blueprint = SearchData.Blueprint.Get())
 						{
-							if (UBlueprint* LoadedBlueprintAsset = Cast<UBlueprint>(AssetData.GetAsset()))
-							{
-								FFindInBlueprintSearchManager::Get().AddOrUpdateBlueprintSearchMetadata(LoadedBlueprintAsset, true);
-							}
+							using namespace BlueprintSearchMetaDataHelpers;
+							SearchData.Value = GatherBlueprintSearchMetadata<FFindInBlueprintJsonWriter>(Blueprint);
+							SearchData.VersionInfo = FSearchDataVersionInfo::Current;
 						}
 					}
-					else
+
+					// Signal that this deferred asset has now been indexed.
+					SearchData.StateFlags |= ESearchDataStateFlags::IsIndexed;
+
+					// Update the entry in the search database.
+					FFindInBlueprintSearchManager::Get().ApplySearchDataToDatabase(MoveTemp(SearchData));
+
+					// Execute the completion callback, if bound.
+					CacheParams.OnCached.ExecuteIfBound(AssetPath);
+
+					// Increment the counter for progress/UI display.
+					++TickCacheIndex;
+				}
+			}
+			else
+			{
+				const bool bIncludeOnlyOnDiskAssets = false;	// @todo - time this false (default) vs. true
+				FAssetRegistryModule* AssetRegistryModule = &FModuleManager::LoadModuleChecked<FAssetRegistryModule>(TEXT("AssetRegistry"));
+				FAssetData AssetData = AssetRegistryModule->Get().GetAssetByObjectPath(UncachedAssets[TickCacheIndex], bIncludeOnlyOnDiskAssets);
+				if (AssetData.IsValid())
+				{
+					const bool bIsWorldAsset = AssetData.AssetClass == UWorld::StaticClass()->GetFName();
+
+					// Construct a full package filename with path so we can query the read only status and save to disk
+					FString FinalPackageFilename = FPackageName::LongPackageNameToFilename(AssetData.PackageName.ToString());
+					if (FinalPackageFilename.Len() > 0 && FPaths::GetExtension(FinalPackageFilename).Len() == 0)
+					{
+						FinalPackageFilename += bIsWorldAsset ? FPackageName::GetMapPackageExtension() : FPackageName::GetAssetPackageExtension();
+					}
+					FText ErrorMessage;
+					bool bValidFilename = FFileHelper::IsFilenameValidForSaving(FinalPackageFilename, ErrorMessage);
+					if (bValidFilename)
+					{
+						bValidFilename = bIsWorldAsset ? FEditorFileUtils::IsValidMapFilename(FinalPackageFilename, ErrorMessage) : FPackageName::IsValidLongPackageName(FinalPackageFilename, false, &ErrorMessage);
+					}
+
+					const bool bCheckOutAndSave = EnumHasAnyFlags(CacheParams.OpFlags, EFiBCacheOpFlags::CheckOutAndSave);
+
+					bool bIsAssetReadOnlyOnDisk = IFileManager::Get().IsReadOnly(*FinalPackageFilename);
+					bool bFailedToCache = bCheckOutAndSave;
+
+					if (!bIsAssetReadOnlyOnDisk || !bCheckOutAndSave)
 					{
 						UObject* Asset = AssetData.GetAsset();
 						if (Asset && bCheckOutAndSave)
@@ -1492,23 +1805,24 @@ public:
 							}
 						}
 					}
-				}
 
-				if (bFailedToCache)
-				{
-					FailedToCacheList.Add(UncachedAssets[TickCacheIndex]);
+					if (bFailedToCache)
+					{
+						FailedToCacheList.Add(UncachedAssets[TickCacheIndex]);
+					}
+					else
+					{
+						CacheParams.OnCached.ExecuteIfBound(UncachedAssets[TickCacheIndex]);
+					}
+
+					++TickCacheIndex;
 				}
 				else
 				{
-					CacheParams.OnCached.ExecuteIfBound(UncachedAssets[TickCacheIndex]);
+					FailedToCacheList.Add(UncachedAssets[TickCacheIndex]);
+					++TickCacheIndex;
 				}
 			}
-			else
-			{
-				FailedToCacheList.Add(UncachedAssets[TickCacheIndex]);
-			}
-
-			++TickCacheIndex;
 
 			// Check if done caching Blueprints
 			if(TickCacheIndex == UncachedAssets.Num())
@@ -1529,6 +1843,11 @@ protected:
 	/** Completes a successful caching process */
 	void Finish()
 	{
+		if (AsyncIndexBuilderTask.IsValid())
+		{
+			AsyncIndexBuilderTask->EnsureCompletion();
+		}
+
 		if (ProgressNotification.IsValid())
 		{
 			ProgressNotification.Pin()->SetCompletionState(SNotificationItem::CS_Success);
@@ -1541,14 +1860,10 @@ protected:
 		CacheParams.OnFinished.ExecuteIfBound();
 
 		// Note: This will effectively delete this instance. It should not be used after this!
-		FFindInBlueprintSearchManager::Get().FinishedCachingBlueprints(TickCacheIndex, FailedToCacheList);
+		FFindInBlueprintSearchManager::Get().FinishedCachingBlueprints(CacheParams.OpType, CacheParams.OpFlags, TickCacheIndex, FailedToCacheList);
 	}
 
 private:
-
-	/** The current index, increases at a rate of once per tick */
-	int32 TickCacheIndex;
-
 	/** The list of assets that are in the process of being cached */
 	TArray<FName> UncachedAssets;
 
@@ -1558,14 +1873,38 @@ private:
 	/** Set of Blueprints that failed to be saved */
 	TSet<FName> FailedToCacheList;
 
+	/** Parameters for task configuration */
+	FCacheParams CacheParams;
+
+	/** Async search index builder task */
+	TUniquePtr<FAsyncSearchIndexTaskRunnable> AsyncIndexBuilderTask;
+
+	/** The current index, increases at a rate of once per tick */
+	int32 TickCacheIndex;
+
+	/** Tracks the next available work index for async index builder tasks */
+	TAtomic<int32> AsyncTaskBatchIndex;
+
+	/** Tracks completed async index builder tasks since the previous tick */
+	TSet<FName> CompletedAsyncTaskAssets;
+
+	/** Synchronize between async index builder task worker threads and completion logic */
+	FCriticalSection AsyncTaskCompletionMutex;
+
+	/** Thread-safe queue for tracking asset paths that need to gather search metadata from a loaded object. This must be done on the main thread */
+	TQueue<FName, EQueueMode::Mpsc> AssetsPendingGatherQueue;
+
+	/** Thread-safe queue for tracking asset paths that have exited the gather queue on the main thread and have now been re-queued for async indexing */
+	TQueue<FName, EQueueMode::Spsc> AssetsPendingAsyncIndexing;
+
+	/** TRUE if we're busy gathering search metadata from a loaded object on the main thread */
+	TAtomic<bool> bIsGatheringSearchMetadata;
+
 	/** TRUE if the caching process is started */
 	bool bIsStarted;
 
 	/** TRUE if the user has requested to cancel the caching process */
 	bool bIsCancelled;
-
-	/** Parameters for task configuration */
-	FCacheParams CacheParams;
 };
 
 FFindInBlueprintSearchManager& FFindInBlueprintSearchManager::Get()
@@ -1582,12 +1921,14 @@ FFindInBlueprintSearchManager& FFindInBlueprintSearchManager::Get()
 FFindInBlueprintSearchManager::FFindInBlueprintSearchManager()
 	: AssetRegistryModule(nullptr)
 	, CachingObject(nullptr)
+	, AsyncTaskBatchSize(0)
+	, bIsPausing(false)
+	, bHasFirstSearchOccurred(false)
 	, bEnableGatheringData(true)
 	, bDisableDeferredIndexing(false)
 	, bEnableCSVStatsProfiling(false)
 	, bEnableDeveloperMenuTools(false)
 	, bDisableSearchResultTemplates(false)
-	, bIsPausing(false)
 {
 	for (int32 TabIdx = 0; TabIdx < UE_ARRAY_COUNT(GlobalFindResultsTabIDs); TabIdx++)
 	{
@@ -1623,7 +1964,9 @@ FFindInBlueprintSearchManager::~FFindInBlueprintSearchManager()
 void FFindInBlueprintSearchManager::Initialize()
 {
 	// Init configuration
+	GConfig->GetInt(TEXT("BlueprintSearchSettings"), TEXT("AsyncTaskBatchSize"), AsyncTaskBatchSize, GEditorIni);
 	GConfig->GetBool(TEXT("BlueprintSearchSettings"), TEXT("bDisableDeferredIndexing"), bDisableDeferredIndexing, GEditorIni);
+	GConfig->GetBool(TEXT("BlueprintSearchSettings"), TEXT("bDisableThreadedIndexing"), bDisableThreadedIndexing, GEditorIni);
 	GConfig->GetBool(TEXT("BlueprintSearchSettings"), TEXT("bEnableCsvStatsProfiling"), bEnableCSVStatsProfiling, GEditorIni);
 	GConfig->GetBool(TEXT("BlueprintSearchSettings"), TEXT("bEnableDeveloperMenuTools"), bEnableDeveloperMenuTools, GEditorIni);
 	GConfig->GetBool(TEXT("BlueprintSearchSettings"), TEXT("bDisableSearchResultTemplates"), bDisableSearchResultTemplates, GEditorIni);
@@ -1715,7 +2058,7 @@ void FFindInBlueprintSearchManager::OnAssetAdded(const FAssetData& InAssetData)
 		UBlueprint* Blueprint   = Handler->RetrieveBlueprint(AssetObject);
 		if (Blueprint)
 		{
-			AddOrUpdateBlueprintSearchMetadata(Blueprint);
+			AddOrUpdateBlueprintSearchMetadata(Blueprint, EAddOrUpdateBlueprintSearchMetadataFlags::AllowNewEntry);
 		}
 	}
 	else if (Handler->AssetContainsBlueprint(InAssetData))
@@ -1757,7 +2100,9 @@ void FFindInBlueprintSearchManager::ExtractUnloadedFiBData(const FAssetData& InA
 	CSV_SCOPED_TIMING_STAT(FindInBlueprint, ExtractUnloadedFiBData);
 	CSV_CUSTOM_STAT(FindInBlueprint, ExtractUnloadedCountThisFrame, 1, ECsvCustomStatOp::Accumulate);
 
-	if (SearchMap.Contains(InAssetData.ObjectPath))
+	// Check whether this asset has already had its search data cached.
+	FSearchData SearchData = GetSearchDataForAssetPath(InAssetData.ObjectPath);
+	if (SearchData.IsValid())
 	{
 		return;
 	}
@@ -1813,14 +2158,38 @@ void FFindInBlueprintSearchManager::ExtractUnloadedFiBData(const FAssetData& InA
 		}
 	}
 
-	NewSearchData.bMarkedForDeletion = false;
 	NewSearchData.Value = *InFiBData;
 
 	// This will be set to 'None' if the data is versioned. Deserialization of the actual version from the tag value is deferred until later.
 	NewSearchData.VersionInfo.FiBDataVersion = InFiBDataVersion;
 
-	// Since the asset was not loaded, pull out the searchable data stored in the asset
-	AddSearchDataToDatabase(MoveTemp(NewSearchData));
+	// In these modes, or if there is no tag data, no additional indexing work is deferred for unloaded assets.
+	if (!bDisableDeferredIndexing && !bDisableThreadedIndexing && NewSearchData.Value.Len() > 0)
+	{
+		// Add it to the list of assets that require a full index rebuild from the metadata. This work will not block the main thread and is decoupled from the search thread.
+		PendingAssets.Add(NewSearchData.AssetPath);
+	}
+	else
+	{
+		// We're not going to defer any indexing work, so mark it as having been indexed (i.e. it's now searchable).
+		NewSearchData.StateFlags |= ESearchDataStateFlags::IsIndexed;
+	}
+
+	// Add search data for the asset into the database
+	const bool bAllowNewEntry = true;
+	ApplySearchDataToDatabase(MoveTemp(NewSearchData), bAllowNewEntry);
+}
+
+FSearchData FFindInBlueprintSearchManager::GetSearchDataForIndex(int32 CacheIndex)
+{
+	FScopeLock ScopeLock(&SafeModifyCacheCriticalSection);
+
+	if (CacheIndex >= 0 && CacheIndex < SearchArray.Num())
+	{
+		return SearchArray[CacheIndex];
+	}
+
+	return FSearchData();
 }
 
 FSearchData FFindInBlueprintSearchManager::GetSearchDataForAssetPath(FName InAssetPath)
@@ -1841,25 +2210,41 @@ FSearchData FFindInBlueprintSearchManager::GetSearchDataForAssetPath(FName InAss
 	return FSearchData();
 }
 
-int32 FFindInBlueprintSearchManager::AddSearchDataToDatabase(FSearchData InSearchData)
+void FFindInBlueprintSearchManager::ApplySearchDataToDatabase(FSearchData InSearchData, bool bAllowNewEntry)
 {
-	FName AssetPath = InSearchData.AssetPath; // Copy before we move the data into the array
+	FScopeLock ScopeLock(&SafeModifyCacheCriticalSection);
+
+	const int32* ArrayIdx = SearchMap.Find(InSearchData.AssetPath);
+	if (ArrayIdx)
+	{
+		checkf(*ArrayIdx < SearchArray.Num(),
+			TEXT("ArrayIdx:%d, SearchArray.Num():%d"),
+			*ArrayIdx,
+			SearchArray.Num());
 
-	int32 ArrayIndex = SearchArray.Add(MoveTemp(InSearchData));
+		SearchArray[*ArrayIdx] = MoveTemp(InSearchData);
+	}
+	else if (bAllowNewEntry)
+	{
+		FName AssetPath = InSearchData.AssetPath; // Copy before we move the data into the array
 
-	// Add the asset file path to the map along with the index into the array
-	SearchMap.Add(AssetPath, ArrayIndex);
+		int32 ArrayIndex = SearchArray.Add(MoveTemp(InSearchData));
 
-	return ArrayIndex;
+		// Add the asset file path to the map along with the index into the array
+		SearchMap.Add(AssetPath, ArrayIndex);
+	}
 }
 
 void FFindInBlueprintSearchManager::RemoveBlueprintByPath(FName InPath)
 {
+	FScopeLock ScopeLock(&SafeModifyCacheCriticalSection);
+
 	int32* SearchIdx = SearchMap.Find(InPath);
 
 	if(SearchIdx)
 	{
-		SearchArray[*SearchIdx].bMarkedForDeletion = true;
+		// Stale entries are flagged to be removed later, which happens after a GC pass.
+		SearchArray[*SearchIdx].StateFlags |= ESearchDataStateFlags::WasRemoved;
 	}
 }
 void FFindInBlueprintSearchManager::OnAssetRemoved(const struct FAssetData& InAssetData)
@@ -1882,6 +2267,18 @@ void FFindInBlueprintSearchManager::OnAssetRenamed(const struct FAssetData& InAs
 void FFindInBlueprintSearchManager::OnAssetRegistryFilesLoaded()
 {
 	CSV_EVENT(FindInBlueprint, TEXT("OnAssetRegistryFilesLoaded"));
+
+	if (!IsCacheInProgress() && PendingAssets.Num() == 0)
+	{
+		// Invoke the completion callback on any active global FiB tabs that are currently open to signal that the discovery stage is complete.
+		for (TWeakPtr<SFindInBlueprints> FindResultsPtr : GlobalFindResults)
+		{
+			if (FindResultsPtr.IsValid())
+			{
+				FindResultsPtr.Pin()->OnCacheComplete(EFiBCacheOpType::CachePendingAssets, EFiBCacheOpFlags::None);
+			}
+		}
+	}
 }
 
 void FFindInBlueprintSearchManager::OnAssetLoaded(UObject* InAsset)
@@ -1896,15 +2293,18 @@ void FFindInBlueprintSearchManager::OnAssetLoaded(UObject* InAsset)
 		// Find and update the item in the search array. Searches may currently be active, this will do no harm to them
 
 		// Confirm that the Blueprint has not been added already, this can occur during duplication of Blueprints.
-		int32* IndexPtr = SearchMap.Find(AssetPath);
+		FSearchData SearchData = GetSearchDataForAssetPath(AssetPath);
 
 		// The asset registry might not have informed us of this asset yet.
-		if(IndexPtr)
+		if (SearchData.IsValid())
 		{
 			// That index should never have a Blueprint already, but if it does, it should be the same Blueprint!
-			ensureMsgf(!SearchArray[*IndexPtr].Blueprint.IsValid() || SearchArray[*IndexPtr].Blueprint == BlueprintObject, TEXT("Blueprint in database has path %s and is being stomped by %s"), *(SearchArray[*IndexPtr].AssetPath.ToString()), *AssetPath.ToString());
-			ensureMsgf(!SearchArray[*IndexPtr].Blueprint.IsValid() || SearchArray[*IndexPtr].AssetPath == AssetPath, TEXT("Blueprint in database has path %s and is being stomped by %s"), *(SearchArray[*IndexPtr].AssetPath.ToString()), *AssetPath.ToString());
-			SearchArray[*IndexPtr].Blueprint = BlueprintObject;
+			ensureMsgf(!SearchData.Blueprint.IsValid() || SearchData.Blueprint == BlueprintObject, TEXT("Blueprint in database has path %s and is being stomped by %s"), *(SearchData.AssetPath.ToString()), *AssetPath.ToString());
+			ensureMsgf(!SearchData.Blueprint.IsValid() || SearchData.AssetPath == AssetPath, TEXT("Blueprint in database has path %s and is being stomped by %s"), *(SearchData.AssetPath.ToString()), *AssetPath.ToString());
+			SearchData.Blueprint = BlueprintObject;
+
+			// Apply the updated entry to the database.
+			ApplySearchDataToDatabase(MoveTemp(SearchData));
 		}
 
 		UnindexedAssets.Remove(AssetPath);
@@ -1921,7 +2321,7 @@ void FFindInBlueprintSearchManager::OnHotReload(bool bWasTriggeredAutomatically)
 	CachedAssetClasses.Reset();
 }
 
-void FFindInBlueprintSearchManager::AddOrUpdateBlueprintSearchMetadata(UBlueprint* InBlueprint, bool bInForceReCache/* = false*/)
+void FFindInBlueprintSearchManager::AddOrUpdateBlueprintSearchMetadata(UBlueprint* InBlueprint, EAddOrUpdateBlueprintSearchMetadataFlags InFlags/* = EAddOrUpdateBlueprintSearchMetadataFlags::None*/)
 {
 	CSV_SCOPED_TIMING_STAT(FindInBlueprint, AddOrUpdateBlueprintSearchMetadata);
 	CSV_CUSTOM_STAT(FindInBlueprint, AddOrUpdateCountThisFrame, 1, ECsvCustomStatOp::Accumulate);
@@ -1938,62 +2338,91 @@ void FFindInBlueprintSearchManager::AddOrUpdateBlueprintSearchMetadata(UBlueprin
 		return;
 	}
 
+	// Control flags
+	const bool bForceRecache = EnumHasAllFlags(InFlags, EAddOrUpdateBlueprintSearchMetadataFlags::ForceRecache);
+	const bool bAllowNewEntry = EnumHasAllFlags(InFlags, EAddOrUpdateBlueprintSearchMetadataFlags::AllowNewEntry);
+
 	UObject* AssetObject = GetAssetObject(InBlueprint);
 
 	check(InBlueprint && AssetObject);
 
-	// Allow only one thread modify the search data at a time
-	FScopeLock ScopeLock(&SafeModifyCacheCriticalSection);
-
 	FName AssetPath = *AssetObject->GetPathName();
+	FSearchData SearchData = GetSearchDataForAssetPath(AssetPath);
 
-	int32* IndexPtr = SearchMap.Find(AssetPath);
-	int32 Index = 0;
-	if(!IndexPtr)
+	if (SearchData.IsValid())
 	{
-		FSearchData SearchData;
-		SearchData.Blueprint = InBlueprint;
-		SearchData.AssetPath = AssetPath;
-		Index = AddSearchDataToDatabase(MoveTemp(SearchData));
+		SearchData.Blueprint = InBlueprint; // Blueprint instance may change due to reloading
+		SearchData.StateFlags &= ~(ESearchDataStateFlags::IsIndexed | ESearchDataStateFlags::WasRemoved);
 	}
-	else
+	else if (bAllowNewEntry)
 	{
-		Index = *IndexPtr;
-		SearchArray[Index].Blueprint = InBlueprint; // Blueprint instance may change due to reloading
-		SearchArray[Index].bMarkedForDeletion = false;
+		SearchData = FSearchData();
+		SearchData.AssetPath = AssetPath;
+		SearchData.Blueprint = InBlueprint;
 	}
 
 	// Build the search data
 	if (FProperty* ParentClassProp = InBlueprint->GetClass()->FindPropertyByName(GET_MEMBER_NAME_CHECKED(UBlueprint, ParentClass)))
 	{
-		ParentClassProp->ExportTextItem(SearchArray[Index].ParentClass, ParentClassProp->ContainerPtrToValuePtr<uint8>(InBlueprint), nullptr, InBlueprint, 0);
+		ParentClassProp->ExportTextItem(SearchData.ParentClass, ParentClassProp->ContainerPtrToValuePtr<uint8>(InBlueprint), nullptr, InBlueprint, 0);
 	}
 
-	// During unindexed/out-of-date caching we will arrive here as a result of loading the asset, so don't remove the IsCacheInProgress() check!
-	if (bInForceReCache || IsCacheInProgress() || bDisableDeferredIndexing)
+	if (SearchData.IsValid())
 	{
-		// Cannot successfully gather most searchable data if there is no SkeletonGeneratedClass, so don't try, leave it as whatever it was last set to
-		if (InBlueprint->SkeletonGeneratedClass != nullptr)
+		// Clear any previously-gathered data.
+		SearchData.Value.Empty();
+
+		// During unindexed/out-of-date caching we will arrive here as a result of loading the asset, so don't remove the IsUnindexedCacheInProgress() check!
+		if (bForceRecache || IsUnindexedCacheInProgress() || bDisableDeferredIndexing)
 		{
-			using namespace BlueprintSearchMetaDataHelpers;
+			// Cannot successfully gather most searchable data if there is no SkeletonGeneratedClass, so don't try, leave it as whatever it was last set to
+			if (InBlueprint->SkeletonGeneratedClass != nullptr)
+			{
+				using namespace BlueprintSearchMetaDataHelpers;
+
+				// Update search metadata string content
+				SearchData.Value = GatherBlueprintSearchMetadata<FFindInBlueprintJsonWriter>(InBlueprint);
+
+				// Update version info stored in database to latest
+				SearchData.VersionInfo = FSearchDataVersionInfo::Current;
+			}
 
-			// Update search metadata string content
-			SearchArray[Index].Value = GatherBlueprintSearchMetadata<FFindInBlueprintJsonWriter>(InBlueprint);
+			// Mark it as having been indexed in the following cases:
+			// a) Deferred or multithreaded indexing is disabled. In these cases, indexing will be handled by the search thread.
+			// b) There is no generated metadata. In this case, there's nothing to search, so there's no need to index the asset.
+			if (bDisableDeferredIndexing || bDisableThreadedIndexing || SearchData.Value.Len() == 0)
+			{
+				// Mark it as having been indexed (it's now searchable).
+				SearchData.StateFlags |= ESearchDataStateFlags::IsIndexed;
 
-			// Update version info stored in database to latest
-			SearchArray[Index].VersionInfo = FSearchDataVersionInfo::Current;
+				// Remove it from the list of pending assets (if it exists).
+				PendingAssets.Remove(AssetPath);
+			}
+			else if(!bDisableThreadedIndexing)
+			{
+				// With multithreaded indexing, we defer additional work beyond regenerating the metadata, so ensure it's been added.
+				PendingAssets.Add(AssetPath);
+			}
+		}
+		else if (!bDisableDeferredIndexing)
+		{
+			// Add it to the list of assets to be indexed (deferred)
+			PendingAssets.Add(AssetPath);
 		}
 
-		// Remove it from the list of pending assets (if it exists)
-		PendingAssets.Remove(AssetPath);
-	}
-	else
-	{
-		// Add it to the list of assets to be indexed (deferred until the next global search)
-		PendingAssets.Add(AssetPath);
+		// Copy new/updated search data into the cache
+		ApplySearchDataToDatabase(MoveTemp(SearchData), bAllowNewEntry);
 	}
 }
 
+FFindInBlueprintSearchManager::FActiveSearchQueryPtr FFindInBlueprintSearchManager::FindSearchQuery(const class FStreamSearch* InSearchOriginator) const
+{
+	// Must lock this behind a critical section to ensure that no other thread is accessing it at the same time
+	FScopeLock ScopeLock(&SafeQueryModifyCriticalSection);
+
+	return ActiveSearchQueries.FindRef(InSearchOriginator);
+}
+
 void FFindInBlueprintSearchManager::BeginSearchQuery(const FStreamSearch* InSearchOriginator)
 {
 	if (AssetRegistryModule == nullptr)
@@ -2001,18 +2430,24 @@ void FFindInBlueprintSearchManager::BeginSearchQuery(const FStreamSearch* InSear
 		UE_LOG(LogBlueprint, Warning, TEXT("Find-in-Blueprints was not fully initialized, possibly due to problems being initialized while saving a package. Please explicitly initialize earlier!"));
 	}
 
+	// Unblock caching operations from doing a full indexing pass on pending assets
+	bHasFirstSearchOccurred = true;
+
 	// Cannot begin a search thread while saving
 	FScopeLock ScopeLock(&PauseThreadsCriticalSection);
 	FScopeLock ScopeLock2(&SafeQueryModifyCriticalSection);
 
+	FActiveSearchQueryPtr NewSearchQuery = MakeShared<FActiveSearchQuery, ESPMode::ThreadSafe>();
+	check(NewSearchQuery.IsValid());
+
 	ActiveSearchCounter.Increment();
-	ActiveSearchQueries.FindOrAdd(InSearchOriginator) = 0;
+	ActiveSearchQueries.Add(InSearchOriginator, NewSearchQuery);
 }
 
 bool FFindInBlueprintSearchManager::ContinueSearchQuery(const FStreamSearch* InSearchOriginator, FSearchData& OutSearchData)
 {
 	// Check if the thread has been told to pause, this occurs for the Garbage Collector and for saving to disk
-	if(bIsPausing == true)
+	if (bIsPausing == true)
 	{
 		// Pause all searching, the GC is running and we will also be saving the database
 		ActiveSearchCounter.Decrement();
@@ -2020,43 +2455,112 @@ bool FFindInBlueprintSearchManager::ContinueSearchQuery(const FStreamSearch* InS
 		ActiveSearchCounter.Increment();
 	}
 
-	// Must lock this behind a critical section to ensure that no other thread is accessing it at the same time
-	FScopeLock ScopeLock(&SafeQueryModifyCriticalSection);
+	FActiveSearchQueryPtr SearchQuery = FindSearchQuery(InSearchOriginator);
+	if (!SearchQuery.IsValid())
+	{
+		return false;
+	}
 
-	if(int32* SearchIdxPtr = ActiveSearchQueries.Find(InSearchOriginator))
+	auto GetNextSearchDataLambda = [this, SearchQuery]()
 	{
-		int32& SearchIdx = *SearchIdxPtr;
-		while(SearchIdx < SearchArray.Num())
+		// Get the entry in the index cache for the next asset to search.
+		FSearchData SearchData = GetSearchDataForIndex(SearchQuery->NextIndex);
+		while (SearchData.IsValid())
 		{
-			// If the Blueprint is not marked for deletion, and the asset is valid, we will check to see if we want to refresh the searchable data.
-			if( SearchArray[SearchIdx].bMarkedForDeletion || (SearchArray[SearchIdx].Blueprint.IsValid() && SearchArray[SearchIdx].Blueprint->IsPendingKill()) )
+			// Advance the current search index.
+			++SearchQuery->NextIndex;
+
+			// If this asset has not been indexed, don't search it yet.
+			if (!SearchData.IsIndexingCompleted())
 			{
-				// Mark it for deletion, it will be removed on next save
-				SearchArray[SearchIdx].bMarkedForDeletion = true;
+				SearchQuery->DeferredAssetPaths.Enqueue(SearchData.AssetPath);
 			}
-			else
+			else if (!SearchData.IsMarkedForDeletion())
 			{
-				// If there is FiB data, parse it into an ImaginaryBlueprint
-				if (SearchArray[SearchIdx].Value.Len() > 0)
+				// Ok to use this asset's index entry; break out of the loop.
+				break;
+			}
+
+			// Move on to the next entry in the index cache.
+			SearchData = GetSearchDataForIndex(SearchQuery->NextIndex);
+		}
+
+		// If we don't have valid search data, try the deferred list from above.
+		if (!SearchData.IsValid())
+		{
+			FName AssetPath;
+			while (SearchQuery->DeferredAssetPaths.Dequeue(AssetPath))
+			{
+				SearchData = GetSearchDataForAssetPath(AssetPath);
+				if (SearchData.IsValid() && !SearchData.IsMarkedForDeletion())
 				{
-					const FString AssetPath = SearchArray[SearchIdx].AssetPath.ToString();
-					if (FiBSerializationHelpers::ValidateSearchDataVersionInfo(AssetPath, SearchArray[SearchIdx].Value, SearchArray[SearchIdx].VersionInfo))
+					// If this asset is still waiting to be indexed, put it back into the queue.
+					if (!SearchData.IsIndexingCompleted())
+					{
+						SearchQuery->DeferredAssetPaths.Enqueue(AssetPath);
+
+						// Yield to allow indexing to progress a bit further.
+						FPlatformProcess::Sleep(0.1f);
+					}
+					else
 					{
-						SearchArray[SearchIdx].ImaginaryBlueprint = MakeShareable(new FImaginaryBlueprint(FPaths::GetBaseFilename(AssetPath), AssetPath, SearchArray[SearchIdx].ParentClass, SearchArray[SearchIdx].Interfaces, SearchArray[SearchIdx].Value, SearchArray[SearchIdx].VersionInfo));
-						SearchArray[SearchIdx].Value.Empty();
+						// Ok to use this entry; break out of the loop.
+						break;
 					}
 				}
-
- 				OutSearchData = SearchArray[SearchIdx++];
-				return true;
 			}
+		}
+
+		// Increment the search counter if we have valid search data.
+		if (SearchData.IsValid())
+		{
+			++SearchQuery->SearchCount;
+		}
+
+		return SearchData;
+	};
+
+	FSearchData SearchData = GetNextSearchDataLambda();
+
+	if (SearchData.IsValid())
+	{
+		// In these modes, the full index may not have been parsed yet. We'll do that now on the search thread.
+		if (bDisableDeferredIndexing || bDisableThreadedIndexing)
+		{
+			// Must lock this behind a critical section to ensure that no other thread is accessing it at the same time
+			FScopeLock Lock(&SafeModifyCacheCriticalSection);
+
+			// Grab the latest entry from the cache. In the case of parallel global searches, one of the other threads may have already updated it while we were blocked by the mutex above.
+			SearchData = GetSearchDataForAssetPath(SearchData.AssetPath);
 
-			++SearchIdx;
+			// If there is FiB data, parse it into an ImaginaryBlueprint
+			if (SearchData.IsValid() && SearchData.Value.Len() > 0)
+			{
+				const FString AssetPath = SearchData.AssetPath.ToString();
+				if (FiBSerializationHelpers::ValidateSearchDataVersionInfo(AssetPath, SearchData.Value, SearchData.VersionInfo))
+				{
+					SearchData.ImaginaryBlueprint = MakeShareable(new FImaginaryBlueprint(FPaths::GetBaseFilename(SearchData.AssetPath.ToString()), SearchData.AssetPath.ToString(), SearchData.ParentClass, SearchData.Interfaces, SearchData.Value, SearchData.VersionInfo));
+					SearchData.Value.Empty();
+
+					// In the case of parallel global searches, two search threads may be looking at the same entry. Thus, we only allow one search thread to be actively parsing JSON nodes.
+					SearchData.ImaginaryBlueprint->EnableInterlockedParsing();
+				}
+
+				// Update the entry in the database
+				ApplySearchDataToDatabase(SearchData);
+			}
 		}
+
+		OutSearchData = SearchData;
+		return true;
 	}
+	else
+	{
+		FScopeLock ScopeLock(&SafeQueryModifyCriticalSection);
 
-	ActiveSearchQueries.Remove(InSearchOriginator);
-	ActiveSearchCounter.Decrement();
+		ActiveSearchQueries.Remove(InSearchOriginator);
+		ActiveSearchCounter.Decrement();
+	}
 
 	return false;
 }
@@ -2065,10 +2569,9 @@ void FFindInBlueprintSearchManager::EnsureSearchQueryEnds(const class FStreamSea
 {
 	// Must lock this behind a critical section to ensure that no other thread is accessing it at the same time
 	FScopeLock ScopeLock(&SafeQueryModifyCriticalSection);
-	int32* SearchIdxPtr = ActiveSearchQueries.Find(InSearchOriginator);
 
 	// If the search thread is still considered active, remove it
-	if(SearchIdxPtr)
+	if (ActiveSearchQueries.Contains(InSearchOriginator))
 	{
 		ActiveSearchQueries.Remove(InSearchOriginator);
 		ActiveSearchCounter.Decrement();
@@ -2078,19 +2581,19 @@ void FFindInBlueprintSearchManager::EnsureSearchQueryEnds(const class FStreamSea
 float FFindInBlueprintSearchManager::GetPercentComplete(const FStreamSearch* InSearchOriginator) const
 {
 	FScopeLock ScopeLock(&SafeQueryModifyCriticalSection);
-	const int32* SearchIdxPtr = ActiveSearchQueries.Find(InSearchOriginator);
+	FActiveSearchQueryPtr SearchQuery = FindSearchQuery(InSearchOriginator);
 
 	float ReturnPercent = 0.0f;
 
-	if(SearchIdxPtr)
+	if (SearchQuery.IsValid())
 	{
-		ReturnPercent = (float)*SearchIdxPtr / (float)SearchArray.Num();
+		ReturnPercent = (float)SearchQuery->SearchCount / (float)SearchArray.Num();
 	}
 
 	return ReturnPercent;
 }
 
-const FSearchData* FFindInBlueprintSearchManager::QuerySingleBlueprint(UBlueprint* InBlueprint, bool bInRebuildSearchData)
+FSearchData FFindInBlueprintSearchManager::QuerySingleBlueprint(UBlueprint* InBlueprint, bool bInRebuildSearchData)
 {
 	// AddOrUpdateBlueprintSearchMetadata would fail to cache any data for a Blueprint loaded specifically for diffing, but the bigger question
 	// here in this function is how you are doing a search specifically for data within this Blueprint. This function is limited to be called
@@ -2101,22 +2604,17 @@ const FSearchData* FFindInBlueprintSearchManager::QuerySingleBlueprint(UBlueprin
 		if (bInRebuildSearchData)
 		{
 			// Update the Blueprint, make sure it is fully up-to-date
-			AddOrUpdateBlueprintSearchMetadata(InBlueprint, true);
+			AddOrUpdateBlueprintSearchMetadata(InBlueprint, EAddOrUpdateBlueprintSearchMetadataFlags::ForceRecache);
 		}
 
 		UObject* AssetObject = GetAssetObject(InBlueprint);
 		check(AssetObject);
-		FName Key = *AssetObject->GetPathName();
 
-		int32* ArrayIdx = SearchMap.Find(Key);
-		if (ArrayIdx)
+		FName Key = *AssetObject->GetPathName();
+		FSearchData SearchData = GetSearchDataForAssetPath(Key);
+		if (SearchData.IsValid())
 		{
-			checkf(*ArrayIdx < SearchArray.Num(),
-				TEXT("ArrayIdx:%d, SearchArray.Num():%d"),
-				*ArrayIdx,
-				SearchArray.Num());
-
-			return &SearchArray[*ArrayIdx];
+			return SearchData;
 		}
 		else if(bInRebuildSearchData)
 		{
@@ -2130,7 +2628,7 @@ const FSearchData* FFindInBlueprintSearchManager::QuerySingleBlueprint(UBlueprin
 		// Also warn here as we do not index diff-only packages.
 		UE_LOG(LogBlueprint, Warning, TEXT("Attempted to query an old Blueprint package opened for diffing!"));
 	}
-	return nullptr;
+	return FSearchData();
 }
 
 void FFindInBlueprintSearchManager::PauseFindInBlueprintSearch()
@@ -2160,8 +2658,6 @@ void FFindInBlueprintSearchManager::UnpauseFindInBlueprintSearch()
 
 void FFindInBlueprintSearchManager::CleanCache()
 {
-	// *NOTE* SaveCache is a thread safe operation by design, all searching threads are paused during the operation so there is no critical section locking
-
 	// We need to cache where the active queries are so that we can put them back in a safe and expected position
 	TMap< const FStreamSearch*, FName > CacheQueries;
 	for( auto It = ActiveSearchQueries.CreateIterator() ; It ; ++It )
@@ -2180,12 +2676,15 @@ void FFindInBlueprintSearchManager::CleanCache()
 	TMap<FName, int32> NewSearchMap;
 	TArray<FSearchData> NewSearchArray;
 
+	// Don't allow background indexing tasks to access the search database while we fix it up.
+	FScopeLock Lock(&SafeModifyCacheCriticalSection);
+
 	for(auto& SearchValuePair : SearchMap)
 	{
 		// Here it builds the new map/array, clean of deleted content.
 
 		// If the database item is not marked for deletion and not pending kill (if loaded), keep it in the database
-		if( !SearchArray[SearchValuePair.Value].bMarkedForDeletion && !(SearchArray[SearchValuePair.Value].Blueprint.IsValid() && SearchArray[SearchValuePair.Value].Blueprint->IsPendingKill()) )
+		if (!SearchArray[SearchValuePair.Value].IsMarkedForDeletion() && !(SearchArray[SearchValuePair.Value].Blueprint.IsValid() && SearchArray[SearchValuePair.Value].Blueprint->IsPendingKill()))
 		{
 			// Build the new map/array
 			NewSearchMap.Add(SearchValuePair.Key, NewSearchArray.Add(MoveTemp(SearchArray[SearchValuePair.Value])) );
@@ -2196,7 +2695,7 @@ void FFindInBlueprintSearchManager::CleanCache()
 			PendingAssets.Remove(SearchArray[SearchValuePair.Value].AssetPath);
 
 			// Level Blueprints are destroyed when you open a new level, we need to re-add it as an unloaded asset so long as they were not marked for deletion
-			if(!SearchArray[SearchValuePair.Value].bMarkedForDeletion && FModuleManager::Get().IsModuleLoaded(TEXT("AssetRegistry")))
+			if (!SearchArray[SearchValuePair.Value].IsMarkedForDeletion() && FModuleManager::Get().IsModuleLoaded(TEXT("AssetRegistry")))
 			{
 				SearchArray[SearchValuePair.Value].Blueprint = nullptr;
 
@@ -2239,7 +2738,7 @@ void FFindInBlueprintSearchManager::CleanCache()
 	 	}
 	 
 		// Update the active search to the new index of where it is at in the search
-	 	*(ActiveSearchQueries.Find(CacheQuery.Key)) = NewMappedIndex;
+		(ActiveSearchQueries.FindRef(CacheQuery.Key))->NextIndex = NewMappedIndex;
 	}
 }
 
@@ -2312,11 +2811,8 @@ FString FFindInBlueprintSearchManager::GenerateSearchIndexForDebugging(UBlueprin
 
 void FFindInBlueprintSearchManager::OnCacheAllUnindexedAssets(bool bInSourceControlActive, bool bInCheckoutAndSave)
 {
-	// Multiple threads can be adding to this at the same time
-	FScopeLock ScopeLock(&SafeModifyCacheCriticalSection);
-
 	// We need to check validity first in case the user has closed the initiating FiB tab before responding to the source control login dialog (which is modeless).
-	if (CachingObject)
+	if (CachingObject.IsValid())
 	{
 		if(bInSourceControlActive && bInCheckoutAndSave)
 		{
@@ -2341,17 +2837,13 @@ void FFindInBlueprintSearchManager::CacheAllAssets(TWeakPtr< SFindInBlueprints >
 	// Do not start another caching process if one is in progress
 	if(!IsCacheInProgress())
 	{
-		CurrentCacheOpType = InOptions.OpType;
+		FCacheAllBlueprintsTickableObject::FCacheParams CacheParams;
+		CacheParams.OpType = InOptions.OpType;
 
-		if (CurrentCacheOpType == EFiBCacheOpType::CachePendingAssets)
+		if (CacheParams.OpType == EFiBCacheOpType::CachePendingAssets)
 		{
-			FCacheAllBlueprintsTickableObject::FCacheParams CacheParams;
 			CacheParams.OnFinished = InOptions.OnFinished;
-			CacheParams.OnCached.BindLambda([this](FName InAssetName)
-			{
-				// Ensure removal in case it did not get re-indexed for some reason (e.g. not loaded, not a Blueprint asset, etc.)
-				PendingAssets.Remove(InAssetName);
-			});
+			CacheParams.AsyncTaskBatchSize = AsyncTaskBatchSize;
 
 			// Determine if PIE is active - in that case we're potentially streaming assets in at random intervals, so just hide the progress UI while re-indexing those assets
 			const bool bIsPIESimulating = (GEditor->bIsSimulatingInEditor || GEditor->PlayWorld);
@@ -2365,14 +2857,84 @@ void FFindInBlueprintSearchManager::CacheAllAssets(TWeakPtr< SFindInBlueprints >
 			// Keep popup notifications hidden during this operation
 			CacheParams.OpFlags |= EFiBCacheOpFlags::HideNotifications;
 
+			const bool bIsAssetDiscoveryInProgress = IsAssetDiscoveryInProgress();
+
+			// Determine if this operation has been continued from a previous operation.
+			bool bIsCachingDiscoveredAssets = EnumHasAnyFlags(InOptions.OpFlags, EFiBCacheOpFlags::IsCachingDiscoveredAssets);
+			if (!bIsCachingDiscoveredAssets)
+			{
+				// During the initial asset discovery and registration stage, we'll index an unknown number of assets in the background as a continuous caching operation.
+				bIsCachingDiscoveredAssets = !bDisableDeferredIndexing && bIsAssetDiscoveryInProgress;
+			}
+
+			if (bIsCachingDiscoveredAssets)
+			{
+				CacheParams.OpFlags |= EFiBCacheOpFlags::IsCachingDiscoveredAssets;
+
+				if (bIsAssetDiscoveryInProgress)
+				{
+					// Keep progress bars hidden during asset discovery since the endpoint will be a moving target. Also keep progress onscreen between operations.
+					CacheParams.OpFlags |= EFiBCacheOpFlags::HideProgressBars;
+					CacheParams.OpFlags |= EFiBCacheOpFlags::KeepProgressVisibleOnCompletion;
+				}
+			}
+
+			if (bDisableThreadedIndexing)
+			{
+				CacheParams.OpFlags |= EFiBCacheOpFlags::ExecuteOnMainThread;
+			}
+//			else if (bEnableCSVStatsProfiling)
+//			{
+//				CacheParams.OpFlags |= EFiBCacheOpFlags::ExecuteOnSingleThread;
+//			}
+
 			// Keep track of which global FiB context started the operation (if any)
 			SourceCachingWidget = InSourceWidget;
 
-			CachingObject = new FCacheAllBlueprintsTickableObject(PendingAssets, CacheParams);
+			// If async indexing is disabled, or if it's enabled and we're allowed to do a full index pass, use the full array
+			if (bDisableThreadedIndexing || bHasFirstSearchOccurred)
+			{
+				// Create a new instance of the caching object.
+				CachingObject = MakeUnique<FCacheAllBlueprintsTickableObject>(PendingAssets, CacheParams);
+			}
+			else
+			{
+				// Find all pending assets for which we can gather and cache search metadata
+				TSet<FName> AssetsToPartiallyCache;
+				for (const FName& AssetPath : PendingAssets)
+				{
+					if (const int32* IndexPtr = SearchMap.Find(AssetPath))
+					{
+						const FSearchData& SearchData = SearchArray[*IndexPtr];
+						if (!SearchData.IsIndexingCompleted() && SearchData.Value.Len() == 0)
+						{
+							AssetsToPartiallyCache.Add(AssetPath);
+						}
+					}
+
+					// Keep track of all assets for the full indexing pass, which is currently being deferred until the first global search is initiated.
+					AssetsToIndexOnFirstSearch.Add(AssetPath);
+				}
+
+				// If we found one or more assets above...
+				if (AssetsToPartiallyCache.Num() > 0)
+				{
+					// Run a partial indexing pass only.
+					CacheParams.OpFlags |= EFiBCacheOpFlags::ExecuteGatherPhaseOnly;
+
+					// Create a new instance of the caching object.
+					CachingObject = MakeUnique<FCacheAllBlueprintsTickableObject>(MoveTemp(AssetsToPartiallyCache), CacheParams);
+				}
+			}
 
 			// Immediately start the operation (non-interactive)
-			check(CachingObject);
-			CachingObject->Start();
+			if (CachingObject.IsValid())
+			{
+				CachingObject->Start();
+			}
+
+			// Reset for the next set of incoming deferred assets.
+			PendingAssets.Empty();
 		}
 		else
 		{
@@ -2423,14 +2985,13 @@ void FFindInBlueprintSearchManager::CacheAllAssets(TWeakPtr< SFindInBlueprints >
 				TempUncachedAssets.Append(BlueprintsToUpdate);
 
 				const bool bCheckOutAndSave = (ReturnValue == EAppReturnType::Yes);
-				FCacheAllBlueprintsTickableObject::FCacheParams CacheParams;
-				CacheParams.OpFlags = EFiBCacheOpFlags::ShowProgress | EFiBCacheOpFlags::AllowUserCancel;
+				CacheParams.OpFlags = EFiBCacheOpFlags::ShowProgress | EFiBCacheOpFlags::AllowUserCancel | EFiBCacheOpFlags::AllowUserCloseProgress;
 				if (bCheckOutAndSave)
 				{
 					CacheParams.OpFlags |= EFiBCacheOpFlags::CheckOutAndSave;
 				}
 				CacheParams.OnFinished = InOptions.OnFinished;
-				CachingObject = new FCacheAllBlueprintsTickableObject(MoveTemp(TempUncachedAssets), CacheParams);
+				CachingObject = MakeUnique<FCacheAllBlueprintsTickableObject>(TempUncachedAssets, CacheParams);
 
 				const bool bIsSourceControlEnabled = ISourceControlModule::Get().IsEnabled();
 				if (!bIsSourceControlEnabled && bCheckOutAndSave)
@@ -2461,7 +3022,7 @@ void FFindInBlueprintSearchManager::CancelCacheAll(SFindInBlueprints* InFindInBl
 int32 FFindInBlueprintSearchManager::GetCurrentCacheIndex() const
 {
 	int32 CachingIndex = 0;
-	if(CachingObject)
+	if(CachingObject.IsValid())
 	{
 		CachingIndex = CachingObject->GetCurrentCacheIndex();
 	}
@@ -2472,7 +3033,7 @@ int32 FFindInBlueprintSearchManager::GetCurrentCacheIndex() const
 FName FFindInBlueprintSearchManager::GetCurrentCacheBlueprintName() const
 {
 	FName CachingBPName;
-	if(CachingObject)
+	if(CachingObject.IsValid())
 	{
 		CachingBPName = CachingObject->GetCurrentCacheBlueprintName();
 	}
@@ -2484,7 +3045,7 @@ float FFindInBlueprintSearchManager::GetCacheProgress() const
 {
 	float ReturnCacheValue = 1.0f;
 
-	if(CachingObject)
+	if(CachingObject.IsValid())
 	{
 		ReturnCacheValue = CachingObject->GetCacheProgress();
 	}
@@ -2492,11 +3053,6 @@ float FFindInBlueprintSearchManager::GetCacheProgress() const
 	return ReturnCacheValue;
 }
 
-int32 FFindInBlueprintSearchManager::GetNumberPendingAssets() const
-{
-	return PendingAssets.Num();
-}
-
 int32 FFindInBlueprintSearchManager::GetNumberUnindexedAssets() const
 {
 	return UnindexedAssets.Num();
@@ -2504,7 +3060,7 @@ int32 FFindInBlueprintSearchManager::GetNumberUnindexedAssets() const
 
 int32 FFindInBlueprintSearchManager::GetNumberUncachedAssets() const
 {
-	if (CachingObject)
+	if (CachingObject.IsValid())
 	{
 		return CachingObject->GetUncachedAssetCount();
 	}
@@ -2512,30 +3068,38 @@ int32 FFindInBlueprintSearchManager::GetNumberUncachedAssets() const
 	return 0;
 }
 
-void FFindInBlueprintSearchManager::StartedCachingBlueprints(EFiBCacheOpFlags InCacheOpFlags)
+void FFindInBlueprintSearchManager::StartedCachingBlueprints(EFiBCacheOpType InCacheOpType, EFiBCacheOpFlags InCacheOpFlags)
 {
 	// Invoke the callback on any open global widgets
 	for (TWeakPtr<SFindInBlueprints> FindResultsPtr : GlobalFindResults)
 	{
 		if (FindResultsPtr.IsValid())
 		{
-			FindResultsPtr.Pin()->OnCacheStarted(CurrentCacheOpType, InCacheOpFlags);
+			FindResultsPtr.Pin()->OnCacheStarted(InCacheOpType, InCacheOpFlags);
 		}
 	}
 }
 
-void FFindInBlueprintSearchManager::FinishedCachingBlueprints(int32 InNumberCached, TSet<FName>& InFailedToCacheList)
+void FFindInBlueprintSearchManager::FinishedCachingBlueprints(EFiBCacheOpType InCacheOpType, EFiBCacheOpFlags InCacheOpFlags, int32 InNumberCached, TSet<FName>& InFailedToCacheList)
 {
-	// Multiple threads could be adding to this at the same time
-	FScopeLock ScopeLock(&SafeModifyCacheCriticalSection);
-
 	// Update the list of cache failures
 	FailedToCachePaths = InFailedToCacheList;
 
-	// Remove any failed attempts from the pending list
-	if (InFailedToCacheList.Num() > 0)
+	bool bContinueCachingPendingAssets = false;
+	if (InCacheOpType == EFiBCacheOpType::CachePendingAssets)
 	{
-		PendingAssets = PendingAssets.Difference(InFailedToCacheList);
+		// In the discovery stage, check to see if discovery has ended and we have no additional assets to process. 
+		const bool bIsCachingDiscoveredAssets = EnumHasAnyFlags(InCacheOpFlags, EFiBCacheOpFlags::IsCachingDiscoveredAssets);
+		if (bIsCachingDiscoveredAssets && PendingAssets.Num() > 0)
+		{
+			// Discovered assets are still pending; continue immediately.
+			bContinueCachingPendingAssets = true;
+		}
+		else if (!IsAssetDiscoveryInProgress())
+		{
+			// Update the flags for the completion callback.
+			InCacheOpFlags &= ~(EFiBCacheOpFlags::HideProgressBars | EFiBCacheOpFlags::IsCachingDiscoveredAssets | EFiBCacheOpFlags::KeepProgressVisibleOnCompletion);
+		}
 	}
 
 	// Invoke the completion callback on any open widgets that are not the initiating one (e.g. to hide progress bar)
@@ -2543,32 +3107,47 @@ void FFindInBlueprintSearchManager::FinishedCachingBlueprints(int32 InNumberCach
 	{
 		if (FindResultsPtr.IsValid() && FindResultsPtr != SourceCachingWidget)
 		{
-			FindResultsPtr.Pin()->OnCacheComplete(CurrentCacheOpType);
+			FindResultsPtr.Pin()->OnCacheComplete(InCacheOpType, InCacheOpFlags);
 		}
 	}
 
 	// Invoke the completion callback on the initiating widget only if necessary (i.e. only if it doesn't already delegate it)
-	if(SourceCachingWidget.IsValid() && !CachingObject->HasPostCacheWork())
+	if (SourceCachingWidget.IsValid() && !CachingObject->HasPostCacheWork())
 	{
-		SourceCachingWidget.Pin()->OnCacheComplete(CurrentCacheOpType);
+		SourceCachingWidget.Pin()->OnCacheComplete(InCacheOpType, InCacheOpFlags);
 	}
 
-	// Reset the reference to the widget that initiated the caching operation
-	SourceCachingWidget.Reset();
-
 	// Delete the object and NULL it out so we can do it again in the future if needed (if it was canceled)
-	delete CachingObject;
-	CachingObject = nullptr;
+	CachingObject.Reset();
+
+	// If necessary, begin a new operation to continue caching the next set of pending assets.
+	if (bContinueCachingPendingAssets)
+	{
+		FFindInBlueprintCachingOptions CacheOptions;
+		CacheOptions.OpType = InCacheOpType;
+		CacheOptions.OpFlags = InCacheOpFlags;
+		CacheAllAssets(SourceCachingWidget, CacheOptions);
+	}
+	else
+	{
+		// Reset the reference to the widget that initiated the caching operation
+		SourceCachingWidget.Reset();
+	}
 }
 
 bool FFindInBlueprintSearchManager::IsCacheInProgress() const
 {
-	return CachingObject != nullptr;
+	return CachingObject.IsValid();
 }
 
 bool FFindInBlueprintSearchManager::IsUnindexedCacheInProgress() const
 {
-	return IsCacheInProgress() && CurrentCacheOpType == EFiBCacheOpType::CacheUnindexedAssets;
+	return IsCacheInProgress() && CachingObject->GetCurrentCacheOpType() == EFiBCacheOpType::CacheUnindexedAssets;
+}
+
+bool FFindInBlueprintSearchManager::IsAssetDiscoveryInProgress() const
+{
+	return AssetRegistryModule && AssetRegistryModule->Get().IsLoadingAssets();
 }
 
 TSharedPtr< FJsonObject > FFindInBlueprintSearchManager::ConvertJsonStringToObject(FSearchDataVersionInfo InVersionInfo, FString InJsonString, TMap<int32, FText>& OutFTextLookupTable)
@@ -2629,7 +3208,7 @@ void FFindInBlueprintSearchManager::GlobalFindResultsClosed(const TSharedRef<SFi
 		{
 			GlobalFindResults.Remove(FindResultsPtr);
 
-			if (CachingObject && SourceCachingWidget == FindResultsPtr)
+			if (IsCacheInProgress() && SourceCachingWidget == FindResultsPtr)
 			{
 				SourceCachingWidget.Reset();
 			}
@@ -2677,6 +3256,12 @@ TSharedRef<SDockTab> FFindInBlueprintSearchManager::SpawnGlobalFindResultsTab(co
 	TSharedRef<SFindInBlueprints> FindResults = SNew(SFindInBlueprints)
 		.ContainingTab(NewTab);
 
+	// If we're in the middle of a caching operation, signal the tab so that it can update the UI.
+	if (IsCacheInProgress())
+	{
+		FindResults->OnCacheStarted(CachingObject->GetCurrentCacheOpType(), CachingObject->GetCurrentCacheOpFlags());
+	}
+
 	GlobalFindResults.Add(FindResults);
 
 	NewTab->SetContent(FindResults);
@@ -2703,7 +3288,9 @@ TSharedPtr<SFindInBlueprints> FFindInBlueprintSearchManager::OpenGlobalFindResul
 		if (!OpenGlobalTabIDs.Contains(GlobalTabId))
 		{
 			TSharedRef<SDockTab> NewTab = FGlobalTabmanager::Get()->InvokeTab(GlobalTabId);
-			return StaticCastSharedRef<SFindInBlueprints>(NewTab->GetContent());
+			TSharedRef<SFindInBlueprints> NewFindTab = StaticCastSharedRef<SFindInBlueprints>(NewTab->GetContent());
+
+			return NewFindTab;
 		}
 	}
 
@@ -2825,6 +3412,12 @@ void FFindInBlueprintSearchManager::Tick(float DeltaTime)
 {
 	CSV_SCOPED_TIMING_STAT(FindInBlueprint, Tick);
 
+	if (bHasFirstSearchOccurred && AssetsToIndexOnFirstSearch.Num() > 0)
+	{
+		PendingAssets = PendingAssets.Union(AssetsToIndexOnFirstSearch);
+		AssetsToIndexOnFirstSearch.Empty();
+	}
+
 	if(IsCacheInProgress())
 	{
 		check(CachingObject);
@@ -2841,8 +3434,10 @@ void FFindInBlueprintSearchManager::Tick(float DeltaTime)
 
 bool FFindInBlueprintSearchManager::IsTickable() const
 {
+	const bool bHasPendingAssets = PendingAssets.Num() > 0 || (bHasFirstSearchOccurred && AssetsToIndexOnFirstSearch.Num() > 0);
+
 	// Tick only if we have an active caching operation or if we have pending assets and an open FiB context
-	return IsCacheInProgress() || (PendingAssets.Num() > 0 && GlobalFindResults.Num() > 0);
+	return IsCacheInProgress() || (bHasPendingAssets && GlobalFindResults.Num() > 0);
 }
 
 TStatId FFindInBlueprintSearchManager::GetStatId() const
diff --git a/Engine/Source/Editor/Kismet/Private/FindInBlueprints.cpp b/Engine/Source/Editor/Kismet/Private/FindInBlueprints.cpp
--- a/Engine/Source/Editor/Kismet/Private/FindInBlueprints.cpp
+++ b/Engine/Source/Editor/Kismet/Private/FindInBlueprints.cpp
@@ -4,6 +4,7 @@
 #include "Framework/Application/MenuStack.h"
 #include "Framework/Application/SlateApplication.h"
 #include "Widgets/Images/SImage.h"
+#include "Widgets/Images/SThrobber.h"
 #include "Widgets/Notifications/SProgressBar.h"
 #include "Framework/MultiBox/MultiBoxBuilder.h"
 #include "Widgets/Text/SMultiLineEditableText.h"
@@ -586,6 +587,12 @@ void SFindInBlueprints::Construct( const FArguments& InArgs, TSharedPtr<FBluepri
 	HostTab = InArgs._ContainingTab;
 	bIsLocked = false;
 
+	bHideProgressBars = false;
+	bShowCacheBarCloseButton = false;
+	bShowCacheBarCancelButton = false;
+	bShowCacheBarUnresponsiveEditorWarningText = false;
+	bKeepCacheBarProgressVisible = false;
+
 	if (HostTab.IsValid())
 	{
 		HostTab.Pin()->SetOnTabClosed(SDockTab::FOnTabClosedCallback::CreateSP(this, &SFindInBlueprints::OnHostTabClosed));
@@ -690,31 +697,35 @@ void SFindInBlueprints::Construct( const FArguments& InArgs, TSharedPtr<FBluepri
 				SNew(STextBlock)
 				.Font( FEditorStyle::GetFontStyle("AssetDiscoveryIndicator.MainStatusFont") )
 				.Text( LOCTEXT("SearchResults", "Searching...") )
-				.Visibility(this, &SFindInBlueprints::GetSearchbarVisiblity)
+				.Visibility(this, &SFindInBlueprints::GetSearchBarWidgetVisiblity, EFiBSearchBarWidget::StatusText)
+			]
+
+			// Throbber
+			+SHorizontalBox::Slot()
+			.AutoWidth()
+			.Padding(2.0f, 3.0f, 0.0f, 0.0f)
+			[
+				SNew(SThrobber)
+				.Visibility(this, &SFindInBlueprints::GetSearchBarWidgetVisiblity, EFiBSearchBarWidget::Throbber)
 			]
 
 			// Progress bar
 			+SHorizontalBox::Slot()
-				.FillWidth(1.0f)
-				.Padding(2.0f, 3.0f, 0.0f, 0.0f)
+			.FillWidth(1.0f)
+			.Padding(2.0f, 3.0f, 0.0f, 0.0f)
 			[
 				SNew(SProgressBar)
-					.Visibility(this, &SFindInBlueprints::GetSearchbarVisiblity).Percent( this, &SFindInBlueprints::GetPercentCompleteSearch )
+				.Visibility(this, &SFindInBlueprints::GetSearchBarWidgetVisiblity, EFiBSearchBarWidget::ProgressBar)
+				.Percent(this, &SFindInBlueprints::GetPercentCompleteSearch)
 			]
 		]
 	];
-
-	// Include the cache progress bar if we're in global find mode and an active cache operation is still in progress.
-	if (!bIsInFindWithinBlueprintMode && IsCacheInProgress())
-	{
-		ConditionallyAddCacheBar();
-	}
 }
 
 void SFindInBlueprints::ConditionallyAddCacheBar()
 {
 	// Do not add when it should not be visible
-	if(GetCachingBarVisibility() == EVisibility::Visible)
+	if(GetCacheBarVisibility() == EVisibility::Visible)
 	{
 		// Do not add a second cache bar
 		if(MainVerticalBox.IsValid() && !CacheBarSlot.IsValid())
@@ -767,9 +778,9 @@ void SFindInBlueprints::ConditionallyAddCacheBar()
 				.AutoHeight()
 				[
 					SAssignNew(CacheBarSlot, SBorder)
-					.Visibility( this, &SFindInBlueprints::GetCachingBarVisibility )
-					.BorderBackgroundColor( this, &SFindInBlueprints::GetCachingBarColor )
-					.BorderImage( FCoreStyle::Get().GetBrush("ErrorReporting.Box") )
+					.Visibility( this, &SFindInBlueprints::GetCacheBarVisibility )
+					.BorderBackgroundColor( this, &SFindInBlueprints::GetCacheBarColor )
+					.BorderImage(this, &SFindInBlueprints::GetCacheBarImage )
 					.Padding( FMargin(3,1) )
 					[
 						SNew(SVerticalBox)
@@ -783,7 +794,7 @@ void SFindInBlueprints::ConditionallyAddCacheBar()
 							.AutoWidth()
 							[
 								SNew(STextBlock)
-								.Text(this, &SFindInBlueprints::GetUnindexedAssetWarningText)
+								.Text(this, &SFindInBlueprints::GetCacheBarStatusText)
 								.ColorAndOpacity( FCoreStyle::Get().GetColor("ErrorReporting.ForegroundColor") )
 							]
 
@@ -796,7 +807,7 @@ void SFindInBlueprints::ConditionallyAddCacheBar()
 									SNew(SButton)
 									.Text(LOCTEXT("IndexAllBlueprints", "Index All"))
 									.OnClicked( this, &SFindInBlueprints::OnCacheAllUnindexedBlueprints )
-									.Visibility( this, &SFindInBlueprints::GetCacheAllUnindexedButtonVisibility )
+									.Visibility(this, &SFindInBlueprints::GetCacheBarWidgetVisibility, EFiBCacheBarWidget::CacheAllUnindexedButton)
 									.ToolTip(IDocumentation::Get()->CreateToolTip(
 									LOCTEXT("IndexAlLBlueprints_Tooltip", "Loads all non-indexed Blueprints and saves them with their search data. This can be a very slow process and the editor may become unresponsive."),
 									NULL,
@@ -813,7 +824,7 @@ void SFindInBlueprints::ConditionallyAddCacheBar()
 									SNew(SButton)
 									.Text(LOCTEXT("ShowFailedPackages", "Show Failed Packages"))
 									.OnClicked(FOnClicked::CreateLambda(OnDisplayCacheFailLambda, TWeakPtr<SWidget>(SharedThis(this)), PackageList))
-									.Visibility( this, &SFindInBlueprints::GetFailedToCacheListVisibility )
+									.Visibility(this, &SFindInBlueprints::GetCacheBarWidgetVisibility, EFiBCacheBarWidget::ShowCacheFailuresButton)
 									.ToolTip(IDocumentation::Get()->CreateToolTip(
 									LOCTEXT("FailedCache_Tooltip", "Displays a list of packages that failed to save."),
 									NULL,
@@ -828,7 +839,7 @@ void SFindInBlueprints::ConditionallyAddCacheBar()
 								[
 									SNew(SProgressBar)
 									.Percent( this, &SFindInBlueprints::GetPercentCompleteCache )
-									.Visibility( this, &SFindInBlueprints::GetCachingProgressBarVisiblity )
+									.Visibility( this, &SFindInBlueprints::GetCacheBarWidgetVisibility, EFiBCacheBarWidget::ProgressBar )
 								]
 
 							// Cancel button
@@ -839,7 +850,7 @@ void SFindInBlueprints::ConditionallyAddCacheBar()
 									SNew(SButton)
 									.Text(LOCTEXT("CancelCacheAll", "Cancel"))
 									.OnClicked( this, &SFindInBlueprints::OnCancelCacheAll )
-									.Visibility( this, &SFindInBlueprints::GetCacheAllCancelButtonVisibility )
+									.Visibility( this, &SFindInBlueprints::GetCacheBarWidgetVisibility, EFiBCacheBarWidget::CancelButton )
 									.ToolTipText( LOCTEXT("CancelCacheAll_Tooltip", "Stops the caching process from where ever it is, can be started back up where it left off when needed.") )
 								]
 
@@ -854,6 +865,7 @@ void SFindInBlueprints::ConditionallyAddCacheBar()
 									.VAlign(VAlign_Center)
 									.OnClicked( this, &SFindInBlueprints::OnRemoveCacheBar )
 									.ForegroundColor( FSlateColor::UseForeground() )
+									.Visibility( this, &SFindInBlueprints::GetCacheBarWidgetVisibility, EFiBCacheBarWidget::CloseButton )
 									[
 										SNew(SImage)
 										.Image( FCoreStyle::Get().GetBrush("EditableComboBox.Delete") )
@@ -871,8 +883,8 @@ void SFindInBlueprints::ConditionallyAddCacheBar()
 								.AutoHeight()
 								[
 									SNew(STextBlock)
-									.Text(this, &SFindInBlueprints::GetCurrentCacheBlueprintName)
-									.Visibility( this, &SFindInBlueprints::GetCachingBlueprintNameVisiblity )
+									.Text(this, &SFindInBlueprints::GetCacheBarCurrentAssetName)
+									.Visibility( this, &SFindInBlueprints::GetCacheBarWidgetVisibility, EFiBCacheBarWidget::CurrentAssetNameText )
 									.ColorAndOpacity( FCoreStyle::Get().GetColor("ErrorReporting.ForegroundColor") )
 								]
 
@@ -882,7 +894,7 @@ void SFindInBlueprints::ConditionallyAddCacheBar()
 									SNew(STextBlock)
 									.Text(LOCTEXT("FiBUnresponsiveEditorWarning", "NOTE: The editor may become unresponsive while these assets are loaded for indexing. This may take some time!"))
 									.TextStyle(&FCoreStyle::Get().GetWidgetStyle<FTextBlockStyle>( "SmallText" ))
-									.Visibility(this, &SFindInBlueprints::GetUnresponsiveEditorWarningVisibility)
+									.Visibility(this, &SFindInBlueprints::GetCacheBarWidgetVisibility, EFiBCacheBarWidget::UnresponsiveEditorWarningText)
 								]
 							]
 					]
@@ -1069,12 +1081,12 @@ void SFindInBlueprints::MakeSearchQuery(FString InSearchString, bool bInIsFindWi
 			}
 
 			const bool bRebuildSearchData = true;
-			const FSearchData* SearchData = FFindInBlueprintSearchManager::Get().QuerySingleBlueprint(Blueprint, bRebuildSearchData);
-			const bool bIsBlueprintIndexed = SearchData && !SearchData->Value.IsEmpty();
+			FSearchData SearchData = FFindInBlueprintSearchManager::Get().QuerySingleBlueprint(Blueprint, bRebuildSearchData);
+			const bool bHasValidSearchData = SearchData.IsValid() && !SearchData.Value.IsEmpty();
 
-			if (bIsBlueprintIndexed)
+			if (bHasValidSearchData)
 			{
-				FImaginaryFiBDataSharedPtr ImaginaryBlueprint(new FImaginaryBlueprint(Blueprint->GetName(), Blueprint->GetPathName(), ParentClass, Interfaces, SearchData->Value, SearchData->VersionInfo));
+				FImaginaryFiBDataSharedPtr ImaginaryBlueprint(new FImaginaryBlueprint(Blueprint->GetName(), Blueprint->GetPathName(), ParentClass, Interfaces, SearchData.Value, SearchData.VersionInfo));
 				TSharedPtr< FFiBSearchInstance > SearchInstance(new FFiBSearchInstance);
 				FSearchResult SearchResult = RootSearchResult = SearchInstance->StartSearchQuery(SearchValue, ImaginaryBlueprint);
 
@@ -1087,7 +1099,7 @@ void SFindInBlueprints::MakeSearchQuery(FString InSearchString, bool bInIsFindWi
 			if(ItemsFound.Num() == 0)
 			{
 				FText NoResultsText;
-				if (bIsBlueprintIndexed)
+				if (bHasValidSearchData)
 				{
 					NoResultsText = LOCTEXT("BlueprintSearchNoResults", "No Results found");
 				}
@@ -1256,9 +1268,43 @@ TOptional<float> SFindInBlueprints::GetPercentCompleteSearch() const
 	return 0.0f;
 }
 
-EVisibility SFindInBlueprints::GetSearchbarVisiblity() const
+EVisibility SFindInBlueprints::GetSearchBarWidgetVisiblity(EFiBSearchBarWidget InSearchBarWidget) const
 {
-	return StreamSearch.IsValid()? EVisibility::Visible : EVisibility::Collapsed;
+	const bool bShowSearchBarWidgets = StreamSearch.IsValid();
+	if (bShowSearchBarWidgets)
+	{
+		EVisibility Result = EVisibility::Visible;
+		const bool bShouldShowProgressBarWidget = !bHideProgressBars;
+
+		switch (InSearchBarWidget)
+		{
+		case EFiBSearchBarWidget::Throbber:
+			// Keep hidden if progress bar is visible.
+			if (bShouldShowProgressBarWidget)
+			{
+				Result = EVisibility::Collapsed;
+			}
+			break;
+
+		case EFiBSearchBarWidget::ProgressBar:
+			// Keep hidden if not allowed to be shown.
+			if (!bShouldShowProgressBarWidget)
+			{
+				Result = EVisibility::Collapsed;
+			}
+			break;
+
+		default:
+			// Always visible.
+			break;
+		}
+
+		return Result;
+	}
+	else
+	{
+		return EVisibility::Collapsed;
+	}
 }
 
 void SFindInBlueprints::CacheAllBlueprints(const FFindInBlueprintCachingOptions& InOptions)
@@ -1303,45 +1349,84 @@ TOptional<float> SFindInBlueprints::GetPercentCompleteCache() const
 	return FFindInBlueprintSearchManager::Get().GetCacheProgress();
 }
 
-EVisibility SFindInBlueprints::GetCachingProgressBarVisiblity() const
+EVisibility SFindInBlueprints::GetCacheBarVisibility() const
 {
-	return IsCacheInProgress()? EVisibility::Visible : EVisibility::Hidden;
+	const bool bIsPIESimulating = (GEditor->bIsSimulatingInEditor || GEditor->PlayWorld);
+	FFindInBlueprintSearchManager& FindInBlueprintManager = FFindInBlueprintSearchManager::Get();
+	return (bKeepCacheBarProgressVisible || FindInBlueprintManager.GetNumberUncachedAssets() > 0 || (!bIsPIESimulating && (FindInBlueprintManager.GetNumberUnindexedAssets() > 0 || FindInBlueprintManager.GetFailedToCacheCount()))) ? EVisibility::Visible : EVisibility::Collapsed;
 }
 
-EVisibility SFindInBlueprints::GetCacheAllUnindexedButtonVisibility() const
+EVisibility SFindInBlueprints::GetCacheBarWidgetVisibility(EFiBCacheBarWidget InCacheBarWidget) const
 {
-	return IsCacheInProgress()? EVisibility::Collapsed : EVisibility::Visible;
-}
+	EVisibility Result = EVisibility::Visible;
 
-EVisibility SFindInBlueprints::GetCacheAllCancelButtonVisibility() const
-{
-	// The cancel button is visible only if we're doing an unindexed caching operation (slow) and this is the initiating widget
-	const bool bIsUnindexedCacheInProgress = FFindInBlueprintSearchManager::Get().IsUnindexedCacheInProgress();
-	TWeakPtr<SFindInBlueprints> SourceWidgetPtr = FFindInBlueprintSearchManager::Get().GetSourceCachingWidget();
-	return bIsUnindexedCacheInProgress && SourceWidgetPtr.Pin() == SharedThis(this) ? EVisibility::Visible : EVisibility::Collapsed;
-}
+	const bool bShowCacheBarProgressWidgets = IsCacheInProgress() || bKeepCacheBarProgressVisible;
+	const bool bHideCacheBarProgressWidgets = !bShowCacheBarProgressWidgets;
 
-EVisibility SFindInBlueprints::GetCachingBarVisibility() const
-{
-	const bool bIsPIESimulating = (GEditor->bIsSimulatingInEditor || GEditor->PlayWorld);
-	FFindInBlueprintSearchManager& FindInBlueprintManager = FFindInBlueprintSearchManager::Get();
-	return (FindInBlueprintManager.GetNumberUncachedAssets() > 0 || (!bIsPIESimulating && (FindInBlueprintManager.GetNumberUnindexedAssets() > 0 || FindInBlueprintManager.GetFailedToCacheCount())))? EVisibility::Visible : EVisibility::Collapsed;
-}
+	switch (InCacheBarWidget)
+	{
+	case EFiBCacheBarWidget::ProgressBar:
+		// Keep hidden when not caching or when progress bars are explicitly not being shown.
+		if (bHideCacheBarProgressWidgets || bHideProgressBars)
+		{
+			Result = EVisibility::Hidden;
+		}
+		break;
 
-EVisibility SFindInBlueprints::GetCachingBlueprintNameVisiblity() const
-{
-	return IsCacheInProgress()? EVisibility::Visible : EVisibility::Collapsed;
-}
+	case EFiBCacheBarWidget::CloseButton:
+		// Keep hidden while caching if explicitly not being shown.
+		if (bShowCacheBarProgressWidgets && !bShowCacheBarCloseButton)
+		{
+			Result = EVisibility::Collapsed;
+		}
+		break;
 
-EVisibility SFindInBlueprints::GetFailedToCacheListVisibility() const
-{
-	return FFindInBlueprintSearchManager::Get().GetFailedToCacheCount()? EVisibility::Visible : EVisibility::Collapsed;
-}
+	case EFiBCacheBarWidget::CancelButton:
+		// Keep hidden when not caching or when explicitly not being shown.
+		if (bHideCacheBarProgressWidgets || !bShowCacheBarCancelButton)
+		{
+			Result = EVisibility::Collapsed;
+		}
+		break;
 
-EVisibility SFindInBlueprints::GetUnresponsiveEditorWarningVisibility() const
-{
-	// Show the "editor may be unresponsive" note if we're doing an unindexed caching operation, which can block the UI due to file I/O
-	return !IsCacheInProgress() || FFindInBlueprintSearchManager::Get().IsUnindexedCacheInProgress() ? EVisibility::Visible : EVisibility::Collapsed;
+	case EFiBCacheBarWidget::CacheAllUnindexedButton:
+		// Always keep hidden while caching.
+		if (bShowCacheBarProgressWidgets)
+		{
+			Result = EVisibility::Collapsed;
+		}
+		break;
+
+	case EFiBCacheBarWidget::CurrentAssetNameText:
+		// Keep hidden when not caching.
+		if (bHideCacheBarProgressWidgets)
+		{
+			Result = EVisibility::Collapsed;
+		}
+		break;
+
+	case EFiBCacheBarWidget::UnresponsiveEditorWarningText:
+		// Keep hidden while caching if explicitly not being shown.
+		if (bShowCacheBarProgressWidgets && !bShowCacheBarUnresponsiveEditorWarningText)
+		{
+			Result = EVisibility::Collapsed;
+		}
+		break;
+
+	case EFiBCacheBarWidget::ShowCacheFailuresButton:
+		// Always keep hidden while caching. Also keep hidden if there are no assets that failed to be cached. 
+		if (bShowCacheBarProgressWidgets || FFindInBlueprintSearchManager::Get().GetFailedToCacheCount() == 0)
+		{
+			Result = EVisibility::Collapsed;
+		}
+		break;
+
+	default:
+		// Always visible.
+		break;
+	}
+
+	return Result;
 }
 
 bool SFindInBlueprints::IsCacheInProgress() const
@@ -1349,30 +1434,56 @@ bool SFindInBlueprints::IsCacheInProgress() const
 	return FFindInBlueprintSearchManager::Get().IsCacheInProgress();
 }
 
-FSlateColor SFindInBlueprints::GetCachingBarColor() const
+FSlateColor SFindInBlueprints::GetCacheBarColor() const
 {
 	// The caching bar's default color is a darkish red
 	FSlateColor ReturnColor = FSlateColor(FLinearColor(0.4f, 0.0f, 0.0f));
-	if(IsCacheInProgress())
+	if (IsCacheInProgress() || bKeepCacheBarProgressVisible)
 	{
-		// It turns yellow when in progress
-		ReturnColor = FSlateColor(FLinearColor(0.4f, 0.4f, 0.0f));
+		if (FFindInBlueprintSearchManager::Get().IsUnindexedCacheInProgress())
+		{
+			// It turns yellow when an unindexed cache is in progress
+			ReturnColor = FSlateColor(FLinearColor(0.4f, 0.4f, 0.0f));
+		}
+		else
+		{
+			// Use the background image color for a non-unindexed cache
+			ReturnColor = FSlateColor(FLinearColor::White);
+		}
 	}
 	return ReturnColor;
 }
 
-FText SFindInBlueprints::GetUnindexedAssetWarningText() const
+const FSlateBrush* SFindInBlueprints::GetCacheBarImage() const
+{
+	const FSlateBrush* ReturnBrush = FCoreStyle::Get().GetBrush("ErrorReporting.Box");
+	if ((IsCacheInProgress() || bKeepCacheBarProgressVisible) && !FFindInBlueprintSearchManager::Get().IsUnindexedCacheInProgress())
+	{
+		// Allow the content area to show through for a non-unindexed operation.
+		ReturnBrush = FEditorStyle::GetBrush("NoBorder");
+	}
+	return ReturnBrush;
+}
+
+FText SFindInBlueprints::GetCacheBarStatusText() const
 {
 	FFindInBlueprintSearchManager& FindInBlueprintManager = FFindInBlueprintSearchManager::Get();
 
 	FFormatNamedArguments Args;
 	FText ReturnDisplayText;
-	if(IsCacheInProgress())
+	if (IsCacheInProgress() || bKeepCacheBarProgressVisible)
 	{
-		Args.Add(TEXT("CurrentIndex"), FindInBlueprintManager.GetCurrentCacheIndex());
-		Args.Add(TEXT("Count"), FindInBlueprintManager.GetNumberUncachedAssets());
+		if (bHideProgressBars)
+		{
+			ReturnDisplayText = LOCTEXT("CachingBlueprintsWithUnknownEndpoint", "Indexing Blueprints...");
+		}
+		else
+		{
+			Args.Add(TEXT("CurrentIndex"), FindInBlueprintManager.GetCurrentCacheIndex());
+			Args.Add(TEXT("Count"), FindInBlueprintManager.GetNumberUncachedAssets());
 
-		ReturnDisplayText = FText::Format(LOCTEXT("CachingBlueprints", "Indexing Blueprints... {CurrentIndex}/{Count}"), Args);
+			ReturnDisplayText = FText::Format(LOCTEXT("CachingBlueprints", "Indexing Blueprints... {CurrentIndex}/{Count}"), Args);
+		}
 	}
 	else
 	{
@@ -1396,25 +1507,50 @@ FText SFindInBlueprints::GetUnindexedAssetWarningText() const
 	return ReturnDisplayText;
 }
 
-FText SFindInBlueprints::GetCurrentCacheBlueprintName() const
+FText SFindInBlueprints::GetCacheBarCurrentAssetName() const
 {
-	return FText::FromName(FFindInBlueprintSearchManager::Get().GetCurrentCacheBlueprintName());
+	if (IsCacheInProgress())
+	{
+		LastCachedAssetName = FFindInBlueprintSearchManager::Get().GetCurrentCacheBlueprintName();
+	}
+
+	return FText::FromName(LastCachedAssetName);
 }
 
 void SFindInBlueprints::OnCacheStarted(EFiBCacheOpType InOpType, EFiBCacheOpFlags InOpFlags)
 {
-	const bool bIsCacheBarAdded = CacheBarSlot.IsValid();
 	const bool bShowProgress = EnumHasAnyFlags(InOpFlags, EFiBCacheOpFlags::ShowProgress);
-
-	// Ensure that the cache bar is visible to show progress
-	if (!bIsCacheBarAdded && bShowProgress)
+	if (bShowProgress)
 	{
-		ConditionallyAddCacheBar();
+		// Whether to keep both the cache and search bar progress indicators hidden.
+		bHideProgressBars = EnumHasAnyFlags(InOpFlags, EFiBCacheOpFlags::HideProgressBars);
+
+		// Whether to show the cache bar close button and allow users to dismiss the progress display.
+		bShowCacheBarCloseButton = EnumHasAnyFlags(InOpFlags, EFiBCacheOpFlags::AllowUserCloseProgress);
+
+		// Whether to show the cache bar cancel button allowing users to cancel the operation.
+		bShowCacheBarCancelButton = EnumHasAnyFlags(InOpFlags, EFiBCacheOpFlags::AllowUserCancel);
+
+		// Whether to show the unresponsive editor warning text in the cache bar status area.
+		bShowCacheBarUnresponsiveEditorWarningText = (InOpType == EFiBCacheOpType::CacheUnindexedAssets);
+
+		// Ensure that the cache bar is visible to show progress
+		const bool bIsCacheBarAdded = CacheBarSlot.IsValid();
+		if (!bIsCacheBarAdded)
+		{
+			ConditionallyAddCacheBar();
+		}
 	}
 }
 
-void SFindInBlueprints::OnCacheComplete(EFiBCacheOpType InOpType)
+void SFindInBlueprints::OnCacheComplete(EFiBCacheOpType InOpType, EFiBCacheOpFlags InOpFlags)
 {
+	// Indicate whether to keep the search bar progress indicator hidden.
+	bHideProgressBars = EnumHasAnyFlags(InOpFlags, EFiBCacheOpFlags::HideProgressBars);
+
+	// Indicate whether to keep cache bar progress visible. Used to seamlessly transition to the next operation.
+	bKeepCacheBarProgressVisible = EnumHasAnyFlags(InOpFlags, EFiBCacheOpFlags::KeepProgressVisibleOnCompletion);
+
 	TWeakPtr<SFindInBlueprints> SourceCachingWidgetPtr = FFindInBlueprintSearchManager::Get().GetSourceCachingWidget();
 	if (InOpType == EFiBCacheOpType::CacheUnindexedAssets
 		&& SourceCachingWidgetPtr.IsValid() && SourceCachingWidgetPtr.Pin() == SharedThis(this))
@@ -1422,9 +1558,9 @@ void SFindInBlueprints::OnCacheComplete(EFiBCacheOpType InOpType)
 		// Resubmit the last search, which will also remove the bar if needed
 		OnSearchTextCommitted(SearchTextField->GetText(), ETextCommit::OnEnter);
 	}
-	else if (CacheBarSlot.IsValid())
+	else if (CacheBarSlot.IsValid() && !bKeepCacheBarProgressVisible)
 	{
-		// Remove the cache bar
+		// Remove the cache bar, unless this is not the true end of the operation
 		OnRemoveCacheBar();
 	}
 }
diff --git a/Engine/Source/Editor/Kismet/Private/ImaginaryBlueprintData.cpp b/Engine/Source/Editor/Kismet/Private/ImaginaryBlueprintData.cpp
--- a/Engine/Source/Editor/Kismet/Private/ImaginaryBlueprintData.cpp
+++ b/Engine/Source/Editor/Kismet/Private/ImaginaryBlueprintData.cpp
@@ -68,7 +68,15 @@ FImaginaryFiBData::FImaginaryFiBData(FImaginaryFiBDataWeakPtr InOuter, TSharedPt
 	: UnparsedJsonObject(InUnparsedJsonObject)
 	, LookupTablePtr(InLookupTablePtr)
 	, Outer(InOuter)
+	, bHasParsedJsonObject(false)
+	, bRequiresInterlockedParsing(false)
 {
+	// Backwards-compatibility; inherit the flag that only allows one thread at a time into the JSON parsing logic.
+	const FImaginaryFiBDataSharedPtr OuterPtr = Outer.Pin();
+	if (OuterPtr.IsValid())
+	{
+		bRequiresInterlockedParsing = OuterPtr->bRequiresInterlockedParsing;
+	}
 }
 
 FSearchResult FImaginaryFiBData::CreateSearchResult(FSearchResult InParent) const
@@ -157,8 +165,6 @@ bool FImaginaryFiBData::CanCallFilter(ESearchQueryFilter InSearchQueryFilter) co
 
 void FImaginaryFiBData::ParseAllChildData_Internal(ESearchableValueStatus InSearchabilityOverride/* = ESearchableValueStatus::Searchable*/)
 {
-	FScopeLock ScopeLock(&FImaginaryFiBData::ParseChildDataCriticalSection);
-
 	if (UnparsedJsonObject.IsValid())
 	{
 		if (InSearchabilityOverride & ESearchableValueStatus::Searchable)
@@ -236,8 +242,21 @@ void FImaginaryFiBData::ParseAllChildData(ESearchableValueStatus InSearchability
 	CSV_SCOPED_TIMING_STAT(FindInBlueprint, ParseAllChildData);
 	CSV_CUSTOM_STAT(FindInBlueprint, ParseAllChildDataIterations, 1, ECsvCustomStatOp::Accumulate);
 
-	FScopeLock ScopeLock(&FImaginaryFiBData::ParseChildDataCriticalSection);
-	ParseAllChildData_Internal(InSearchabilityOverride);
+	if (bRequiresInterlockedParsing)
+	{
+		ParseChildDataCriticalSection.Lock();
+	}
+
+	if (!bHasParsedJsonObject)
+	{
+		ParseAllChildData_Internal(InSearchabilityOverride);
+		bHasParsedJsonObject = true;
+	}
+
+	if (bRequiresInterlockedParsing)
+	{
+		ParseChildDataCriticalSection.Unlock();
+	}
 }
 
 void FImaginaryFiBData::ParseJsonValue(FText InKey, FText InDisplayKey, TSharedPtr< FJsonValue > InJsonValue, TArray<FSearchableValueInfo>& OutParsedValues, bool bIsInArray/*=false*/, ESearchableValueStatus InSearchabilityOverride/* = ESearchableValueStatus::Searchable*/)
diff --git a/Engine/Source/Editor/Kismet/Public/FindInBlueprintManager.h b/Engine/Source/Editor/Kismet/Public/FindInBlueprintManager.h
--- a/Engine/Source/Editor/Kismet/Public/FindInBlueprintManager.h
+++ b/Engine/Source/Editor/Kismet/Public/FindInBlueprintManager.h
@@ -8,10 +8,12 @@
 #include "Types/WidgetActiveTimerDelegate.h"
 #include "Dom/JsonObject.h"
 #include "HAL/Runnable.h"
+#include "HAL/RunnableThread.h"
 #include "SlateFwd.h"
 #include "Input/Reply.h"
 #include "TickableEditorObject.h"
 #include "ProfilingDebugging/CsvProfiler.h"
+#include "Containers/Queue.h"
 
 DECLARE_LOG_CATEGORY_EXTERN(LogFindInBlueprint, Warning, All);
 
@@ -135,6 +137,18 @@ struct FSearchDataVersionInfo
 	static FSearchDataVersionInfo Current;
 };
 
+/** State flags for search database entries */
+enum ESearchDataStateFlags
+{
+	None = 0,
+	/** Set when this search database entry has been fully indexed, which is completed asynchronously */
+	IsIndexed = 1 << 0,
+	/** Cached to determine if the Blueprint is seen as no longer valid, allows it to be cleared out next save to disk */
+	WasRemoved = 1 << 1,
+};
+
+ENUM_CLASS_FLAGS(ESearchDataStateFlags);
+
 /** Tracks data relevant to a Blueprint for searches */
 struct FSearchData
 {
@@ -153,18 +167,18 @@ struct FSearchData
 	/** Interfaces implemented by the Blueprint */
 	TArray<FString> Interfaces;
 
-	/** Cached to determine if the Blueprint is seen as no longer valid, allows it to be cleared out next save to disk */
-	bool bMarkedForDeletion;
-
 	/** Cached ImaginaryBlueprint data for the searchable content, prevents having to re-parse every search */
 	FImaginaryFiBDataSharedPtr ImaginaryBlueprint;
 
 	/** Data versioning */
 	FSearchDataVersionInfo VersionInfo;
 
+	/** State flags (see enum) */
+	ESearchDataStateFlags StateFlags;
+
 	FSearchData()
 		: Blueprint(nullptr)
-		, bMarkedForDeletion(false)
+		, StateFlags(ESearchDataStateFlags::None)
 	{
 	}
 
@@ -172,6 +186,16 @@ struct FSearchData
 	{
 		return AssetPath != NAME_None;
 	}
+
+	bool IsIndexingCompleted() const
+	{
+		return EnumHasAllFlags(StateFlags, ESearchDataStateFlags::IsIndexed);
+	}
+
+	bool IsMarkedForDeletion() const
+	{
+		return EnumHasAllFlags(StateFlags, ESearchDataStateFlags::WasRemoved);
+	}
 };
 
 /** Filters are used by functions for searching to decide whether items can call certain functions or match the requirements of a function */
@@ -229,6 +253,20 @@ enum class EFiBCacheOpFlags
 	AllowUserCancel = 1 << 2,
 	/** Set if the user wants to check out and save (applies to unindexed caching only) */
 	CheckOutAndSave = 1 << 3,
+	/** Whether to hide progress bar widgets */
+	HideProgressBars = 1 << 4,
+	/** Whether to allow users to hide/close progress */
+	AllowUserCloseProgress = 1 << 5,
+	/** Set if we are caching assets from the discovery stage */
+	IsCachingDiscoveredAssets = 1 << 6,
+	/** Whether to keep progress visible on completion */
+	KeepProgressVisibleOnCompletion = 1 << 7,
+	/** Index deferred assets on the main thread only (used for debugging) */
+	ExecuteOnMainThread = 1 << 8,
+	/** Don't index multiple assets in parallel (used to assist with profiling) */
+	ExecuteOnSingleThread = 1 << 9,
+	/** Only execute the gather phase (used to help minimize memory usage on editor/tab open) */
+	ExecuteGatherPhaseOnly = 1 << 10,
 };
 
 ENUM_CLASS_FLAGS(EFiBCacheOpFlags);
@@ -239,6 +277,9 @@ struct FFindInBlueprintCachingOptions
 	/** Type of caching operation */
 	EFiBCacheOpType OpType = EFiBCacheOpType::CachePendingAssets;
 
+	/** Initial set of control flags */
+	EFiBCacheOpFlags OpFlags = EFiBCacheOpFlags::None;
+
 	/** Callback for when caching is finished */
 	FSimpleDelegate OnFinished;
 
@@ -246,6 +287,18 @@ struct FFindInBlueprintCachingOptions
 	EFiBVersion MinimiumVersionRequirement = EFiBVersion::FIB_VER_LATEST;
 };
 
+/** Options for FFindInBlueprintSearchManager::AddOrUpdateBlueprintSearchMetadata() */
+enum class EAddOrUpdateBlueprintSearchMetadataFlags
+{
+	None = 0,
+	/** Forces the Blueprint to be recache'd, regardless of what data it believes exists */
+	ForceRecache = 1 << 0,
+	/** Whether to allow a new entry to be added to the database if a match is not found */
+	AllowNewEntry = 1 << 1,
+};
+
+ENUM_CLASS_FLAGS(EAddOrUpdateBlueprintSearchMetadataFlags);
+
 ////////////////////////////////////
 // FFindInBlueprintsResult
 
@@ -374,7 +427,7 @@ public:
 
 public:
 	/** Thread to run the cleanup FRunnable on */
-	FRunnableThread* Thread;
+	TUniquePtr<FRunnableThread> Thread;
 
 	/** A list of items found, cleared whenever the main thread pulls them to display to screen */
 	TArray<TSharedPtr<class FFindInBlueprintsResult>> ItemsFound;
@@ -424,6 +477,14 @@ public:
 	virtual TStatId GetStatId() const override;
 	//~ End FTickableObject Interface
 
+	/**
+	 * Applies the given search data to a matching entry in the database. Optionally adds a new entry if a match is not found.
+	 *
+	 * @param InSearchData		Search data to be applied
+	 * @param bAllowNewEntry	Whether to allow a new entry to be added if a match is not found
+	 */
+	void ApplySearchDataToDatabase(FSearchData InSearchData, bool bAllowNewEntry = false);
+
 	/**
 	 * Given an asset path, locate and return a copy of its matching search data in the index cache.
 	 *
@@ -436,9 +497,9 @@ public:
 	 * Gathers the Blueprint's search metadata and adds or updates it in the cache
 	 *
 	 * @param InBlueprint		Blueprint to cache the searchable data for
-	 * @param bInForceReCache	Forces the Blueprint to be recache'd, regardless of what data it believes exists
+	 * @param InOptions			Options to assist with controlling this function's behavior
 	 */
-	void AddOrUpdateBlueprintSearchMetadata(UBlueprint* InBlueprint, bool bInForceReCache = false);
+	void AddOrUpdateBlueprintSearchMetadata(UBlueprint* InBlueprint, EAddOrUpdateBlueprintSearchMetadataFlags InFlags = EAddOrUpdateBlueprintSearchMetadataFlags::None);
 
 	/**
 	 * Starts a search query, the FiB manager handles where the thread is at in the search query at all times so that post-save of the cache to disk it can correct the index
@@ -473,16 +534,13 @@ public:
 	float GetPercentComplete(const class FStreamSearch* InSearchOriginator) const;
 
 	/**
-	 * Query for a single, specific Blueprint's search block.
+	 * Query for a single, specific Blueprint's search data.
 	 *
 	 * @param InBlueprint				The Blueprint to search for
 	 * @param bInRebuildSearchData		When TRUE the search data will be freshly collected
-	 * @return							The search block
+	 * @return							The search data, or invalid (empty) search data if not found
 	 */
-	const FSearchData* QuerySingleBlueprint(UBlueprint* InBlueprint, bool bInRebuildSearchData);
-
-	/** Returns the number of assets that are waiting to be re-indexed */
-	int32 GetNumberPendingAssets() const;
+	FSearchData QuerySingleBlueprint(UBlueprint* InBlueprint, bool bInRebuildSearchData);
 
 	/** Returns the number of unindexed Blueprints, either due to not having been indexed before, or AR data being out-of-date */
 	int32 GetNumberUnindexedAssets() const;
@@ -528,14 +586,14 @@ public:
 	bool HasCachingFailed() const { return FailedToCachePaths.Num() > 0; };
 
 	/** Callback to note that Blueprint caching is started */
-	void StartedCachingBlueprints(EFiBCacheOpFlags InCacheOpFlags);
+	void StartedCachingBlueprints(EFiBCacheOpType InCacheOpType, EFiBCacheOpFlags InCacheOpFlags);
 
 	/**
 	 * Callback to note that Blueprint caching is complete
 	 *
 	 * @param InNumberCached		The number of Blueprints cached, to be chopped off the existing array so the rest (if any) can be finished later
 	 */
-	void FinishedCachingBlueprints(int32 InNumberCached, TSet<FName>& InFailedToCacheList);
+	void FinishedCachingBlueprints(EFiBCacheOpType InCacheOpType, EFiBCacheOpFlags InCacheOpFlags, int32 InNumberCached, TSet<FName>& InFailedToCacheList);
 
 	/** Returns TRUE if Blueprints are being cached. */
 	bool IsCacheInProgress() const;
@@ -543,6 +601,9 @@ public:
 	/** Returns TRUE if unindexed Blueprints are being cached (since this can block the UI) */
 	bool IsUnindexedCacheInProgress() const;
 
+	/** Returns TRUE if we're still inside the initial asset discovery and registration stage */
+	bool IsAssetDiscoveryInProgress() const;
+
 	/** Returns a weak reference to the widget that initiated the current caching operation */
 	TWeakPtr<SFindInBlueprints> GetSourceCachingWidget() const { return SourceCachingWidget; }
 
@@ -565,6 +626,9 @@ public:
 	/** Close any orphaned global find results tabs for a particular tab manager */
 	void CloseOrphanedGlobalFindResultsTabs(TSharedPtr<class FTabManager> TabManager);
 
+	/** Returns true if a global find results tab is currently open */
+	bool IsGlobalFindResultsOpen() const { return GlobalFindResults.Num() > 0; }
+
 	void GlobalFindResultsClosed(const TSharedRef<SFindInBlueprints>& FindResults);
 
 	/** Dumps the full index cache to the given stream (for debugging purposes) */
@@ -614,6 +678,9 @@ private:
 	/** Callback hook from the Hot Reload manager that indicates that a module has been hot-reloaded */
 	void OnHotReload(bool bWasTriggeredAutomatically);
 
+	/** Returns a copy of the search data that's cached at the given index. Will return invalid (empty) search data if the index is out of range */
+	FSearchData GetSearchDataForIndex(int32 CacheIndex);
+
 	/** Cleans the cache of any excess data from Blueprints that have been moved, renamed, or deleted. Occurs during post-garbage collection */
 	void CleanCache();
 
@@ -644,6 +711,29 @@ private:
 	TSharedPtr<SFindInBlueprints> OpenGlobalFindResultsTab();
 
 protected:
+	/** Contains info about an active search query */
+	struct FActiveSearchQuery
+	{
+		/** Current search array index */
+		TAtomic<int32> NextIndex;
+		/** Current count of assets searched */
+		TAtomic<int32> SearchCount;
+		/** Asset paths for which searching was deferred due to being indexed */
+		TQueue<FName> DeferredAssetPaths;
+
+		FActiveSearchQuery()
+			:NextIndex(0)
+			, SearchCount(0)
+		{
+		}
+	};
+
+	typedef TSharedPtr<FActiveSearchQuery, ESPMode::ThreadSafe> FActiveSearchQueryPtr;
+
+	/** Thread-safe access to the active search query that's mapped to the given stream search */
+	FActiveSearchQueryPtr FindSearchQuery(const class FStreamSearch* InSearchOriginator) const;
+
+private:
 	/** Maps the Blueprint paths to their index in the SearchArray */
 	TMap<FName, int32> SearchMap;
 
@@ -654,7 +744,7 @@ protected:
 	FThreadSafeCounter ActiveSearchCounter;
 
 	/** A mapping of active search queries and where they are currently at in the search data */
-	TMap< const class FStreamSearch*, int32 > ActiveSearchQueries;
+	TMap<const class FStreamSearch*, FActiveSearchQueryPtr> ActiveSearchQueries;
 
 	/** Critical section to safely add, remove, and find data in ActiveSearchQueries */
 	mutable FCriticalSection SafeQueryModifyCriticalSection;
@@ -671,7 +761,7 @@ protected:
 	/** FindInBlueprints widget that started the cache process */
 	TWeakPtr<SFindInBlueprints> SourceCachingWidget;
 
-	/** Asset paths that are either loaded or modified and require re-indexing before caching */
+	/** Asset paths that were discovered, loaded or modified and now require indexing (or re-indexing) */
 	TSet<FName> PendingAssets;
 
 	/** Asset paths that have not been cached for searching due to lack of FiB data, this means that they are either older Blueprints, or the DDC cannot find the data */
@@ -680,8 +770,11 @@ protected:
 	/** List of paths for Blueprints that failed to cache */
 	TSet<FName> FailedToCachePaths;
 
+	/** List of paths that require a full index pass during the first global search */
+	TSet<FName> AssetsToIndexOnFirstSearch;
+
 	/** Tickable object that does the caching of uncached Blueprints at a rate of once per tick */
-	class FCacheAllBlueprintsTickableObject* CachingObject;
+	TUniquePtr<class FCacheAllBlueprintsTickableObject> CachingObject;
 
 	/** Stores the type of caching operation that's currently in progress */
 	EFiBCacheOpType CurrentCacheOpType;
@@ -698,12 +791,25 @@ protected:
 	/** Global Find Results workspace menu item */
 	TSharedPtr<class FWorkspaceItem> GlobalFindResultsMenuItem;
 
+	/** Size of a single async indexing task batch; zero or negative means that a caching operation will consist of a single batch */
+	int32 AsyncTaskBatchSize;
+
+	/** TRUE when the the FiB manager wants to pause all searches, helps manage the pausing procedure */
+	TAtomic<bool> bIsPausing;
+
+	/** Currently used to delay the full indexing pass until the first search in order to control memory usage */
+	TAtomic<bool> bHasFirstSearchOccurred;
+
 	/** Tells if gathering data is currently allowed */
 	bool bEnableGatheringData;
 
-	/** Tells if deferred indexing is configured as disabled */
+	/** If true, search metadata will be regenerated for loaded assets on the main thread immediately at discovery or load time. By default, this work is deferred to avoid hitching. */
 	bool bDisableDeferredIndexing;
 
+	/** If true, and if deferred indexing is not disabled, search metadata will be regenerated for loaded assets on the main thread at a rate of one asset per tick. Additional indexing
+		work will be deferred to the search thread. By default, all indexing work is deferred to a background thread and both loaded and unloaded assets are fully indexed asynchronously. */
+	bool bDisableThreadedIndexing;
+
 	/** Whether CSV profiling has been enabled (default=false) */
 	bool bEnableCSVStatsProfiling;
 
@@ -712,9 +818,6 @@ protected:
 
 	/** Disable the use of search result templates. Setting this to TRUE will slightly decrease overall memory usage, but will also increase global search times */
 	bool bDisableSearchResultTemplates;
-
-	/** TRUE when the the FiB manager wants to pause all searches, helps manage the pausing procedure */
-	volatile bool bIsPausing;
 };
 
 struct KISMET_API FDisableGatheringDataOnScope
diff --git a/Engine/Source/Editor/Kismet/Public/FindInBlueprints.h b/Engine/Source/Editor/Kismet/Public/FindInBlueprints.h
--- a/Engine/Source/Editor/Kismet/Public/FindInBlueprints.h
+++ b/Engine/Source/Editor/Kismet/Public/FindInBlueprints.h
@@ -206,6 +206,26 @@ private:
 	EGraphType GraphType;
 };
 
+// Cache bar widgets.
+enum class EFiBCacheBarWidget
+{
+	ProgressBar,
+	CloseButton,
+	CancelButton,
+	CacheAllUnindexedButton,
+	CurrentAssetNameText,
+	UnresponsiveEditorWarningText,
+	ShowCacheFailuresButton
+};
+
+// Search bar widgets.
+enum class EFiBSearchBarWidget
+{
+	StatusText,
+	Throbber,
+	ProgressBar,
+};
+
 /*Widget for searching for (functions/events) across all blueprints or just a single blueprint */
 class SFindInBlueprints: public SCompoundWidget
 {
@@ -240,7 +260,7 @@ public:
 	void OnCacheStarted(EFiBCacheOpType InOpType, EFiBCacheOpFlags InOpFlags);
 	
 	/** Called when caching Blueprints is complete */
-	void OnCacheComplete(EFiBCacheOpType InOpType);
+	void OnCacheComplete(EFiBCacheOpType InOpType, EFiBCacheOpFlags InOpFlags);
 
 	/**
 	 * Asynchronously caches all Blueprints below a specified version.
@@ -298,8 +318,8 @@ private:
 	/** Returns the percent complete on the search for the progress bar */
 	TOptional<float> GetPercentCompleteSearch() const;
 
-	/** Returns the progress bar visiblity */
-	EVisibility GetSearchbarVisiblity() const;
+	/** Returns the search bar visiblity for the given widget */
+	EVisibility GetSearchBarWidgetVisiblity(EFiBSearchBarWidget InSearchBarWidget) const;
 
 	/** Adds the "cache" bar at the bottom of the Find-in-Blueprints widget, to notify the user that the search is incomplete */
 	void ConditionallyAddCacheBar();
@@ -308,10 +328,10 @@ private:
 	FReply OnRemoveCacheBar();
 
 	/** Callback to return the cache bar's display text, informing the user of the situation */
-	FText GetUnindexedAssetWarningText() const;
+	FText GetCacheBarStatusText() const;
 
-	/** Callback to return the cache bar's current indexing Blueprint name */
-	FText GetCurrentCacheBlueprintName() const;
+	/** Callback to return the current asset name during a cache operation */
+	FText GetCacheBarCurrentAssetName() const;
 
 	/** Callback to cache all unindexed Blueprints */
 	FReply OnCacheAllUnindexedBlueprints();
@@ -328,32 +348,20 @@ private:
 	/** Gets the percent complete of the caching process */
 	TOptional<float> GetPercentCompleteCache() const;
 
-	/** Returns the visibility of the caching progress bar, visible when in progress, hidden when not */
-	EVisibility GetCachingProgressBarVisiblity() const;
-
-	/** Returns the visibility of the "Cache All" button, visible when not caching, collapsed when caching is in progress */
-	EVisibility GetCacheAllUnindexedButtonVisibility() const;
-
-	/** Returns the visibility of the "Cancel" button, visible when caching is in progress, collapsed when not caching */
-	EVisibility GetCacheAllCancelButtonVisibility() const;
-
 	/** Returns the caching bar's visibility, it goes invisible when there is nothing to be cached. The next search will remove this bar or make it visible again */
-	EVisibility GetCachingBarVisibility() const;
+	EVisibility GetCacheBarVisibility() const;
 
-	/** Returns the visibility of the caching Blueprint name, visible when in progress, collapsed when not */
-	EVisibility GetCachingBlueprintNameVisiblity() const;
-
-	/** Returns the visibility of the popup button that displays the list of Blueprints that failed to cache */
-	EVisibility GetFailedToCacheListVisibility() const;
-
-	/** Returns the visibility of the unresponsive editor warning note text in the caching progress bar */
-	EVisibility GetUnresponsiveEditorWarningVisibility() const;
+	/** Returns the cache bar visibility for the given widget */
+	EVisibility GetCacheBarWidgetVisibility(EFiBCacheBarWidget InCacheBarWidget) const;
 
 	/** Returns TRUE if Blueprint caching is in progress */
 	bool IsCacheInProgress() const;
 
-	/** Returns the color of the caching bar */
-	FSlateColor GetCachingBarColor() const;
+	/** Returns the color of the cache bar */
+	FSlateColor GetCacheBarColor() const;
+
+	/** Returns the BG image used for the cache bar */
+	const FSlateBrush* GetCacheBarImage() const;
 
 	/** Callback to build the context menu when right clicking in the tree */
 	TSharedPtr<SWidget> OnContextMenuOpening();
@@ -401,9 +409,6 @@ private:
 	/* The string to search for */
 	FString	SearchValue;
 
-	/** Should we search within the current blueprint only (rather than all blueprints) */
-	bool bIsInFindWithinBlueprintMode;
-
 	/** Thread object that searches through Blueprint data on a separate thread */
 	TSharedPtr< class FStreamSearch> StreamSearch;
 
@@ -428,6 +433,27 @@ private:
 	/** Tab hosting this widget. May be invalid. */
 	TWeakPtr<SDockTab> HostTab;
 
+	/** Last cached asset name (used during continuous cache operations). */
+	mutable FName LastCachedAssetName;
+
+	/** Should we search within the current blueprint only (rather than all blueprints) */
+	bool bIsInFindWithinBlueprintMode;
+
 	/** True if current search should not be changed by an external source */
 	bool bIsLocked;
+
+	/** True if progress bar widgets should be hidden */
+	bool bHideProgressBars;
+
+	/** True if users should be allowed to close the cache bar while caching */
+	bool bShowCacheBarCloseButton;
+
+	/** True if users should be allowed to cancel the active caching operation */
+	bool bShowCacheBarCancelButton;
+
+	/** True if the unresponsive warning text should be visible in the cache bar */
+	bool bShowCacheBarUnresponsiveEditorWarningText;
+
+	/** True if cache bar should remain visible after a caching operation has ended */
+	bool bKeepCacheBarProgressVisible;
 };
diff --git a/Engine/Source/Editor/Kismet/Public/ImaginaryBlueprintData.h b/Engine/Source/Editor/Kismet/Public/ImaginaryBlueprintData.h
--- a/Engine/Source/Editor/Kismet/Public/ImaginaryBlueprintData.h
+++ b/Engine/Source/Editor/Kismet/Public/ImaginaryBlueprintData.h
@@ -169,6 +169,12 @@ public:
 		return Outer;
 	}
 
+	/** Called to enable interlocked parsing (only allow one thread at a time). In place to support backwards-compatibility with non-deferred indexing, may be removed later. */
+	void EnableInterlockedParsing()
+	{
+		bRequiresInterlockedParsing = true;
+	}
+
 	/** Dumps the parsed object (including all children) to the given archive */
 	void DumpParsedObject(FArchive& Ar, int32 InTreeLevel = 0) const;
 
@@ -219,6 +225,12 @@ protected:
 	/** Outer of this object that owns it, used for climbing up the hierarchy. Must be declared as thread-safe because it may be accessed on different threads. */
 	FImaginaryFiBDataWeakPtr Outer;
 
+	/** Set after the JSON object has been parsed. */
+	TAtomic<bool> bHasParsedJsonObject;
+
+	/** Set if this instance requires interlocked parsing. */
+	bool bRequiresInterlockedParsing;
+
 	/** Allows for thread-safe parsing of the imaginary data. Only a single Imaginary data can be parsed at a time. */
 	static FCriticalSection ParseChildDataCriticalSection;
 
diff --git a/Engine/Source/Runtime/Engine/Private/Blueprint.cpp b/Engine/Source/Runtime/Engine/Private/Blueprint.cpp
--- a/Engine/Source/Runtime/Engine/Private/Blueprint.cpp
+++ b/Engine/Source/Runtime/Engine/Private/Blueprint.cpp
@@ -363,7 +363,7 @@ void UBlueprint::PreSave(const class ITargetPlatform* TargetPlatform)
 	if (!TargetPlatform || TargetPlatform->HasEditorOnlyData())
 	{
 		// Cache the BP for use (immediate, since we're about to save)
-		FFindInBlueprintSearchManager::Get().AddOrUpdateBlueprintSearchMetadata(this, true);
+		FFindInBlueprintSearchManager::Get().AddOrUpdateBlueprintSearchMetadata(this, EAddOrUpdateBlueprintSearchMetadataFlags::ForceRecache);
 	}
 }
 #endif // WITH_EDITORONLY_DATA
@@ -1000,9 +1000,10 @@ void UBlueprint::GetAssetRegistryTags(TArray<FAssetRegistryTag>& OutTags) const
 	{
 		FString Value;
 		const bool bRebuildSearchData = false;
-		if (const FSearchData* SearchData = FFindInBlueprintSearchManager::Get().QuerySingleBlueprint((UBlueprint*)this, bRebuildSearchData))
+		FSearchData SearchData = FFindInBlueprintSearchManager::Get().QuerySingleBlueprint((UBlueprint*)this, bRebuildSearchData);
+		if (SearchData.IsValid())
 		{
-			Value = SearchData->Value;
+			Value = SearchData.Value;
 		}
 		
 		OutTags.Add( FAssetRegistryTag(FBlueprintTags::FindInBlueprintsData, Value, FAssetRegistryTag::TT_Hidden) );
