Fix duplicate symbol errors by making enums into enumclass.
#rnx

diff --git a/Engine/Source/Editor/Kismet/Public/FindInBlueprintManager.h b/Engine/Source/Editor/Kismet/Public/FindInBlueprintManager.h
--- a/Engine/Source/Editor/Kismet/Public/FindInBlueprintManager.h
+++ b/Engine/Source/Editor/Kismet/Public/FindInBlueprintManager.h
@@ -138,7 +138,7 @@ struct FSearchDataVersionInfo
 };
 
 /** State flags for search database entries */
-enum ESearchDataStateFlags
+enum class ESearchDataStateFlags : uint8
 {
 	None = 0,
 	/** Set when this search database entry has been fully indexed, which is completed asynchronously */
diff --git a/Engine/Source/Runtime/Engine/Private/BlueprintGeneratedClass.cpp b/Engine/Source/Runtime/Engine/Private/BlueprintGeneratedClass.cpp
--- a/Engine/Source/Runtime/Engine/Private/BlueprintGeneratedClass.cpp
+++ b/Engine/Source/Runtime/Engine/Private/BlueprintGeneratedClass.cpp
@@ -205,7 +205,7 @@ UClass* UBlueprintGeneratedClass::GetAuthoritativeClass()
 
 struct FConditionalRecompileClassHepler
 {
-	enum ENeededAction
+	enum class ENeededAction : uint8
 	{
 		None,
 		StaticLink,
