Fixing SBlueprintPalette storing raw property pointer instead of WeakFieldPtr which resulted in crashes when the property became invalid

diff --git a/Engine/Source/Editor/Kismet/Private/SBlueprintPalette.cpp b/Engine/Source/Editor/Kismet/Private/SBlueprintPalette.cpp
--- a/Engine/Source/Editor/Kismet/Private/SBlueprintPalette.cpp
+++ b/Engine/Source/Editor/Kismet/Private/SBlueprintPalette.cpp
@@ -54,6 +54,7 @@
 #include "Widgets/Text/SInlineEditableTextBlock.h"
 #include "SPinTypeSelector.h"
 #include "GraphEditorSettings.h"
+#include "UObject/WeakFieldPtr.h"
 
 #define LOCTEXT_NAMESPACE "BlueprintPalette"
 
@@ -715,7 +716,7 @@ public:
 private:
 	FEdGraphPinType OnGetVarType() const
 	{
-		if (FProperty* VarProp = const_cast<FProperty*>(VariableProperty))
+		if (FProperty* VarProp = const_cast<FProperty*>(VariableProperty.Get()))
 		{
 			const UEdGraphSchema_K2* K2Schema = GetDefault<UEdGraphSchema_K2>();
 			FEdGraphPinType Type;
@@ -729,7 +730,7 @@ private:
 	{
 		if (FBlueprintEditorUtils::IsPinTypeValid(InNewPinType))
 		{
-			if (FProperty* VarProp = VariableProperty)
+			if (FProperty* VarProp = VariableProperty.Get())
 			{
 				FName VarName = VarProp->GetFName();
 
@@ -762,7 +763,7 @@ private:
 	TWeakPtr<FBlueprintEditor>     BlueprintEditorPtr;
 
 	/** Variable Property to change the type of */
-	FProperty* VariableProperty;
+	TWeakFieldPtr<FProperty> VariableProperty;
 };
 
 /*******************************************************************************
