FiB debug utility to facilitate regression testing for unindexed Blueprints and out-of-date Blueprint index tag values.

- Also addresses an issue introduced with my recent async indexing change where some new Blueprint assets may not have been getting added to the index cache on first save/compile.

#rnx

diff --git a/Engine/Source/Editor/BlueprintGraph/Classes/K2Node_Variable.h b/Engine/Source/Editor/BlueprintGraph/Classes/K2Node_Variable.h
--- a/Engine/Source/Editor/BlueprintGraph/Classes/K2Node_Variable.h
+++ b/Engine/Source/Editor/BlueprintGraph/Classes/K2Node_Variable.h
@@ -30,7 +30,7 @@ class BLUEPRINTGRAPH_API UK2Node_Variable : public UK2Node
 	GENERATED_UCLASS_BODY()
 
 	/** Reference to variable we want to set/get */
-	UPROPERTY(meta=(BlueprintSearchable="true", BlueprintSearchableHiddenExplicit="true"))
+	UPROPERTY(meta=(BlueprintSearchable="true", BlueprintSearchableHiddenExplicit="true", BlueprintSearchableFormatVersion="FIB_VER_VARIABLE_REFERENCE"))
 	FMemberReference	VariableReference;
 
 	UPROPERTY()
diff --git a/Engine/Source/Editor/Kismet/Private/FindInBlueprintManager.cpp b/Engine/Source/Editor/Kismet/Private/FindInBlueprintManager.cpp
--- a/Engine/Source/Editor/Kismet/Private/FindInBlueprintManager.cpp
+++ b/Engine/Source/Editor/Kismet/Private/FindInBlueprintManager.cpp
@@ -104,6 +104,7 @@ const FString FFiBMD::FiBSearchableMD = TEXT("BlueprintSearchable");
 const FString FFiBMD::FiBSearchableShallowMD = TEXT("BlueprintSearchableShallow");
 const FString FFiBMD::FiBSearchableExplicitMD = TEXT("BlueprintSearchableExplicit");
 const FString FFiBMD::FiBSearchableHiddenExplicitMD = TEXT("BlueprintSearchableHiddenExplicit");
+const FString FFiBMD::FiBSearchableFormatVersionMD = TEXT("BlueprintSearchableFormatVersion");
 
 /* Return the outer of the specified object that is a direct child of a package */
 inline UObject* GetAssetObject(UObject* InObject)
@@ -432,6 +433,7 @@ namespace BlueprintSearchMetaDataHelpers
 		bool bIsSearchableMD;
 		bool bIsShallowSearchableMD;
 		bool bIsMarkedNotSearchableMD;
+		EFiBVersion MinDataFormatVersion;
 		TArray<FSearchableProperty> ChildProperties;
 	};
 
@@ -442,13 +444,19 @@ namespace BlueprintSearchMetaDataHelpers
 	public:
 		typedef PrintPolicy InnerPrintPolicy;
 
-		TFindInBlueprintJsonStringWriter(FString* const InOutString)
+		TFindInBlueprintJsonStringWriter(FString* const InOutString, int32 InFormatVersion)
 			: TJsonStringWriter<PrintPolicy>(InOutString, 0)
+			, CachedFormatVersion(InFormatVersion)
 		{
 		}
 
 		using TJsonStringWriter<PrintPolicy>::WriteObjectStart;
 
+		const int32& GetFormatVersion() const
+		{
+			return CachedFormatVersion;
+		}
+
 		void WriteObjectStart( const FText& Identifier )
 		{
 			check( this->Stack.Top() == EJson::Object );
@@ -512,14 +520,18 @@ namespace BlueprintSearchMetaDataHelpers
 	public:
 		/** Cached mapping of all searchable properties that have been discovered while gathering searchable data for the current Blueprint */
 		TMap<UStruct*, TArray<FSearchableProperty>> CachedPropertyMapping;
+
+	private:
+		/** Cached version used to determine the data serialization format */
+		int32 CachedFormatVersion;
 	};
 
 	/** Json Writer used for serializing FText's in the correct format for Find-in-Blueprints */
 	class FFindInBlueprintJsonWriter : public TFindInBlueprintJsonStringWriter<TCondensedJsonPrintPolicy<TCHAR>>
 			{
 	public:
-		FFindInBlueprintJsonWriter(FString* const InOutString)
-			:TFindInBlueprintJsonStringWriter<TCondensedJsonPrintPolicy<TCHAR>>(InOutString)
+		FFindInBlueprintJsonWriter(FString* const InOutString, int32 InFormatVersion)
+			:TFindInBlueprintJsonStringWriter<TCondensedJsonPrintPolicy<TCHAR>>(InOutString, InFormatVersion)
 			,JsonOutput(InOutString)
 				{
 		}
@@ -530,7 +542,8 @@ namespace BlueprintSearchMetaDataHelpers
 			bool bResult = TFindInBlueprintJsonStringWriter<TCondensedJsonPrintPolicy<TCHAR>>::Close();
 
 			// Build the search metadata string for the asset tag (version + LUT + JSON)
-			*JsonOutput = FiBSerializationHelpers::Serialize(FSearchDataVersionInfo::Current.FiBDataVersion, false)
+			int32 DataVersion = GetFormatVersion();
+			*JsonOutput = FiBSerializationHelpers::Serialize(DataVersion, false)
 				+ FiBSerializationHelpers::Serialize(LookupTable, true)
 				+ MoveTemp(*JsonOutput);
 
@@ -860,8 +873,9 @@ namespace BlueprintSearchMetaDataHelpers
 				// It only is truly marked as not searchable if it has the metadata set to false, if the metadata is missing then we assume the searchable type that is passed in unless SEARCHABLE_AS_DESIRED
 				bool bIsMarkedNotSearchableMD = SearchableProperty.bIsMarkedNotSearchableMD;
 
-				if ( (InSearchableType != SEARCHABLE_AS_DESIRED && !bIsMarkedNotSearchableMD) 
-					|| bIsShallowSearchableMD || bIsSearchableMD)
+				const bool bShouldGatherSearchableProperty = (InSearchableType != SEARCHABLE_AS_DESIRED && !bIsMarkedNotSearchableMD)
+					|| bIsShallowSearchableMD || bIsSearchableMD;
+				if (bShouldGatherSearchableProperty && InWriter->GetFormatVersion() >= SearchableProperty.MinDataFormatVersion)
 				{
 					const void* Value = Property->ContainerPtrToValuePtr<uint8>(InValue);
 
@@ -977,6 +991,22 @@ namespace BlueprintSearchMetaDataHelpers
 				// It only is truly marked as not searchable if it has the metadata set to false, if the metadata is missing then we assume the searchable type that is passed in unless SEARCHABLE_AS_DESIRED
 				bool bIsMarkedNotSearchableMD = Property->HasMetaData(*FFiBMD::FiBSearchableMD) && !bIsSearchableMD;
 
+				// Searchable properties default to the latest format version, unless otherwise specified in the property metadata.
+				EFiBVersion MinDataFormatVersion = FIB_VER_LATEST;
+				if (Property->HasMetaData(*FFiBMD::FiBSearchableFormatVersionMD))
+				{
+					const FString VersionString = Property->GetMetaData(*FFiBMD::FiBSearchableFormatVersionMD);
+					if (!VersionString.IsEmpty())
+					{
+						const UEnum* FiBVersionEnum = StaticEnum<EFiBVersion>();
+						const int64 VersionValue = FiBVersionEnum->GetValueByNameString(VersionString);
+						if (VersionValue != INDEX_NONE)
+						{
+							MinDataFormatVersion = static_cast<EFiBVersion>(VersionValue);
+						}
+					}
+				}
+
 				if ((InSearchableType != SEARCHABLE_AS_DESIRED && !bIsMarkedNotSearchableMD)
 					|| bIsShallowSearchableMD || bIsSearchableMD)
 				{
@@ -987,6 +1017,7 @@ namespace BlueprintSearchMetaDataHelpers
 					SearchableProperty.bIsSearchableMD = bIsSearchableMD;
 					SearchableProperty.bIsShallowSearchableMD = bIsShallowSearchableMD;
 					SearchableProperty.bIsMarkedNotSearchableMD = bIsMarkedNotSearchableMD;
+					SearchableProperty.MinDataFormatVersion = MinDataFormatVersion;
 
 					if (Property->ArrayDim == 1)
 					{
@@ -1185,11 +1216,16 @@ namespace BlueprintSearchMetaDataHelpers
 		{
 			TArray<UEdGraph*> CompleteGraphList;
 			CompleteGraphList.Append(Blueprint->FunctionGraphs);
+
 			// Gather all interface graphs as functions
-			for (const FBPInterfaceDescription& InterfaceDesc : Blueprint->ImplementedInterfaces)
+			if (InWriter->GetFormatVersion() >= EFiBVersion::FIB_VER_INTERFACE_GRAPHS)
 			{
-				CompleteGraphList.Append(InterfaceDesc.Graphs);
+				for (const FBPInterfaceDescription& InterfaceDesc : Blueprint->ImplementedInterfaces)
+				{
+					CompleteGraphList.Append(InterfaceDesc.Graphs);
+				}
 			}
+			
 			GatherGraphSearchData(InWriter, Blueprint, CompleteGraphList, FFindInBlueprintSearchTags::FiB_Functions, &SubGraphs);
 		}
 
@@ -1235,14 +1271,14 @@ namespace BlueprintSearchMetaDataHelpers
 	}
 
 	template<class JsonWriterType>
-	FString GatherBlueprintSearchMetadata(const UBlueprint* Blueprint)
+	FString GatherBlueprintSearchMetadata(const UBlueprint* Blueprint, int32 FiBDataVersion)
 	{
 		CSV_SCOPED_TIMING_STAT(FindInBlueprint, GatherBlueprintSearchMetadata);
 
 		FString SearchMetaData;
 
 		// The search registry tags for a Blueprint are all in Json
-		TSharedRef<JsonWriterType> Writer = MakeShared<JsonWriterType>(&SearchMetaData);
+		TSharedRef<JsonWriterType> Writer = MakeShared<JsonWriterType>(&SearchMetaData, FiBDataVersion);
 
 		typedef typename JsonWriterType::InnerPrintPolicy PrintPolicyType;
 		GatherBlueprintSearchMetadata<PrintPolicyType>(Writer, Blueprint);
@@ -1651,8 +1687,7 @@ public:
 								if (UBlueprint* Blueprint = SearchData.Blueprint.Get())
 								{
 									using namespace BlueprintSearchMetaDataHelpers;
-									SearchData.Value = GatherBlueprintSearchMetadata<FFindInBlueprintJsonWriter>(Blueprint);
-									SearchData.VersionInfo = FSearchDataVersionInfo::Current;
+									SearchData.Value = GatherBlueprintSearchMetadata<FFindInBlueprintJsonWriter>(Blueprint, SearchData.VersionInfo.FiBDataVersion);
 
 									if (SearchData.Value.Len() > 0)
 									{
@@ -1710,8 +1745,7 @@ public:
 						if (UBlueprint* Blueprint = SearchData.Blueprint.Get())
 						{
 							using namespace BlueprintSearchMetaDataHelpers;
-							SearchData.Value = GatherBlueprintSearchMetadata<FFindInBlueprintJsonWriter>(Blueprint);
-							SearchData.VersionInfo = FSearchDataVersionInfo::Current;
+							SearchData.Value = GatherBlueprintSearchMetadata<FFindInBlueprintJsonWriter>(Blueprint, SearchData.VersionInfo.FiBDataVersion);
 						}
 					}
 
@@ -2058,7 +2092,7 @@ void FFindInBlueprintSearchManager::OnAssetAdded(const FAssetData& InAssetData)
 		UBlueprint* Blueprint   = Handler->RetrieveBlueprint(AssetObject);
 		if (Blueprint)
 		{
-			AddOrUpdateBlueprintSearchMetadata(Blueprint, EAddOrUpdateBlueprintSearchMetadataFlags::AllowNewEntry);
+			AddOrUpdateBlueprintSearchMetadata(Blueprint);
 		}
 	}
 	else if (Handler->AssetContainsBlueprint(InAssetData))
@@ -2321,18 +2355,19 @@ void FFindInBlueprintSearchManager::OnHotReload(bool bWasTriggeredAutomatically)
 	CachedAssetClasses.Reset();
 }
 
-void FFindInBlueprintSearchManager::AddOrUpdateBlueprintSearchMetadata(UBlueprint* InBlueprint, EAddOrUpdateBlueprintSearchMetadataFlags InFlags/* = EAddOrUpdateBlueprintSearchMetadataFlags::None*/)
+void FFindInBlueprintSearchManager::AddOrUpdateBlueprintSearchMetadata(UBlueprint* InBlueprint, EAddOrUpdateBlueprintSearchMetadataFlags InFlags/* = EAddOrUpdateBlueprintSearchMetadataFlags::None*/, EFiBVersion InVersion/* = EFiBVersion::FIB_VER_LATEST*/)
 {
 	CSV_SCOPED_TIMING_STAT(FindInBlueprint, AddOrUpdateBlueprintSearchMetadata);
 	CSV_CUSTOM_STAT(FindInBlueprint, AddOrUpdateCountThisFrame, 1, ECsvCustomStatOp::Accumulate);
 
 	// No need to update the cache in the following cases:
 	//	a) Indexing is disabled.
-	//	b) The Blueprint is not yet fully loaded. This ensures that we don't make attempts to re-index before load completion.
-	//	c) The Blueprint was loaded for diffing. It makes search all very strange and allows you to fully open those Blueprints.
-	//	d) The Blueprint was loaded/copied for PIE. These assets are temporarily created for a session and don't need to be re-indexed.
+	//  b) The Blueprint is explicitly marked as being transient (i.e. internal utility-type assets that aren't saved).
+	//	c) The Blueprint is not yet fully loaded. This ensures that we don't make attempts to re-index before load completion.
+	//	d) The Blueprint was loaded for diffing. It makes search all very strange and allows you to fully open those Blueprints.
+	//	e) The Blueprint was loaded/copied for PIE. These assets are temporarily created for a session and don't need to be re-indexed.
 	if (!bEnableGatheringData
-		|| InBlueprint->HasAnyFlags(RF_NeedLoad | RF_NeedPostLoad)
+		|| InBlueprint->HasAnyFlags(RF_NeedLoad | RF_NeedPostLoad | RF_Transient)
 		|| InBlueprint->GetOutermost()->HasAnyPackageFlags(PKG_ForDiffing | PKG_PlayInEditor))
 	{
 		return;
@@ -2340,7 +2375,7 @@ void FFindInBlueprintSearchManager::AddOrUpdateBlueprintSearchMetadata(UBlueprin
 
 	// Control flags
 	const bool bForceRecache = EnumHasAllFlags(InFlags, EAddOrUpdateBlueprintSearchMetadataFlags::ForceRecache);
-	const bool bAllowNewEntry = EnumHasAllFlags(InFlags, EAddOrUpdateBlueprintSearchMetadataFlags::AllowNewEntry);
+	const bool bClearCachedValue = EnumHasAllFlags(InFlags, EAddOrUpdateBlueprintSearchMetadataFlags::ClearCachedValue);
 
 	UObject* AssetObject = GetAssetObject(InBlueprint);
 
@@ -2354,7 +2389,7 @@ void FFindInBlueprintSearchManager::AddOrUpdateBlueprintSearchMetadata(UBlueprin
 		SearchData.Blueprint = InBlueprint; // Blueprint instance may change due to reloading
 		SearchData.StateFlags &= ~(ESearchDataStateFlags::IsIndexed | ESearchDataStateFlags::WasRemoved);
 	}
-	else if (bAllowNewEntry)
+	else
 	{
 		SearchData = FSearchData();
 		SearchData.AssetPath = AssetPath;
@@ -2372,19 +2407,23 @@ void FFindInBlueprintSearchManager::AddOrUpdateBlueprintSearchMetadata(UBlueprin
 		// Clear any previously-gathered data.
 		SearchData.Value.Empty();
 
+		// Update version info stored in database. This indicates which format to use when regenerating the tag value.
+		SearchData.VersionInfo = FSearchDataVersionInfo::Current;
+		if (InVersion != EFiBVersion::FIB_VER_NONE)
+		{
+			SearchData.VersionInfo.FiBDataVersion = InVersion;
+		}
+
 		// During unindexed/out-of-date caching we will arrive here as a result of loading the asset, so don't remove the IsUnindexedCacheInProgress() check!
-		if (bForceRecache || IsUnindexedCacheInProgress() || bDisableDeferredIndexing)
+		if (bForceRecache || bClearCachedValue || IsUnindexedCacheInProgress() || bDisableDeferredIndexing)
 		{
 			// Cannot successfully gather most searchable data if there is no SkeletonGeneratedClass, so don't try, leave it as whatever it was last set to
-			if (InBlueprint->SkeletonGeneratedClass != nullptr)
+			if (!bClearCachedValue && InBlueprint->SkeletonGeneratedClass != nullptr)
 			{
 				using namespace BlueprintSearchMetaDataHelpers;
 
 				// Update search metadata string content
-				SearchData.Value = GatherBlueprintSearchMetadata<FFindInBlueprintJsonWriter>(InBlueprint);
-
-				// Update version info stored in database to latest
-				SearchData.VersionInfo = FSearchDataVersionInfo::Current;
+				SearchData.Value = GatherBlueprintSearchMetadata<FFindInBlueprintJsonWriter>(InBlueprint, SearchData.VersionInfo.FiBDataVersion);
 			}
 
 			// Mark it as having been indexed in the following cases:
@@ -2411,6 +2450,7 @@ void FFindInBlueprintSearchManager::AddOrUpdateBlueprintSearchMetadata(UBlueprin
 		}
 
 		// Copy new/updated search data into the cache
+		const bool bAllowNewEntry = true;
 		ApplySearchDataToDatabase(MoveTemp(SearchData), bAllowNewEntry);
 	}
 }
@@ -2806,7 +2846,7 @@ FString FFindInBlueprintSearchManager::ConvertFTextToHexString(FText InValue)
 FString FFindInBlueprintSearchManager::GenerateSearchIndexForDebugging(UBlueprint* InBlueprint)
 {
 	using namespace BlueprintSearchMetaDataHelpers;
-	return GatherBlueprintSearchMetadata<TFindInBlueprintJsonStringWriter<TPrettyJsonPrintPolicy<TCHAR>>>(InBlueprint);
+	return GatherBlueprintSearchMetadata<TFindInBlueprintJsonStringWriter<TPrettyJsonPrintPolicy<TCHAR>>>(InBlueprint, EFiBVersion::FIB_VER_LATEST);
 }
 
 void FFindInBlueprintSearchManager::OnCacheAllUnindexedAssets(bool bInSourceControlActive, bool bInCheckoutAndSave)
diff --git a/Engine/Source/Editor/Kismet/Public/FindInBlueprintManager.h b/Engine/Source/Editor/Kismet/Public/FindInBlueprintManager.h
--- a/Engine/Source/Editor/Kismet/Public/FindInBlueprintManager.h
+++ b/Engine/Source/Editor/Kismet/Public/FindInBlueprintManager.h
@@ -111,6 +111,7 @@ struct KISMET_API FFindInBlueprintSearchTags
 };
 
 /** FiB data versioning */
+UENUM()
 enum EFiBVersion
 {
 	FIB_VER_NONE = -1, // Unknown version (not set)
@@ -232,6 +233,7 @@ struct KISMET_API FFiBMD
 	static const FString FiBSearchableShallowMD;
 	static const FString FiBSearchableExplicitMD;
 	static const FString FiBSearchableHiddenExplicitMD;
+	static const FString FiBSearchableFormatVersionMD;
 };
 
 /** Which assets to index for caching */
@@ -293,8 +295,8 @@ enum class EAddOrUpdateBlueprintSearchMetadataFlags
 	None = 0,
 	/** Forces the Blueprint to be recache'd, regardless of what data it believes exists */
 	ForceRecache = 1 << 0,
-	/** Whether to allow a new entry to be added to the database if a match is not found */
-	AllowNewEntry = 1 << 1,
+	/** Clear any cached data value for this Blueprint */
+	ClearCachedValue = 1 << 1,
 };
 
 ENUM_CLASS_FLAGS(EAddOrUpdateBlueprintSearchMetadataFlags);
@@ -497,9 +499,10 @@ public:
 	 * Gathers the Blueprint's search metadata and adds or updates it in the cache
 	 *
 	 * @param InBlueprint		Blueprint to cache the searchable data for
-	 * @param InOptions			Options to assist with controlling this function's behavior
+	 * @param InFlags			Flags to assist with controlling this function's behavior
+	 * @param InVersion			Allows the caller to override the format version to use for caching
 	 */
-	void AddOrUpdateBlueprintSearchMetadata(UBlueprint* InBlueprint, EAddOrUpdateBlueprintSearchMetadataFlags InFlags = EAddOrUpdateBlueprintSearchMetadataFlags::None);
+	void AddOrUpdateBlueprintSearchMetadata(UBlueprint* InBlueprint, EAddOrUpdateBlueprintSearchMetadataFlags InFlags = EAddOrUpdateBlueprintSearchMetadataFlags::None, EFiBVersion InVersion = EFiBVersion::FIB_VER_NONE);
 
 	/**
 	 * Starts a search query, the FiB manager handles where the thread is at in the search query at all times so that post-save of the cache to disk it can correct the index
diff --git a/Engine/Source/Runtime/Engine/Private/Blueprint.cpp b/Engine/Source/Runtime/Engine/Private/Blueprint.cpp
--- a/Engine/Source/Runtime/Engine/Private/Blueprint.cpp
+++ b/Engine/Source/Runtime/Engine/Private/Blueprint.cpp
@@ -353,6 +353,18 @@ UBlueprint::UBlueprint(const FObjectInitializer& ObjectInitializer)
 }
 
 #if WITH_EDITORONLY_DATA
+static TAutoConsoleVariable<bool> CVarBPDisableSearchDataUpdateOnSave(
+	TEXT("bp.DisableSearchDataUpdateOnSave"),
+	false,
+	TEXT("Don't update Blueprint search metadata on save (for QA/testing purposes only). On an editor relaunch, it should include the BP in the unindexed count after the first search."),
+	ECVF_Cheat);
+
+static TAutoConsoleVariable<bool> CVarBPForceOldSearchDataFormatVersionOnSave(
+	TEXT("bp.ForceOldSearchDataFormatVersionOnSave"),
+	false,
+	TEXT("Force Blueprint search metadata to use an old format version on save (for QA/testing purposes only). On an editor relaunch, it should include the BP in the out-of-date count after the first search."),
+	ECVF_Cheat);
+
 void UBlueprint::PreSave(const class ITargetPlatform* TargetPlatform)
 {
 	Super::PreSave(TargetPlatform);
@@ -362,8 +374,26 @@ void UBlueprint::PreSave(const class ITargetPlatform* TargetPlatform)
 
 	if (!TargetPlatform || TargetPlatform->HasEditorOnlyData())
 	{
+		// This will force an immediate (synchronous) update of this Blueprint's index tag value.
+		EAddOrUpdateBlueprintSearchMetadataFlags Flags = EAddOrUpdateBlueprintSearchMetadataFlags::ForceRecache;
+
+		// For regression testing, we exclude the registry tag on save by clearing the cached value.
+		// Expected result: On an editor relaunch it should cause this BP to be reported as "unindexed," until the asset is loaded.
+		if (CVarBPDisableSearchDataUpdateOnSave.GetValueOnGameThread())
+		{
+			Flags |= EAddOrUpdateBlueprintSearchMetadataFlags::ClearCachedValue;
+		}
+
+		// For regression testing, we allow an old format version to be used as an override on save.
+		// Expected result: On an editor relaunch it should cause this BP to be reported as "out-of-date," until the asset is loaded.
+		EFiBVersion OverrideVersion = EFiBVersion::FIB_VER_NONE;
+		if (CVarBPForceOldSearchDataFormatVersionOnSave.GetValueOnGameThread())
+		{
+			OverrideVersion = EFiBVersion::FIB_VER_BASE;
+		}
+
 		// Cache the BP for use (immediate, since we're about to save)
-		FFindInBlueprintSearchManager::Get().AddOrUpdateBlueprintSearchMetadata(this, EAddOrUpdateBlueprintSearchMetadataFlags::ForceRecache);
+		FFindInBlueprintSearchManager::Get().AddOrUpdateBlueprintSearchMetadata(this, Flags, OverrideVersion);
 	}
 }
 #endif // WITH_EDITORONLY_DATA
