Fix a minor FiB issue where a global search could occasionally become blocked until async indexing has finished.

#rnx

diff --git a/Engine/Source/Editor/Kismet/Private/FindInBlueprintManager.cpp b/Engine/Source/Editor/Kismet/Private/FindInBlueprintManager.cpp
--- a/Engine/Source/Editor/Kismet/Private/FindInBlueprintManager.cpp
+++ b/Engine/Source/Editor/Kismet/Private/FindInBlueprintManager.cpp
@@ -2463,6 +2463,87 @@ FFindInBlueprintSearchManager::FActiveSearchQueryPtr FFindInBlueprintSearchManag
 	return ActiveSearchQueries.FindRef(InSearchOriginator);
 }
 
+FSearchData FFindInBlueprintSearchManager::GetNextSearchDataForQuery(FActiveSearchQueryPtr SearchQuery, bool bCheckDeferredList)
+{
+	// Get the entry in the index cache for the next asset to search.
+	FSearchData SearchData = GetSearchDataForIndex(SearchQuery->NextIndex);
+	while (SearchData.IsValid())
+	{
+		// Advance the current search index.
+		++SearchQuery->NextIndex;
+
+		// If this asset has not been indexed, don't search it yet.
+		if (!SearchData.IsIndexingCompleted())
+		{
+			SearchQuery->DeferredAssetPaths.Enqueue(SearchData.AssetPath);
+		}
+		else if (!SearchData.IsMarkedForDeletion())
+		{
+			// Ok to use this asset's index entry; break out of the loop.
+			break;
+		}
+
+		// Move on to the next entry in the index cache.
+		SearchData = GetSearchDataForIndex(SearchQuery->NextIndex);
+	}
+
+	// If we don't have valid search data, try the deferred list from above.
+	if (!SearchData.IsValid() && bCheckDeferredList)
+	{
+		FName AssetPath, FirstAssetPath = NAME_None;
+		while (SearchQuery->DeferredAssetPaths.Dequeue(AssetPath))
+		{
+			// Skip invalid paths (shouldn't happen, but just in case).
+			if (AssetPath == NAME_None)
+			{
+				continue;
+			}
+			else if (AssetPath == FirstAssetPath)
+			{
+				// Yield to allow indexing to progress a bit further.
+				FPlatformProcess::Sleep(0.1f);
+
+				// Check for a new entry in case the cache has grown in size.
+				SearchData = GetNextSearchDataForQuery(SearchQuery, /*bCheckDeferredList = */false);
+				if (SearchData.IsValid())
+				{
+					// Ok to use this entry; break out of the loop.
+					break;
+				}
+			}
+
+			SearchData = GetSearchDataForAssetPath(AssetPath);
+			if (SearchData.IsValid() && !SearchData.IsMarkedForDeletion())
+			{
+				// If this asset is still waiting to be indexed, put it back into the queue.
+				if (!SearchData.IsIndexingCompleted())
+				{
+					SearchQuery->DeferredAssetPaths.Enqueue(AssetPath);
+
+					// Keep track of the first dequeued asset path. If we wrap back around, we'll yield to give the indexing thread more time to work.
+					if (FirstAssetPath == NAME_None)
+					{
+						FirstAssetPath = AssetPath;
+					}
+				}
+				else
+				{
+					// Ok to use this entry; break out of the loop.
+					break;
+				}
+			}
+		}
+	}
+
+	// Increment the search counter if we have valid search data.
+	if (SearchData.IsValid())
+	{
+		++SearchQuery->SearchCount;
+	}
+
+	return SearchData;
+}
+
 void FFindInBlueprintSearchManager::BeginSearchQuery(const FStreamSearch* InSearchOriginator)
 {
 	if (AssetRegistryModule == nullptr)
@@ -2501,67 +2582,9 @@ bool FFindInBlueprintSearchManager::ContinueSearchQuery(const FStreamSearch* InS
 		return false;
 	}
 
-	auto GetNextSearchDataLambda = [this, SearchQuery]()
-	{
-		// Get the entry in the index cache for the next asset to search.
-		FSearchData SearchData = GetSearchDataForIndex(SearchQuery->NextIndex);
-		while (SearchData.IsValid())
-		{
-			// Advance the current search index.
-			++SearchQuery->NextIndex;
-
-			// If this asset has not been indexed, don't search it yet.
-			if (!SearchData.IsIndexingCompleted())
-			{
-				SearchQuery->DeferredAssetPaths.Enqueue(SearchData.AssetPath);
-			}
-			else if (!SearchData.IsMarkedForDeletion())
-			{
-				// Ok to use this asset's index entry; break out of the loop.
-				break;
-			}
-
-			// Move on to the next entry in the index cache.
-			SearchData = GetSearchDataForIndex(SearchQuery->NextIndex);
-		}
-
-		// If we don't have valid search data, try the deferred list from above.
-		if (!SearchData.IsValid())
-		{
-			FName AssetPath;
-			while (SearchQuery->DeferredAssetPaths.Dequeue(AssetPath))
-			{
-				SearchData = GetSearchDataForAssetPath(AssetPath);
-				if (SearchData.IsValid() && !SearchData.IsMarkedForDeletion())
-				{
-					// If this asset is still waiting to be indexed, put it back into the queue.
-					if (!SearchData.IsIndexingCompleted())
-					{
-						SearchQuery->DeferredAssetPaths.Enqueue(AssetPath);
-
-						// Yield to allow indexing to progress a bit further.
-						FPlatformProcess::Sleep(0.1f);
-					}
-					else
-					{
-						// Ok to use this entry; break out of the loop.
-						break;
-					}
-				}
-			}
-		}
-
-		// Increment the search counter if we have valid search data.
-		if (SearchData.IsValid())
-		{
-			++SearchQuery->SearchCount;
-		}
-
-		return SearchData;
-	};
-
-	FSearchData SearchData = GetNextSearchDataLambda();
-
+	// Grab the next entry and update the active search query. Include the list of entries that were deferred pending completion of async indexing work.
+	const bool bCheckDeferredList = true;
+	FSearchData SearchData = GetNextSearchDataForQuery(SearchQuery, bCheckDeferredList);
 	if (SearchData.IsValid())
 	{
 		// In these modes, the full index may not have been parsed yet. We'll do that now on the search thread.
diff --git a/Engine/Source/Editor/Kismet/Public/FindInBlueprintManager.h b/Engine/Source/Editor/Kismet/Public/FindInBlueprintManager.h
--- a/Engine/Source/Editor/Kismet/Public/FindInBlueprintManager.h
+++ b/Engine/Source/Editor/Kismet/Public/FindInBlueprintManager.h
@@ -736,6 +736,9 @@ protected:
 	/** Thread-safe access to the active search query that's mapped to the given stream search */
 	FActiveSearchQueryPtr FindSearchQuery(const class FStreamSearch* InSearchOriginator) const;
 
+	/** Returns the next pending search data for the given query and advances the index to the next entry */
+	FSearchData GetNextSearchDataForQuery(FActiveSearchQueryPtr InSearchQueryPtr, bool bCheckDeferredList);
+
 private:
 	/** Maps the Blueprint paths to their index in the SearchArray */
 	TMap<FName, int32> SearchMap;
