Blueprint Editor menus converted into tool menus
#rnx

diff --git a/Engine/Source/Editor/AnimationBlueprintEditor/Private/AnimationBlueprintEditor.cpp b/Engine/Source/Editor/AnimationBlueprintEditor/Private/AnimationBlueprintEditor.cpp
--- a/Engine/Source/Editor/AnimationBlueprintEditor/Private/AnimationBlueprintEditor.cpp
+++ b/Engine/Source/Editor/AnimationBlueprintEditor/Private/AnimationBlueprintEditor.cpp
@@ -67,6 +67,7 @@
 #include "Widgets/Input/SButton.h"
 #include "EditorFontGlyphs.h"
 #include "AnimationBlueprintInterfaceEditorMode.h"
+#include "ToolMenus.h"
 
 // Hide related nodes feature
 #include "Preferences/AnimationBlueprintEditorOptions.h"
@@ -193,15 +194,6 @@ void FAnimationBlueprintEditor::ExtendMenu()
 	}
 
 	MenuExtender = MakeShareable(new FExtender);
-
-	UAnimBlueprint* AnimBlueprint = GetAnimBlueprint();
-	if(AnimBlueprint)
-	{
-		TSharedPtr<FExtender> AnimBPMenuExtender = MakeShareable(new FExtender);
-		FKismet2Menu::SetupBlueprintEditorMenu(AnimBPMenuExtender, *this);
-		AddMenuExtender(AnimBPMenuExtender);
-	}
-
 	AddMenuExtender(MenuExtender);
 
 	// add extensible menu if exists
@@ -209,6 +201,11 @@ void FAnimationBlueprintEditor::ExtendMenu()
 	AddMenuExtender(AnimationBlueprintEditorModule.GetMenuExtensibilityManager()->GetAllExtenders(GetToolkitCommands(), GetEditingObjects()));
 }
 
+void FAnimationBlueprintEditor::RegisterMenus()
+{
+	FBlueprintEditor::RegisterMenus();
+}
+
 void FAnimationBlueprintEditor::InitAnimationBlueprintEditor(const EToolkitMode::Type Mode, const TSharedPtr< class IToolkitHost >& InitToolkitHost, UAnimBlueprint* InAnimBlueprint)
 {
 	// Record if we have been newly created
@@ -249,6 +246,12 @@ void FAnimationBlueprintEditor::InitAnimationBlueprintEditor(const EToolkitMode:
 	TArray<UObject*> ObjectsBeingEdited;
 	ObjectsBeingEdited.Add(InAnimBlueprint);
 
+	CreateDefaultCommands();
+
+	BindCommands();
+
+	RegisterMenus();
+
 	// Initialize the asset editor and spawn tabs
 	const TSharedRef<FTabManager::FLayout> DummyLayout = FTabManager::NewLayout("NullLayout")->AddArea(FTabManager::NewPrimaryArea());
 	const bool bCreateDefaultStandaloneMenu = true;
@@ -260,8 +263,6 @@ void FAnimationBlueprintEditor::InitAnimationBlueprintEditor(const EToolkitMode:
 
 	CommonInitialization(AnimBlueprints);
 
-	BindCommands();
-
 	if(InAnimBlueprint->BlueprintType == BPTYPE_Interface)
 	{
 		AddApplicationMode(
@@ -462,7 +463,6 @@ void FAnimationBlueprintEditor::OnGraphEditorBackgrounded(const TSharedRef<SGrap
 /** Create Default Tabs **/
 void FAnimationBlueprintEditor::CreateDefaultCommands() 
 {
-	if (GetBlueprintObj())
 	{
 		FBlueprintEditor::CreateDefaultCommands();
 
@@ -472,13 +472,6 @@ void FAnimationBlueprintEditor::CreateDefaultCommands()
 			FCanExecuteAction(),
 			FIsActionChecked::CreateSP(this, &FBlueprintEditor::IsToggleHideUnrelatedNodesChecked));
 	}
-	else
-	{
-		ToolkitCommands->MapAction( FGenericCommands::Get().Undo, 
-			FExecuteAction::CreateSP( this, &FAnimationBlueprintEditor::UndoAction ));
-		ToolkitCommands->MapAction( FGenericCommands::Get().Redo, 
-			FExecuteAction::CreateSP( this, &FAnimationBlueprintEditor::RedoAction ));
-	}
 }
 
 void FAnimationBlueprintEditor::OnCreateGraphEditorCommands(TSharedPtr<FUICommandList> GraphEditorCommandsList)
diff --git a/Engine/Source/Editor/AnimationBlueprintEditor/Private/AnimationBlueprintEditor.h b/Engine/Source/Editor/AnimationBlueprintEditor/Private/AnimationBlueprintEditor.h
--- a/Engine/Source/Editor/AnimationBlueprintEditor/Private/AnimationBlueprintEditor.h
+++ b/Engine/Source/Editor/AnimationBlueprintEditor/Private/AnimationBlueprintEditor.h
@@ -255,6 +255,9 @@ private:
 	/** Extend menu */
 	void ExtendMenu();
 
+	/** Register menus */
+	void RegisterMenus();
+
 	/** Extend toolbar */
 	void ExtendToolbar();
 
diff --git a/Engine/Source/Editor/Kismet/Private/BlueprintEditor.cpp b/Engine/Source/Editor/Kismet/Private/BlueprintEditor.cpp
--- a/Engine/Source/Editor/Kismet/Private/BlueprintEditor.cpp
+++ b/Engine/Source/Editor/Kismet/Private/BlueprintEditor.cpp
@@ -1720,7 +1720,6 @@ void FBlueprintEditor::CommonInitialization(const TArray<UBlueprint*>& InitBluep
 		InitBlueprint->OnSetObjectBeingDebugged().AddSP(this, &FBlueprintEditor::HandleSetObjectBeingDebugged);
 	}
 
-	CreateDefaultCommands();
 	CreateDefaultTabContents(InitBlueprints);
 
 	FKismetEditorUtilities::OnBlueprintUnloaded.AddSP(this, &FBlueprintEditor::OnBlueprintUnloaded);
@@ -1840,6 +1839,10 @@ void FBlueprintEditor::InitBlueprintEditor(
 
 	GetToolkitCommands()->Append(FPlayWorldCommands::GlobalPlayWorldActions.ToSharedRef());
 
+	CreateDefaultCommands();
+
+	RegisterMenus();
+
 	// Initialize the asset editor and spawn nothing (dummy layout)
 	const TSharedRef<FTabManager::FLayout> DummyLayout = FTabManager::NewLayout("NullLayout")->AddArea(FTabManager::NewPrimaryArea());
 	const bool bCreateDefaultStandaloneMenu = true;
@@ -1909,10 +1912,6 @@ void FBlueprintEditor::InitToolMenuContext(FToolMenuContext& MenuContext)
 
 void FBlueprintEditor::InitalizeExtenders()
 {
-	TSharedPtr<FExtender> MenuExtender = MakeShareable(new FExtender);
-	FKismet2Menu::SetupBlueprintEditorMenu(MenuExtender, *this);
-	AddMenuExtender(MenuExtender);
-
 	FBlueprintEditorModule* BlueprintEditorModule = &FModuleManager::LoadModuleChecked<FBlueprintEditorModule>("Kismet");
 	TSharedPtr<FExtender> CustomExtenders = BlueprintEditorModule->GetMenuExtensibilityManager()->GetAllExtenders(GetToolkitCommands(), GetEditingObjects());
 	BlueprintEditorModule->OnGatherBlueprintMenuExtensions().Broadcast(CustomExtenders, GetBlueprintObj());
@@ -1921,6 +1920,15 @@ void FBlueprintEditor::InitalizeExtenders()
 	AddToolbarExtender(CustomExtenders);
 }
 
+void FBlueprintEditor::RegisterMenus()
+{
+	const FName MainMenuName = GetToolMenuName();
+	if (!UToolMenus::Get()->IsMenuRegistered(MainMenuName))
+	{
+		FKismet2Menu::SetupBlueprintEditorMenu(MainMenuName);
+	}
+}
+
 void FBlueprintEditor::RegisterApplicationModes(const TArray<UBlueprint*>& InBlueprints, bool bShouldOpenInDefaultsMode, bool bNewlyCreated/* = false*/)
 {
 	// Newly-created Blueprints will open in Components mode rather than Standard mode
diff --git a/Engine/Source/Editor/Kismet/Private/SBlueprintEditorToolbar.cpp b/Engine/Source/Editor/Kismet/Private/SBlueprintEditorToolbar.cpp
--- a/Engine/Source/Editor/Kismet/Private/SBlueprintEditorToolbar.cpp
+++ b/Engine/Source/Editor/Kismet/Private/SBlueprintEditorToolbar.cpp
@@ -69,169 +69,164 @@ public:
 //////////////////////////////////////////////////////////////////////////
 // FKismet2Menu
 
-void FKismet2Menu::FillFileMenuBlueprintSection( FMenuBuilder& MenuBuilder, FBlueprintEditor& Kismet )
+void FKismet2Menu::FillFileMenuBlueprintSection(UToolMenu* InMenu)
 {
-	MenuBuilder.BeginSection("FileBlueprint", LOCTEXT("BlueprintHeading", "Blueprint"));
+	FToolMenuInsert InsertPosition("FileLoadAndSave", EToolMenuInsertType::After);
+
 	{
-		MenuBuilder.AddMenuEntry( FBlueprintEditorCommands::Get().CompileBlueprint );
-		MenuBuilder.AddMenuEntry( FBlueprintEditorCommands::Get().RefreshAllNodes );
-		MenuBuilder.AddMenuEntry( FBlueprintEditorCommands::Get().ReparentBlueprint );
-		MenuBuilder.AddWrapperSubMenu(
+		FToolMenuSection& Section = InMenu->AddSection("FileBlueprint", LOCTEXT("BlueprintHeading", "Blueprint"));
+		Section.InsertPosition = InsertPosition;
+		Section.AddMenuEntry( FBlueprintEditorCommands::Get().CompileBlueprint );
+		Section.AddMenuEntry( FBlueprintEditorCommands::Get().RefreshAllNodes );
+		Section.AddMenuEntry( FBlueprintEditorCommands::Get().ReparentBlueprint );
+		Section.AddSubMenu(
+			"Diff",
 			LOCTEXT("Diff", "Diff"),
 			LOCTEXT("BlueprintEditorDiffToolTip", "Diff against previous revisions"),
-			FOnGetContent::CreateStatic< FBlueprintEditor& >( &FKismet2Menu::MakeDiffMenu, Kismet),
-			FSlateIcon());
-		MenuBuilder.AddMenuEntry(FBlueprintEditorCommands::Get().BeginBlueprintMerge);
+			FNewToolMenuWidget::CreateStatic(&FKismet2Menu::MakeDiffMenu));
+		Section.AddMenuEntry(FBlueprintEditorCommands::Get().BeginBlueprintMerge);
 	}
-	MenuBuilder.EndSection();
 
-	// Only show the developer menu on machines with the solution (assuming they can build it)
-	ISourceCodeAccessModule* SourceCodeAccessModule = FModuleManager::GetModulePtr<ISourceCodeAccessModule>("SourceCodeAccess");
-	if(SourceCodeAccessModule != nullptr && SourceCodeAccessModule->GetAccessor().CanAccessSourceCode())
+	InMenu->AddDynamicSection("FileDeveloper", FNewToolMenuDelegate::CreateLambda([InsertPosition](UToolMenu* InMenu)
 	{
-		MenuBuilder.BeginSection("FileDeveloper");
+		// Only show the developer menu on machines with the solution (assuming they can build it)
+		ISourceCodeAccessModule* SourceCodeAccessModule = FModuleManager::GetModulePtr<ISourceCodeAccessModule>("SourceCodeAccess");
+		if (SourceCodeAccessModule != nullptr && SourceCodeAccessModule->GetAccessor().CanAccessSourceCode())
 		{
-			MenuBuilder.AddSubMenu( 
+			FToolMenuSection& Section = InMenu->AddSection("FileDeveloper");
+			Section.InsertPosition = InsertPosition;
+			Section.AddSubMenu(
+				"DeveloperMenu",
 				LOCTEXT("DeveloperMenu", "Developer"),
 				LOCTEXT("DeveloperMenu_ToolTip", "Open the developer menu"),
-				FNewMenuDelegate::CreateStatic( &FKismet2Menu::FillDeveloperMenu ),
+				FNewToolMenuDelegate::CreateStatic( &FKismet2Menu::FillDeveloperMenu ),
 				false);
 		}
-		MenuBuilder.EndSection();
-	}
+	}));
 }
 
-void FKismet2Menu::FillDeveloperMenu( FMenuBuilder& MenuBuilder )
+void FKismet2Menu::FillDeveloperMenu(UToolMenu* InMenu)
 {
-	MenuBuilder.BeginSection("FileDeveloperCompilerSettings", LOCTEXT("CompileOptionsHeading", "Compiler Settings"));
 	{
-		MenuBuilder.AddMenuEntry( FBlueprintEditorCommands::Get().SaveIntermediateBuildProducts );
+		FToolMenuSection& Section = InMenu->AddSection("FileDeveloperCompilerSettings", LOCTEXT("CompileOptionsHeading", "Compiler Settings"));
+		Section.AddMenuEntry( FBlueprintEditorCommands::Get().SaveIntermediateBuildProducts );
 	}
-	MenuBuilder.EndSection();
 
-	MenuBuilder.BeginSection("GenerateNativeCode", LOCTEXT("Cpp", "C++"));
 	{
-		MenuBuilder.AddMenuEntry(FBlueprintEditorCommands::Get().GenerateNativeCode);
+		FToolMenuSection& Section = InMenu->AddSection("GenerateNativeCode", LOCTEXT("Cpp", "C++"));
+		Section.AddMenuEntry(FBlueprintEditorCommands::Get().GenerateNativeCode);
 	}
-	MenuBuilder.EndSection();
 
 	if (FFindInBlueprintSearchManager::Get().ShouldEnableDeveloperMenuTools())
 	{
-		MenuBuilder.BeginSection("FileDeveloperSearchTools", LOCTEXT("SearchToolsHeading", "Search Tools"));
-		{
-			MenuBuilder.AddMenuEntry(FBlueprintEditorCommands::Get().GenerateSearchIndex);
-			MenuBuilder.AddMenuEntry(FBlueprintEditorCommands::Get().DumpCachedIndexData);
-		}
-		MenuBuilder.EndSection();
+		FToolMenuSection& Section = InMenu->AddSection("FileDeveloperSearchTools", LOCTEXT("SearchToolsHeading", "Search Tools"));
+		Section.AddMenuEntry(FBlueprintEditorCommands::Get().GenerateSearchIndex);
+		Section.AddMenuEntry(FBlueprintEditorCommands::Get().DumpCachedIndexData);
 	}
 
 	if (false)
 	{
-		MenuBuilder.BeginSection("FileDeveloperFindReferences");
 		{
-			MenuBuilder.AddMenuEntry( FBlueprintEditorCommands::Get().FindReferencesFromClass );
-			MenuBuilder.AddMenuEntry( FBlueprintEditorCommands::Get().FindReferencesFromBlueprint );
-			MenuBuilder.AddMenuEntry( FBlueprintEditorCommands::Get().FindReferencesFromBlueprint );
+			FToolMenuSection& Section = InMenu->AddSection("FileDeveloperFindReferences");
+			Section.AddMenuEntry(FBlueprintEditorCommands::Get().FindReferencesFromClass);
+			Section.AddMenuEntry(FBlueprintEditorCommands::Get().FindReferencesFromBlueprint);
+			Section.AddMenuEntry(FBlueprintEditorCommands::Get().FindReferencesFromBlueprint);
 		}
-		MenuBuilder.EndSection();
 	}
 
-	MenuBuilder.BeginSection("SchemaDeveloperSettings", LOCTEXT("SchemaDevUtilsHeading", "Schema Utilities"));
 	{
-		MenuBuilder.AddMenuEntry(FBlueprintEditorCommands::Get().ShowActionMenuItemSignatures);
+		FToolMenuSection& Section = InMenu->AddSection("SchemaDeveloperSettings", LOCTEXT("SchemaDevUtilsHeading", "Schema Utilities"));
+		Section.AddMenuEntry(FBlueprintEditorCommands::Get().ShowActionMenuItemSignatures);
 	}
-	MenuBuilder.EndSection();
 }
 
-void FKismet2Menu::FillEditMenu( FMenuBuilder& MenuBuilder )
+void FKismet2Menu::FillEditMenu(UToolMenu* InMenu)
 {
-	MenuBuilder.BeginSection("EditSearch", LOCTEXT("EditMenu_SearchHeading", "Search") );
 	{
-		MenuBuilder.AddMenuEntry( FBlueprintEditorCommands::Get().FindInBlueprint );
-		MenuBuilder.AddMenuEntry( FBlueprintEditorCommands::Get().FindInBlueprints );
-		MenuBuilder.AddMenuEntry( FBlueprintEditorCommands::Get().DeleteUnusedVariables );
+		FToolMenuSection& Section = InMenu->AddSection("EditSearch", LOCTEXT("EditMenu_SearchHeading", "Search"));
+		Section.InsertPosition = FToolMenuInsert("EditHistory", EToolMenuInsertType::After);
+		Section.AddMenuEntry( FBlueprintEditorCommands::Get().FindInBlueprint );
+		Section.AddMenuEntry( FBlueprintEditorCommands::Get().FindInBlueprints );
+		Section.AddMenuEntry( FBlueprintEditorCommands::Get().DeleteUnusedVariables );
 	}
-	MenuBuilder.EndSection();
 }
 
-void FKismet2Menu::FillViewMenu( FMenuBuilder& MenuBuilder )
+void FKismet2Menu::FillViewMenu(UToolMenu* InMenu)
 {
-	MenuBuilder.BeginSection("ViewPinVisibility", LOCTEXT("ViewMenu_PinVisibilityHeading", "Pin Visibility") );
 	{
-		MenuBuilder.AddMenuEntry(FGraphEditorCommands::Get().ShowAllPins);
-		MenuBuilder.AddMenuEntry(FGraphEditorCommands::Get().HideNoConnectionNoDefaultPins);
-		MenuBuilder.AddMenuEntry(FGraphEditorCommands::Get().HideNoConnectionPins);
+		FToolMenuSection& Section = InMenu->AddSection("ViewPinVisibility", LOCTEXT("ViewMenu_PinVisibilityHeading", "Pin Visibility"));
+		Section.AddMenuEntry(FGraphEditorCommands::Get().ShowAllPins);
+		Section.AddMenuEntry(FGraphEditorCommands::Get().HideNoConnectionNoDefaultPins);
+		Section.AddMenuEntry(FGraphEditorCommands::Get().HideNoConnectionPins);
 	}
-	MenuBuilder.EndSection();
 
-	MenuBuilder.BeginSection("ViewZoom", LOCTEXT("ViewMenu_ZoomHeading", "Zoom") );
 	{
-		MenuBuilder.AddMenuEntry( FBlueprintEditorCommands::Get().ZoomToWindow );
-		MenuBuilder.AddMenuEntry( FBlueprintEditorCommands::Get().ZoomToSelection );
+		FToolMenuSection& Section = InMenu->AddSection("ViewZoom", LOCTEXT("ViewMenu_ZoomHeading", "Zoom") );
+		Section.AddMenuEntry( FBlueprintEditorCommands::Get().ZoomToWindow );
+		Section.AddMenuEntry( FBlueprintEditorCommands::Get().ZoomToSelection );
 	}
-	MenuBuilder.EndSection();
 }
 
-void FKismet2Menu::FillDebugMenu( FMenuBuilder& MenuBuilder )
+void FKismet2Menu::FillDebugMenu(UToolMenu* InMenu)
 {
-	MenuBuilder.BeginSection("DebugBreakpoints", LOCTEXT("DebugMenu_BreakpointHeading", "Breakpoints") );
 	{
-		MenuBuilder.AddMenuEntry( FBlueprintEditorCommands::Get().DisableAllBreakpoints );
-		MenuBuilder.AddMenuEntry( FBlueprintEditorCommands::Get().EnableAllBreakpoints );
-		MenuBuilder.AddMenuEntry( FBlueprintEditorCommands::Get().ClearAllBreakpoints );
+		FToolMenuSection& Section = InMenu->AddSection("DebugBreakpoints", LOCTEXT("DebugMenu_BreakpointHeading", "Breakpoints"));
+		Section.AddMenuEntry( FBlueprintEditorCommands::Get().DisableAllBreakpoints );
+		Section.AddMenuEntry( FBlueprintEditorCommands::Get().EnableAllBreakpoints );
+		Section.AddMenuEntry( FBlueprintEditorCommands::Get().ClearAllBreakpoints );
 	}
-	MenuBuilder.EndSection();
 
-	MenuBuilder.BeginSection("DebugWatches", LOCTEXT("DebugMenu_WatchHeading", "Watches") );
 	{
-		MenuBuilder.AddMenuEntry( FBlueprintEditorCommands::Get().ClearAllWatches );
+		FToolMenuSection& Section = InMenu->AddSection("DebugWatches", LOCTEXT("DebugMenu_WatchHeading", "Watches"));
+		Section.AddMenuEntry( FBlueprintEditorCommands::Get().ClearAllWatches );
 	}
-	MenuBuilder.EndSection();
 }
 
-void FKismet2Menu::SetupBlueprintEditorMenu( TSharedPtr< FExtender > Extender, FBlueprintEditor& BlueprintEditor)
+void FKismet2Menu::SetupBlueprintEditorMenu(const FName MainMenuName)
 {
-	// Extend the File menu with asset actions
-	Extender->AddMenuExtension(
-		"FileLoadAndSave",
-		EExtensionHook::After,
-		BlueprintEditor.GetToolkitCommands(),
-		FMenuExtensionDelegate::CreateStatic< FBlueprintEditor& >( &FKismet2Menu::FillFileMenuBlueprintSection, BlueprintEditor ) );
-
-	// Extend the Edit menu
-	Extender->AddMenuExtension(
-		"EditHistory",
-		EExtensionHook::After,
-		BlueprintEditor.GetToolkitCommands(),
-		FMenuExtensionDelegate::CreateStatic( &FKismet2Menu::FillEditMenu ) );
+	const FName ParentMenuName("MainFrame.MainMenu");
+
+	{
+		const FName FileMenuName = *(MainMenuName.ToString() + TEXT(".File"));
+		if (!UToolMenus::Get()->IsMenuRegistered(FileMenuName))
+		{
+			FKismet2Menu::FillFileMenuBlueprintSection(UToolMenus::Get()->RegisterMenu(FileMenuName, *(ParentMenuName.ToString() + TEXT(".File"))));
+		}
+	}
+
+	{
+		const FName EditMenuName = *(MainMenuName.ToString() + TEXT(".Edit"));
+		if (!UToolMenus::Get()->IsMenuRegistered(EditMenuName))
+		{
+			FKismet2Menu::FillEditMenu(UToolMenus::Get()->RegisterMenu(EditMenuName, *(ParentMenuName.ToString() + TEXT(".Edit"))));
+		}
+	}
 
 	// Add additional blueprint editor menus
 	{
-		struct Local
+		FToolMenuSection& Section = UToolMenus::Get()->ExtendMenu(MainMenuName)->FindOrAddSection(NAME_None);
+
+		// View
+		if (!Section.FindEntry("View"))
 		{
-			static void AddBlueprintEditorMenus( FMenuBarBuilder& MenuBarBuilder )
-			{
-				// View
-				MenuBarBuilder.AddPullDownMenu( 
-					LOCTEXT("ViewMenu", "View"),
-					LOCTEXT("ViewMenu_ToolTip", "Open the View menu"),
-					FNewMenuDelegate::CreateStatic( &FKismet2Menu::FillViewMenu ),
-					"View");
-
-				// Debug
-				MenuBarBuilder.AddPullDownMenu( 
-					LOCTEXT("DebugMenu", "Debug"),
-					LOCTEXT("DebugMenu_ToolTip", "Open the debug menu"),
-					FNewMenuDelegate::CreateStatic( &FKismet2Menu::FillDebugMenu ),
-					"Debug");
-			}
-		};
+			Section.AddSubMenu(
+				"View",
+				LOCTEXT("ViewMenu", "View"),
+				LOCTEXT("ViewMenu_ToolTip", "Open the View menu"),
+				FNewToolMenuDelegate::CreateStatic(&FKismet2Menu::FillViewMenu)
+			).InsertPosition = FToolMenuInsert("Edit", EToolMenuInsertType::After);
+		}
 
-		Extender->AddMenuBarExtension(
-			"Edit",
-			EExtensionHook::After,
-			BlueprintEditor.GetToolkitCommands(), 
-			FMenuBarExtensionDelegate::CreateStatic( &Local::AddBlueprintEditorMenus ) );
+		// Debug
+		if (!Section.FindEntry("Debug"))
+		{
+			Section.AddSubMenu(
+				"Debug",
+				LOCTEXT("DebugMenu", "Debug"),
+				LOCTEXT("DebugMenu_ToolTip", "Open the debug menu"),
+				FNewToolMenuDelegate::CreateStatic(&FKismet2Menu::FillDebugMenu)
+			).InsertPosition = FToolMenuInsert("Edit", EToolMenuInsertType::After);
+		}
 	}
 }
 
@@ -310,11 +305,12 @@ static void OnDiffRevisionPicked(FRevisionInfo const& RevisionInfo, TWeakObjectP
 	}
 }
 
-TSharedRef<SWidget> FKismet2Menu::MakeDiffMenu(FBlueprintEditor& Kismet)
+TSharedRef<SWidget> FKismet2Menu::MakeDiffMenu(const FToolMenuContext& InToolMenuContext)
 {
 	if (ISourceControlModule::Get().IsEnabled() && ISourceControlModule::Get().GetProvider().IsAvailable())
 	{
-		UBlueprint* BlueprintObj = Kismet.GetBlueprintObj();
+		UBlueprintEditorToolMenuContext* Context = InToolMenuContext.FindContext<UBlueprintEditorToolMenuContext>();
+		UBlueprint* BlueprintObj = Context ? Context->GetBlueprintObj() : nullptr;
 		if(BlueprintObj)
 		{
 			TWeakObjectPtr<UBlueprint> BlueprintPtr = BlueprintObj;
diff --git a/Engine/Source/Editor/Kismet/Public/BlueprintEditor.h b/Engine/Source/Editor/Kismet/Public/BlueprintEditor.h
--- a/Engine/Source/Editor/Kismet/Public/BlueprintEditor.h
+++ b/Engine/Source/Editor/Kismet/Public/BlueprintEditor.h
@@ -1302,6 +1302,9 @@ protected:
 	void CollectPureUpstreamNodes(UEdGraphNode* CurrentNode, TArray<UEdGraphNode*>& CollectedNodes);
 	void HideUnrelatedNodes();
 
+	/** Register Menus */
+	void RegisterMenus();
+
 public:
 	/** Make nodes which are unrelated to the selected nodes fade out */
 	void ToggleHideUnrelatedNodes();
diff --git a/Engine/Source/Editor/Kismet/Public/SBlueprintEditorToolbar.h b/Engine/Source/Editor/Kismet/Public/SBlueprintEditorToolbar.h
--- a/Engine/Source/Editor/Kismet/Public/SBlueprintEditorToolbar.h
+++ b/Engine/Source/Editor/Kismet/Public/SBlueprintEditorToolbar.h
@@ -15,6 +15,7 @@ class FExtender;
 class FMenuBuilder;
 class FToolBarBuilder;
 class UToolMenu;
+struct FToolMenuContext;
 
 /**
  * Kismet menu
@@ -22,24 +23,24 @@ class UToolMenu;
 class KISMET_API FKismet2Menu
 {
 public:
-	static void SetupBlueprintEditorMenu( TSharedPtr< FExtender > Extender, FBlueprintEditor& Kismet);
+	static void SetupBlueprintEditorMenu(const FName MainMenuName);
 	
 protected:
-	static void FillFileMenuBlueprintSection( FMenuBuilder& MenuBuilder, FBlueprintEditor& Kismet );
+	static void FillFileMenuBlueprintSection(UToolMenu* Menu);
 
-	static void FillEditMenu( FMenuBuilder& MenuBuilder );
+	static void FillEditMenu(UToolMenu* Menu);
 
-	static void FillViewMenu( FMenuBuilder& MenuBuilder );
+	static void FillViewMenu(UToolMenu* Menu);
 
-	static void FillDebugMenu( FMenuBuilder& MenuBuilder );
+	static void FillDebugMenu(UToolMenu* Menu);
 
-	static void FillDeveloperMenu( FMenuBuilder& MenuBuilder );
+	static void FillDeveloperMenu(UToolMenu* Menu);
 
 private:
 	/** Diff current blueprint against the specified revision */
 	static void DiffAgainstRevision( class UBlueprint* Current, int32 OldRevision );
 
-	static TSharedRef<SWidget> MakeDiffMenu(FBlueprintEditor& Kismet);
+	static TSharedRef<SWidget> MakeDiffMenu(const FToolMenuContext& ToolMenuContext);
 };
 
 
