Added horizontal view mode to the Blueprint diff tool

diff --git a/Engine/Source/Editor/EditorStyle/Private/SlateEditorStyle.cpp b/Engine/Source/Editor/EditorStyle/Private/SlateEditorStyle.cpp
--- a/Engine/Source/Editor/EditorStyle/Private/SlateEditorStyle.cpp
+++ b/Engine/Source/Editor/EditorStyle/Private/SlateEditorStyle.cpp
@@ -6037,6 +6037,8 @@ void FSlateEditorStyle::FStyle::SetupPersonaStyle()
 		Set("BlueprintDif.NextDiff.Small", new IMAGE_BRUSH("/Icons/diff_next_40x", Icon16x16));
 		Set("BlueprintDif.PrevDiff", new IMAGE_BRUSH("/Icons/diff_prev_40x", Icon16x16));
 		Set("BlueprintDif.PrevDiff.Small", new IMAGE_BRUSH("/Icons/diff_prev_40x", Icon16x16));
+		Set("BlueprintDif.HorizontalDiff.Small", new IMAGE_BRUSH("/Icons/icon_horizontal_diff_view_16px", Icon16x16));
+		Set("BlueprintDif.VerticalDiff.Small", new IMAGE_BRUSH("/Icons/icon_vertical_diff_view_16px", Icon16x16));
 
 		Set("BlueprintDif.ItalicText", 
 			FTextBlockStyle(NormalText)
diff --git a/Engine/Source/Editor/Kismet/Private/SBlueprintDiff.cpp b/Engine/Source/Editor/Kismet/Private/SBlueprintDiff.cpp
--- a/Engine/Source/Editor/Kismet/Private/SBlueprintDiff.cpp
+++ b/Engine/Source/Editor/Kismet/Private/SBlueprintDiff.cpp
@@ -913,8 +913,8 @@ void SBlueprintDiff::Construct( const FArguments& InArgs)
 		AssetEditorCloseDelegate = GEditor->GetEditorSubsystem<UAssetEditorSubsystem>()->OnAssetEditorRequestClose().AddSP(this, &SBlueprintDiff::OnCloseAssetEditor);
 	}
 
-	FToolBarBuilder ToolbarBuilder(TSharedPtr< const FUICommandList >(), FMultiBoxCustomization::None);
-	ToolbarBuilder.AddToolBarButton(
+	FToolBarBuilder NavToolBarBuilder(TSharedPtr< const FUICommandList >(), FMultiBoxCustomization::None);
+	NavToolBarBuilder.AddToolBarButton(
 		FUIAction(
 			FExecuteAction::CreateSP(this, &SBlueprintDiff::PrevDiff),
 			FCanExecuteAction::CreateSP( this, &SBlueprintDiff::HasPrevDiff)
@@ -924,7 +924,7 @@ void SBlueprintDiff::Construct( const FArguments& InArgs)
 		, LOCTEXT("PrevDiffTooltip", "Go to previous difference")
 		, FSlateIcon(FEditorStyle::GetStyleSetName(), "BlueprintDif.PrevDiff")
 	);
-	ToolbarBuilder.AddToolBarButton(
+	NavToolBarBuilder.AddToolBarButton(
 		FUIAction(
 			FExecuteAction::CreateSP(this, &SBlueprintDiff::NextDiff),
 			FCanExecuteAction::CreateSP(this, &SBlueprintDiff::HasNextDiff)
@@ -934,14 +934,22 @@ void SBlueprintDiff::Construct( const FArguments& InArgs)
 		, LOCTEXT("NextDiffTooltip", "Go to next difference")
 		, FSlateIcon(FEditorStyle::GetStyleSetName(), "BlueprintDif.NextDiff")
 	);
-	ToolbarBuilder.AddSeparator();
-	ToolbarBuilder.AddToolBarButton(
+
+	FToolBarBuilder GraphToolbarBuilder(TSharedPtr< const FUICommandList >(), FMultiBoxCustomization::None);
+	GraphToolbarBuilder.AddToolBarButton(
 		FUIAction(FExecuteAction::CreateSP(this, &SBlueprintDiff::OnToggleLockView))
 		, NAME_None
 		, LOCTEXT("LockGraphsLabel", "Lock/Unlock")
 		, LOCTEXT("LockGraphsTooltip", "Force all graph views to change together, or allow independent scrolling/zooming")
 		, TAttribute<FSlateIcon>(this, &SBlueprintDiff::GetLockViewImage)
 	);
+	GraphToolbarBuilder.AddToolBarButton(
+		FUIAction(FExecuteAction::CreateSP(this, &SBlueprintDiff::OnToggleSplitViewMode))
+		, NAME_None
+		, LOCTEXT("SplitGraphsModeLabel", "Vertical/Horizontal")
+		, LOCTEXT("SplitGraphsModeLabelTooltip", "Toggles the split view of graphs between vertical and horizontal")
+		, TAttribute<FSlateIcon>(this, &SBlueprintDiff::GetSplitViewModeImage)
+	);
 
 	DifferencesTreeView = DiffTreeView::CreateTreeView(&MasterDifferencesList);
 
@@ -961,7 +969,7 @@ void SBlueprintDiff::Construct( const FArguments& InArgs)
 		];
 	};
 
-	TSharedRef<SWidget> Overlay = 
+	TopRevisionInfoWidget =
 		SNew(SSplitter)
 		.Visibility(EVisibility::HitTestInvisible)
 		+ SSplitter::Slot()
@@ -986,6 +994,25 @@ void SBlueprintDiff::Construct( const FArguments& InArgs)
 			]
 		];
 
+	GraphToolBarWidget = 
+		SNew(SSplitter)
+		.Visibility(EVisibility::HitTestInvisible)
+		+ SSplitter::Slot()
+		.Value(.2f)
+		[
+			SNew(SBox)
+		]
+		+ SSplitter::Slot()
+		.Value(.8f)
+		[
+			SNew(SHorizontalBox)
+			+SHorizontalBox::Slot()
+			.AutoWidth()
+			[
+				GraphToolbarBuilder.MakeWidget()
+			]	
+		];
+
 	this->ChildSlot
 		[
 			SNew(SBorder)
@@ -1004,7 +1031,7 @@ void SBlueprintDiff::Construct( const FArguments& InArgs)
 						.Padding(4.f)
 						.AutoWidth()
 						[
-							ToolbarBuilder.MakeWidget()
+							NavToolBarBuilder.MakeWidget()
 						]
 						+ SHorizontalBox::Slot()
 						[
@@ -1033,7 +1060,13 @@ void SBlueprintDiff::Construct( const FArguments& InArgs)
 				+ SOverlay::Slot()
 				.VAlign(VAlign_Top)
 				[
-					Overlay
+					TopRevisionInfoWidget.ToSharedRef()		
+				]
+				+ SOverlay::Slot()
+				.VAlign(VAlign_Top)
+				.Padding(0.0f, 6.0f, 0.0f, 4.0f)
+				[
+					GraphToolBarWidget.ToSharedRef()
 				]
 			]
 		];
@@ -1237,11 +1270,26 @@ void SBlueprintDiff::OnToggleLockView()
 	ResetGraphEditors();
 }
 
+void SBlueprintDiff::OnToggleSplitViewMode()
+{
+	bVerticalSplitGraphMode = !bVerticalSplitGraphMode;
+
+	if(SSplitter* DiffGraphSplitterPtr = DiffGraphSplitter.Get())
+	{
+		DiffGraphSplitterPtr->SetOrientation(bVerticalSplitGraphMode ? Orient_Horizontal : Orient_Vertical);
+	}
+}
+
 FSlateIcon SBlueprintDiff::GetLockViewImage() const
 {
 	return FSlateIcon(FEditorStyle::GetStyleSetName(), bLockViews ? "GenericLock" : "GenericUnlock");
 }
 
+FSlateIcon SBlueprintDiff::GetSplitViewModeImage() const
+{
+	return FSlateIcon(FEditorStyle::GetStyleSetName(), bVerticalSplitGraphMode ? "BlueprintDif.VerticalDiff.Small" : "BlueprintDif.HorizontalDiff.Small");
+}
+
 void SBlueprintDiff::ResetGraphEditors()
 {
 	if(PanelOld.GraphEditor.IsValid() && PanelNew.GraphEditor.IsValid())
@@ -1634,23 +1682,16 @@ SBlueprintDiff::FDiffControl SBlueprintDiff::GenerateGraphPanel()
 			+SSplitter::Slot()
 			.Value(.8f)
 			[
-				SNew(SSplitter)
+				SAssignNew(DiffGraphSplitter,SSplitter)
 				.PhysicalSplitterHandleSize(10.0f)
-				+ SSplitter::Slot()
+				.Orientation(bVerticalSplitGraphMode ? Orient_Horizontal : Orient_Vertical)
+				+ SSplitter::Slot() // Old revision graph slot
 				[
-					SAssignNew(PanelOld.GraphEditorBox, SBox)
-					.VAlign(VAlign_Fill)
-					[
-						DefaultEmptyPanel()
-					]
+					GenerateGraphWidgetForPanel(PanelOld)
 				]
-				+ SSplitter::Slot()
+				+ SSplitter::Slot() // New revision graph slot
 				[
-					SAssignNew(PanelNew.GraphEditorBox, SBox)
-					.VAlign(VAlign_Fill)
-					[
-						DefaultEmptyPanel()
-					]
+					GenerateGraphWidgetForPanel(PanelNew)
 				]
 			]
 			+ SSplitter::Slot()
@@ -1745,6 +1786,43 @@ SBlueprintDiff::FDiffControl SBlueprintDiff::GenerateComponentsPanel()
 	return Ret;
 }
 
+TSharedRef<SOverlay> SBlueprintDiff::GenerateGraphWidgetForPanel(FDiffPanel& OutDiffPanel) const
+{
+	return SNew(SOverlay)
+		+ SOverlay::Slot() // Graph slot
+		[
+			SAssignNew(OutDiffPanel.GraphEditorBox, SBox)
+			.HAlign(HAlign_Fill)
+			[
+				DefaultEmptyPanel()
+			]
+		]
+		+ SOverlay::Slot() // Revision info slot
+		.VAlign(VAlign_Bottom)
+		.HAlign(HAlign_Right)
+		.Padding(FMargin(20.0f, 10.0f))
+		[
+			GenerateRevisionInfoWidgetForPanel(OutDiffPanel.OverlayGraphRevisionInfo,
+			DiffViewUtils::GetPanelLabel(OutDiffPanel.Blueprint, OutDiffPanel.RevisionInfo, FText()))
+		];
+}
+
+TSharedRef<SBox> SBlueprintDiff::GenerateRevisionInfoWidgetForPanel(TSharedPtr<SWidget>& OutGeneratedWidget,
+	const FText& InRevisionText) const
+{
+	return SAssignNew(OutGeneratedWidget,SBox)
+		.Padding(FMargin(4.0f, 10.0f))
+		.VAlign(VAlign_Center)
+		.HAlign(HAlign_Left)
+		[
+			SNew(STextBlock)
+			.TextStyle(FEditorStyle::Get(), "DetailsView.CategoryTextStyle")
+			.Text(InRevisionText)
+			.ShadowColorAndOpacity(FColor::Black)
+			.ShadowOffset(FVector2D(1.4,1.4))
+		];
+}
+
 void SBlueprintDiff::SetCurrentMode(FName NewMode)
 {
 	if( CurrentMode == NewMode )
@@ -1768,6 +1846,35 @@ void SBlueprintDiff::SetCurrentMode(FName NewMode)
 	{
 		ensureMsgf(false, TEXT("Diff panel does not support mode %s"), *NewMode.ToString() );
 	}
+
+	OnModeChanged(NewMode);
+}
+
+void SBlueprintDiff::UpdateTopSectionVisibility(const FName& InNewViewMode) const
+{
+	SSplitter* GraphToolBarPtr = GraphToolBarWidget.Get();
+	SSplitter* TopRevisionInfoWidgetPtr = TopRevisionInfoWidget.Get();
+
+	if (!GraphToolBarPtr || !TopRevisionInfoWidgetPtr)
+	{
+		return;
+	}
+	
+	if (InNewViewMode == GraphMode)
+	{
+		GraphToolBarPtr->SetVisibility(EVisibility::Visible);
+		TopRevisionInfoWidgetPtr->SetVisibility(EVisibility::Collapsed);
+	}
+	else
+	{
+		GraphToolBarPtr->SetVisibility(EVisibility::Collapsed);
+		TopRevisionInfoWidgetPtr->SetVisibility(EVisibility::HitTestInvisible);
+	}
+}
+
+void SBlueprintDiff::OnModeChanged(const FName& InNewViewMode) const
+{
+	UpdateTopSectionVisibility(InNewViewMode);
 }
 
 bool SBlueprintDiff::IsGraphDiffNeeded(UEdGraph* InGraph) const
diff --git a/Engine/Source/Editor/Kismet/Public/SBlueprintDiff.h b/Engine/Source/Editor/Kismet/Public/SBlueprintDiff.h
--- a/Engine/Source/Editor/Kismet/Public/SBlueprintDiff.h
+++ b/Engine/Source/Editor/Kismet/Public/SBlueprintDiff.h
@@ -93,6 +93,9 @@ struct KISMET_API FDiffPanel
 
 	/** The panel stores the last pin that was focused on by the user, so that it can clear the visual style when selection changes */
 	UEdGraphPin*					LastFocusedPin;
+
+	/** The widget that contains the revision info in graph mode */
+	TSharedPtr<SWidget>				OverlayGraphRevisionInfo;
 private:
 	/** Command list for this diff panel */
 	TSharedPtr<FUICommandList> GraphEditorCommands;
@@ -159,12 +162,18 @@ protected:
 
 	/** User toggles the option to lock the views between the two blueprints */
 	void OnToggleLockView();
+	
+	/** User toggles the option to change the split view mode betwwen vertical and horizontal */
+	void OnToggleSplitViewMode();
 
 	/** Reset the graph editor, called when user switches graphs to display*/
 	void ResetGraphEditors();
 
 	/** Get the image to show for the toggle lock option*/
 	FSlateIcon GetLockViewImage() const;
+	
+	/** Get the image to show for the toggle split view mode option*/
+	FSlateIcon GetSplitViewModeImage() const;
 
 	/** List of graphs to diff, are added to panel last */
 	TArray<TSharedPtr<FGraphToDiff>> Graphs;
@@ -203,9 +212,15 @@ protected:
 	FDiffControl GenerateClassSettingsPanel();
 	FDiffControl GenerateComponentsPanel();
 
+	TSharedRef<SOverlay> GenerateGraphWidgetForPanel(FDiffPanel& OutDiffPanel) const;
+	TSharedRef<SBox> GenerateRevisionInfoWidgetForPanel(TSharedPtr<SWidget>& OutGeneratedWidget,const FText& InRevisionText) const;
+
 	/** Accessor and event handler for toggling between diff view modes (defaults, components, graph view, interface, macro): */
 	void SetCurrentMode(FName NewMode);
 	FName GetCurrentMode() const { return CurrentMode; }
+	void OnModeChanged(const FName& InNewViewMode) const;
+
+	void UpdateTopSectionVisibility(const FName& InNewViewMode) const;
 
 	FName CurrentMode;
 
@@ -215,9 +230,18 @@ protected:
 	/** If the two views should be locked */
 	bool	bLockViews;
 
+	/** If the view on Graph Mode should be divided vertically */
+	bool bVerticalSplitGraphMode = true;
+
 	/** Contents widget that we swap when mode changes (defaults, components, etc) */
 	TSharedPtr<SBox> ModeContents;
 
+	TSharedPtr<SSplitter> TopRevisionInfoWidget;
+
+	TSharedPtr<SSplitter> DiffGraphSplitter;
+	
+	TSharedPtr<SSplitter> GraphToolBarWidget;
+
 	friend struct FListItemGraphToDiff;
 
 	/** We can't use the global tab manager because we need to instance the diff control, so we have our own tab manager: */
