Check Blueprint graph editor tab validity before setting view location.

#rnx

diff --git a/Engine/Source/Editor/Kismet/Private/BlueprintEditor.cpp b/Engine/Source/Editor/Kismet/Private/BlueprintEditor.cpp
--- a/Engine/Source/Editor/Kismet/Private/BlueprintEditor.cpp
+++ b/Engine/Source/Editor/Kismet/Private/BlueprintEditor.cpp
@@ -8830,9 +8830,11 @@ void FBlueprintEditor::RestoreEditedObjectState()
 					}
 				};
 				TSharedPtr<SDockTab> TabWithGraph = LocalStruct::OpenGraphTree(this, Graph);
-
-				TSharedRef<SGraphEditor> GraphEditor = StaticCastSharedRef<SGraphEditor>(TabWithGraph->GetContent());
-				GraphEditor->SetViewLocation(Blueprint->LastEditedDocuments[i].SavedViewOffset, Blueprint->LastEditedDocuments[i].SavedZoomAmount);
+				if (TabWithGraph.IsValid())
+				{
+					TSharedRef<SGraphEditor> GraphEditor = StaticCastSharedRef<SGraphEditor>(TabWithGraph->GetContent());
+					GraphEditor->SetViewLocation(Blueprint->LastEditedDocuments[i].SavedViewOffset, Blueprint->LastEditedDocuments[i].SavedZoomAmount);
+				}
 			}
 			else
 			{
