Check blacklist filter before spawning tabs according to workflow layout.

#rnx

diff --git a/Engine/Source/Editor/Kismet/Private/WorkflowOrientedApp/WorkflowTabFactory.cpp b/Engine/Source/Editor/Kismet/Private/WorkflowOrientedApp/WorkflowTabFactory.cpp
--- a/Engine/Source/Editor/Kismet/Private/WorkflowOrientedApp/WorkflowTabFactory.cpp
+++ b/Engine/Source/Editor/Kismet/Private/WorkflowOrientedApp/WorkflowTabFactory.cpp
@@ -9,6 +9,8 @@
 #include "WorkflowOrientedApp/WorkflowTabManager.h"
 #include "IDocumentation.h"
 #include "Widgets/Docking/SDockTab.h"
+#include "Framework/Docking/TabManager.h"
+#include "Misc/BlacklistNames.h"
 
 /////////////////////////////////////////////////////
 // FWorkflowTabFactory
@@ -112,13 +114,24 @@ TSharedRef<SDockTab> FWorkflowTabFactory::OnSpawnTab(const FSpawnTabArgs& SpawnA
 	return SpawnInfo.TabManager.IsValid() ? SpawnTab(SpawnInfo) : SNew(SDockTab);
 }
 
+bool FWorkflowTabFactory::CanSpawnTab(const FSpawnTabArgs& SpawnArgs, TWeakPtr<FTabManager> WeakTabManager) const
+{
+	TSharedPtr<FTabManager> TabManager = WeakTabManager.Pin();
+	if (TabManager.IsValid())
+	{
+		return TabManager->GetTabBlacklist()->PassesFilter(SpawnArgs.GetTabId().TabType);
+	}
+
+	return true;
+}
+
 FTabSpawnerEntry& FWorkflowTabFactory::RegisterTabSpawner(TSharedRef<FTabManager> InTabManager, const FApplicationMode* CurrentApplicationMode) const
 {
 	FWorkflowTabSpawnInfo SpawnInfo;
 	SpawnInfo.TabManager = InTabManager;
 
 	TWeakPtr<FTabManager> WeakTabManager(InTabManager);
-	FTabSpawnerEntry& SpawnerEntry = InTabManager->RegisterTabSpawner(GetIdentifier(), FOnSpawnTab::CreateSP(this, &FWorkflowTabFactory::OnSpawnTab, WeakTabManager))
+	FTabSpawnerEntry& SpawnerEntry = InTabManager->RegisterTabSpawner(GetIdentifier(), FOnSpawnTab::CreateSP(this, &FWorkflowTabFactory::OnSpawnTab, WeakTabManager), FCanSpawnTab::CreateSP(this, &FWorkflowTabFactory::CanSpawnTab, WeakTabManager))
 		.SetDisplayName(ConstructTabName(SpawnInfo).Get())
 		.SetTooltipText(GetTabToolTipText(SpawnInfo));
 
diff --git a/Engine/Source/Editor/Kismet/Private/WorkflowOrientedApp/WorkflowTabManager.cpp b/Engine/Source/Editor/Kismet/Private/WorkflowOrientedApp/WorkflowTabManager.cpp
--- a/Engine/Source/Editor/Kismet/Private/WorkflowOrientedApp/WorkflowTabManager.cpp
+++ b/Engine/Source/Editor/Kismet/Private/WorkflowOrientedApp/WorkflowTabManager.cpp
@@ -10,6 +10,7 @@
 #include "EditorStyleSet.h"
 #include "WorkflowOrientedApp/WorkflowUObjectDocuments.h"
 #include "Widgets/Docking/SDockTab.h"
+#include "Misc/BlacklistNames.h"
 
 /////////////////////////////////////////////////////
 // FTabInfo
@@ -672,7 +673,7 @@ TSharedPtr<FDocumentTabFactory> FDocumentTracker::FindSupportingFactory(TSharedR
 	{
 		TSharedPtr<FDocumentTabFactory> Factory = FactoryIt.Value();
 
-		if (Factory->IsPayloadSupported(Payload))
+		if (Factory->IsPayloadSupported(Payload) && TabManager->GetTabBlacklist()->PassesFilter(Factory->GetIdentifier()))
 		{
 			return Factory;
 		}
diff --git a/Engine/Source/Editor/Kismet/Public/WorkflowOrientedApp/WorkflowTabFactory.h b/Engine/Source/Editor/Kismet/Public/WorkflowOrientedApp/WorkflowTabFactory.h
--- a/Engine/Source/Editor/Kismet/Public/WorkflowOrientedApp/WorkflowTabFactory.h
+++ b/Engine/Source/Editor/Kismet/Public/WorkflowOrientedApp/WorkflowTabFactory.h
@@ -133,6 +133,9 @@ protected:
 	/** Callback function for spawning the tab */
 	TSharedRef<SDockTab> OnSpawnTab(const FSpawnTabArgs& SpawnArgs, TWeakPtr<FTabManager> WeakTabManager) const;
 
+	/** Callback function for checking to see if a tab can/should be spawned */
+	bool CanSpawnTab(const FSpawnTabArgs& SpawnArgs, TWeakPtr<FTabManager> WeakTabManager) const;
+
 public:
 	// Creates the widget that will be the body of the new tab
 	// Users will most commonly override just this.
