Optimize GetDependentBlueprints by caching dependencies. Saves 4 minutes during CookByTheBook.
#rnx

diff --git a/Engine/Source/Editor/Kismet/Private/BlueprintCompilationManager.cpp b/Engine/Source/Editor/Kismet/Private/BlueprintCompilationManager.cpp
--- a/Engine/Source/Editor/Kismet/Private/BlueprintCompilationManager.cpp
+++ b/Engine/Source/Editor/Kismet/Private/BlueprintCompilationManager.cpp
@@ -2187,7 +2187,6 @@ void FBlueprintCompilationManagerImpl::ReinstanceBatch(TArray<FReinstancingJob>&
 			{
 				ArchetypeReferencers.Add(BP->SkeletonGeneratedClass);
 			}
-			ensure(BP->bCachedDependenciesUpToDate);
 			for(const TWeakObjectPtr<UBlueprint>& Dependency : BP->CachedDependencies)
 			{
 				if (UBlueprint* DependencyBP = Dependency.Get())
