PR #6194: Change blueprint editor tab history to work across entire blueprint (Contributed by [REDACTED_PERSON])

diff --git a/Engine/Source/Editor/Kismet/Private/BlueprintEditor.cpp b/Engine/Source/Editor/Kismet/Private/BlueprintEditor.cpp
--- a/Engine/Source/Editor/Kismet/Private/BlueprintEditor.cpp
+++ b/Engine/Source/Editor/Kismet/Private/BlueprintEditor.cpp
@@ -8699,7 +8699,7 @@ void FBlueprintEditor::RestoreEditedObjectState()
 							if (OuterObject->IsA<UBlueprint>())
 							{
 								// reached up to the blueprint for the graph, we are done climbing the tree
-								OpenCause = FDocumentTracker::OpenNewDocument;
+								OpenCause = FDocumentTracker::RestorePreviousDocument;
 								break;
 							}
 							else if(UEdGraph* OuterGraph = Cast<UEdGraph>(OuterObject))
diff --git a/Engine/Source/Editor/Kismet/Private/WorkflowOrientedApp/WorkflowTabManager.cpp b/Engine/Source/Editor/Kismet/Private/WorkflowOrientedApp/WorkflowTabManager.cpp
--- a/Engine/Source/Editor/Kismet/Private/WorkflowOrientedApp/WorkflowTabManager.cpp
+++ b/Engine/Source/Editor/Kismet/Private/WorkflowOrientedApp/WorkflowTabManager.cpp
@@ -20,15 +20,15 @@ namespace WorkflowTabManagerHelpers
 	const int32 MaxHistoryEntries = 300;
 }
 
-FTabInfo::FTabInfo(const TSharedRef<SDockTab>& InTab, const TSharedPtr<FDocumentTabFactory>& InSpawner)
+FTabInfo::FTabInfo(const TSharedRef<SDockTab>& InTab, const TSharedPtr<FDocumentTabFactory>& InSpawner, const TSharedPtr<class FDocumentTracker>& InTracker)
 	: Tab(InTab)
-	, CurrentHistoryIndex(INDEX_NONE)
+	, WeakTracker(InTracker)
 {
 }
 
 bool FTabInfo::PayloadMatches(const TSharedPtr<FTabPayload> TestPayload) const
 {
-	return PayloadMatches(History[CurrentHistoryIndex]->GetPayload(), TestPayload);
+	return PayloadMatches(GetPayload(), TestPayload);
 }
 
 bool FTabInfo::PayloadMatches(TSharedPtr<FTabPayload> A, TSharedPtr<FTabPayload> B)
@@ -47,55 +47,88 @@ bool FTabInfo::PayloadMatches(TSharedPtr<FTabPayload> A, TSharedPtr<FTabPayload>
 	}
 }
 
-void FTabInfo::AddTabHistory(TSharedPtr< struct FGenericTabHistory > InHistoryNode, bool bInSaveHistory/* = true*/, bool bPrevTabMatches/* = false*/)
+void FTabInfo::AddTabHistory(TSharedPtr< struct FGenericTabHistory > InHistoryNode, bool bInSaveHistory/* = true*/)
 {
-	// If the tab is not new, save the current history.
-	if(CurrentHistoryIndex >= 0 && bInSaveHistory)
-	{
-		History[CurrentHistoryIndex]->SaveHistory();
-	}
+	TSharedPtr<FDocumentTracker> Tracker = WeakTracker.Pin();
 
-	if (CurrentHistoryIndex == History.Num() - 1)
+	if (Tracker.IsValid())
 	{
-		// History added to the end
-		if (History.Num() == WorkflowTabManagerHelpers::MaxHistoryEntries)
+		if (Tracker->CurrentHistoryIndex == Tracker->History.Num() - 1)
 		{
-			// If max history entries has been reached
-			// remove the oldest history
-			History.RemoveAt(0);
+			// History added to the end
+			if (Tracker->History.Num() == WorkflowTabManagerHelpers::MaxHistoryEntries)
+			{
+				// If max history entries has been reached
+				// remove the oldest history
+				Tracker->History.RemoveAt(0);
+			}
 		}
+		else
+		{
+			// Clear out any history that is in front of the current location in the history list
+			Tracker->History.RemoveAt(Tracker->CurrentHistoryIndex + 1, Tracker->History.Num() - (Tracker->CurrentHistoryIndex + 1), true);
+		}
+
+		Tracker->History.Add(InHistoryNode);
+		Tracker->CurrentHistoryIndex = Tracker->History.Num() - 1;
+
+		SetCurrentHistory(InHistoryNode, bInSaveHistory, false);
 	}
-	else
+}
+
+TSharedPtr<struct FGenericTabHistory> FTabInfo::GetCurrentHistory() const
+{
+	return CurrentHistory;
+}
+
+void FTabInfo::SetCurrentHistory(TSharedPtr<FGenericTabHistory> NewHistory, bool bInSaveHistory, bool bShouldRestore)
+{
+	bool bPayloadMatches = false;
+
+	if (CurrentHistory.IsValid() && NewHistory.IsValid())
 	{
-		// Clear out any history that is in front of the current location in the history list
-		History.RemoveAt(CurrentHistoryIndex + 1, History.Num() - (CurrentHistoryIndex + 1), true);
+		if (bInSaveHistory)
+		{
+			CurrentHistory->SaveHistory();
+		}
+
+		bPayloadMatches = PayloadMatches(CurrentHistory->GetPayload(), NewHistory->GetPayload());
 	}
+	CurrentHistory = NewHistory;
 
-	History.Add(InHistoryNode);
-	CurrentHistoryIndex = History.Num() - 1;
+	if (CurrentHistory.IsValid())
+	{
+		CurrentHistory->BindToTab(AsShared());
 
-	// Evoke the history
-	InHistoryNode->EvokeHistory(AsShared(), bPrevTabMatches);
-	InHistoryNode->GetFactory().Pin()->OnTabActivated(Tab.Pin());
+		// This creates the tab widget but does not foreground it
+		CurrentHistory->EvokeHistory(AsShared(), bPayloadMatches);
+		if (bShouldRestore)
+		{
+			CurrentHistory->RestoreHistory();
+		}
+		TSharedPtr<FDocumentTabFactory> Factory = CurrentHistory->GetFactory().Pin();
+		if (Factory.IsValid())
+		{
+			// Notify listeners that tab contents have changed
+			Factory->OnTabActivated(Tab.Pin());
+		}
+		
+	}
 }
 
 FReply FTabInfo::OnGoForwardInHistory()
 {
-	if( CurrentHistoryIndex < History.Num() - 1 )
-	{
-		const int32 PreviousHistoryIndex = CurrentHistoryIndex;
-		History[CurrentHistoryIndex]->SaveHistory();
+	TSharedPtr<FDocumentTracker> Tracker = WeakTracker.Pin();
 
-		while( CurrentHistoryIndex < History.Num() - 1)
+	if (Tracker.IsValid())
+	{
+		int32 NextHistoryIndex = Tracker->CurrentHistoryIndex;
+		while (NextHistoryIndex < Tracker->History.Num() - 1)
 		{
-			++CurrentHistoryIndex;
-
-			if( History[CurrentHistoryIndex]->IsHistoryValid() )
+			++NextHistoryIndex;
+				
+			if (Tracker->NavigateToTabHistory(NextHistoryIndex))
 			{
-				bool bPayloadMatches = PayloadMatches(History[PreviousHistoryIndex]->GetPayload(), History[CurrentHistoryIndex]->GetPayload());
-				History[CurrentHistoryIndex]->EvokeHistory(AsShared(), bPayloadMatches);
-				History[CurrentHistoryIndex]->RestoreHistory();
-				History[CurrentHistoryIndex]->GetFactory().Pin()->OnTabActivated(Tab.Pin());
 				break;
 			}
 		}
@@ -105,22 +138,17 @@ FReply FTabInfo::OnGoForwardInHistory()
 
 FReply FTabInfo::OnGoBackInHistory()
 {
-	if( CurrentHistoryIndex > 0 )
-	{
-		const int32 PreviousHistoryIndex = CurrentHistoryIndex;
-		History[CurrentHistoryIndex]->SaveHistory();
+	TSharedPtr<FDocumentTracker> Tracker = WeakTracker.Pin();
 
-		while( CurrentHistoryIndex > 0 )
+	if (Tracker.IsValid())
+	{
+		int32 PreviousHistoryIndex = Tracker->CurrentHistoryIndex;
+		while (PreviousHistoryIndex > 0)
 		{
-			--CurrentHistoryIndex;
+			--PreviousHistoryIndex;
 				
-			if( History[CurrentHistoryIndex]->IsHistoryValid() )
+			if (Tracker->NavigateToTabHistory(PreviousHistoryIndex))
 			{
-				bool bPayloadMatches = PayloadMatches(History[PreviousHistoryIndex]->GetPayload(), History[CurrentHistoryIndex]->GetPayload());
-				History[CurrentHistoryIndex]->EvokeHistory(AsShared(), bPayloadMatches);
-				History[CurrentHistoryIndex]->RestoreHistory();
-				History[CurrentHistoryIndex]->GetFactory().Pin()->OnTabActivated(Tab.Pin());
-
 				break;
 			}
 		}
@@ -130,76 +158,63 @@ FReply FTabInfo::OnGoBackInHistory()
 
 void FTabInfo::JumpToNearestValidHistoryData()
 {
-	if(!History[CurrentHistoryIndex]->IsHistoryValid())
+	// Each tab only knows current history, close if it's invalid
+	if (!CurrentHistory.IsValid() || !CurrentHistory->IsHistoryValid())
 	{
-		if(History.Num() == 1)
+		if (Tab.IsValid())
 		{
 			Tab.Pin()->RequestCloseTab();
 		}
-		else
-		{
-			OnGoBackInHistory();
-
-			if(!History[CurrentHistoryIndex]->IsHistoryValid())
-			{
-				OnGoForwardInHistory();
-
-				if(!History[CurrentHistoryIndex]->IsHistoryValid())
-				{
-					// There are no valid history nodes to switch to, delete the tab
-					Tab.Pin()->RequestCloseTab();
-					return;
-				}
-			}
-
-			History[CurrentHistoryIndex]->EvokeHistory(AsShared(), false);
-			History[CurrentHistoryIndex]->RestoreHistory();
-			History[CurrentHistoryIndex]->GetFactory().Pin()->OnTabActivated(Tab.Pin());
-			FGlobalTabmanager::Get()->SetActiveTab(nullptr);
-			FSlateApplication::Get().ClearKeyboardFocus(EFocusCause::SetDirectly);
-		}
 	}
 }
 
 TWeakPtr<FDocumentTabFactory> FTabInfo::GetFactory() const
 {
-	return History[CurrentHistoryIndex]->GetFactory();
+	if (CurrentHistory.IsValid())
+	{
+		return CurrentHistory->GetFactory();
+	}
+	return nullptr;
 }
 
 TSharedPtr<FTabPayload> FTabInfo::GetPayload() const
 {
-	return History[CurrentHistoryIndex]->GetPayload();
+	if (CurrentHistory.IsValid())
+	{
+		return CurrentHistory->GetPayload();
+	}
+	return nullptr;
 }
 
 void FTabInfo::GoToHistoryIndex(int32 InHistoryIdx)
 {
-	if(History[InHistoryIdx]->IsHistoryValid())
+	TSharedPtr<FDocumentTracker> Tracker = WeakTracker.Pin();
+	if (Tracker.IsValid())
 	{
-		const int32 PreviousHistoryIndex = CurrentHistoryIndex;
-		History[CurrentHistoryIndex]->SaveHistory();
-
-		CurrentHistoryIndex = InHistoryIdx;
-		
-		bool bPayloadMatches = PayloadMatches(History[PreviousHistoryIndex]->GetPayload(), History[CurrentHistoryIndex]->GetPayload());
-		History[CurrentHistoryIndex]->EvokeHistory(AsShared(), bPayloadMatches);
-		History[CurrentHistoryIndex]->RestoreHistory();
-		History[CurrentHistoryIndex]->GetFactory().Pin()->OnTabActivated(Tab.Pin());
+		Tracker->NavigateToTabHistory(InHistoryIdx);
 	}
 }
 
 TSharedRef<SWidget> FTabInfo::CreateHistoryMenu(bool bInBackHistory) const
 {
-	FMenuBuilder MenuBuilder(true, NULL);
+	TSharedPtr<FDocumentTracker> Tracker = WeakTracker.Pin();
+	if (!Tracker.IsValid())
+	{
+		return SNullWidget::NullWidget;
+	}
+
+	FMenuBuilder MenuBuilder(true, nullptr);
 	if(bInBackHistory)
 	{
-		int32 HistoryIdx = CurrentHistoryIndex - 1;
+		int32 HistoryIdx = Tracker->CurrentHistoryIndex - 1;
 		while( HistoryIdx >= 0 )
 		{
-			if(History[HistoryIdx]->IsHistoryValid())
+			TSharedPtr<FGenericTabHistory> HistoryItem = Tracker->History[HistoryIdx];
+			if(HistoryItem->IsHistoryValid())
 			{
-				MenuBuilder.AddMenuEntry(History[HistoryIdx]->GetHistoryTitle().Get(), FText(), FSlateIcon(), 
+				MenuBuilder.AddMenuEntry(HistoryItem->GetHistoryTitle().Get(), FText(), FSlateIcon(), 
 					FUIAction(
-					FExecuteAction::CreateRaw(const_cast<FTabInfo*>(this), &FTabInfo::GoToHistoryIndex, HistoryIdx)
+					FExecuteAction::CreateRaw(this, &FTabInfo::GoToHistoryIndex, HistoryIdx)
 					), 
 					NAME_None, EUserInterfaceActionType::Button);
 			}
@@ -209,14 +224,15 @@ TSharedRef<SWidget> FTabInfo::CreateHistoryMenu(bool bInBackHistory) const
 	}
 	else
 	{
-		int32 HistoryIdx = CurrentHistoryIndex + 1;
-		while( HistoryIdx < History.Num() )
+		int32 HistoryIdx = Tracker->CurrentHistoryIndex + 1;
+		while( HistoryIdx < Tracker->History.Num() )
 		{
-			if(History[HistoryIdx]->IsHistoryValid())
+			TSharedPtr<FGenericTabHistory> HistoryItem = Tracker->History[HistoryIdx];
+			if(HistoryItem->IsHistoryValid())
 			{
-				MenuBuilder.AddMenuEntry(History[HistoryIdx]->GetHistoryTitle().Get(), FText(), FSlateIcon(), 
+				MenuBuilder.AddMenuEntry(HistoryItem->GetHistoryTitle().Get(), FText(), FSlateIcon(), 
 					FUIAction(
-					FExecuteAction::CreateRaw(const_cast<FTabInfo*>(this), &FTabInfo::GoToHistoryIndex, HistoryIdx)
+					FExecuteAction::CreateRaw(this, &FTabInfo::GoToHistoryIndex, HistoryIdx)
 					), 
 					NAME_None, EUserInterfaceActionType::Button);
 			}
@@ -230,10 +246,16 @@ TSharedRef<SWidget> FTabInfo::CreateHistoryMenu(bool bInBackHistory) const
 
 bool FTabInfo::CanStepBackwardInHistory() const
 {
-	int32 HistoryIdx = CurrentHistoryIndex - 1;
+	TSharedPtr<FDocumentTracker> Tracker = WeakTracker.Pin();
+	if (!Tracker.IsValid())
+	{
+		return false;
+	}
+
+	int32 HistoryIdx = Tracker->CurrentHistoryIndex - 1;
 	while( HistoryIdx >= 0 )
 	{
-		if(History[HistoryIdx]->IsHistoryValid())
+		if(Tracker->History[HistoryIdx]->IsHistoryValid())
 		{
 			return true;
 		}
@@ -245,10 +267,16 @@ bool FTabInfo::CanStepBackwardInHistory() const
 
 bool FTabInfo::CanStepForwardInHistory() const
 {
-	int32 HistoryIdx = CurrentHistoryIndex + 1;
-	while( HistoryIdx < History.Num() )
+	TSharedPtr<FDocumentTracker> Tracker = WeakTracker.Pin();
+	if (!Tracker.IsValid())
+	{
+		return false;
+	}
+
+	int32 HistoryIdx = Tracker->CurrentHistoryIndex + 1;
+	while( HistoryIdx < Tracker->History.Num() )
 	{
-		if(History[HistoryIdx]->IsHistoryValid())
+		if(Tracker->History[HistoryIdx]->IsHistoryValid())
 		{
 			return true;
 		}
@@ -258,7 +286,7 @@ bool FTabInfo::CanStepForwardInHistory() const
 	return false;
 }
 
-FReply FTabInfo::OnMouseDownHisory( const FGeometry& MyGeometry, const FPointerEvent& MouseEvent, TWeakPtr< SMenuAnchor > InMenuAnchor )
+FReply FTabInfo::OnMouseDownHistory( const FGeometry& MyGeometry, const FPointerEvent& MouseEvent, TWeakPtr< SMenuAnchor > InMenuAnchor )
 {
 	if(MouseEvent.GetEffectingButton() == EKeys::RightMouseButton && InMenuAnchor.IsValid())
 	{
@@ -281,7 +309,7 @@ TSharedRef< SWidget > FTabInfo::CreateHistoryNavigationWidget()
 			.AutoWidth()
 			[
 				SNew(SBorder)
-				.OnMouseButtonDown(this, &FTabInfo::OnMouseDownHisory, BackMenuAnchorPtr)
+				.OnMouseButtonDown(this, &FTabInfo::OnMouseDownHistory, BackMenuAnchorPtr)
 				.BorderImage( FEditorStyle::GetBrush("NoBorder") )
 				[
 					SAssignNew(BackMenuAnchorPtr, SMenuAnchor)
@@ -305,7 +333,7 @@ TSharedRef< SWidget > FTabInfo::CreateHistoryNavigationWidget()
 			.AutoWidth()
 			[
 				SNew(SBorder)
-				.OnMouseButtonDown(this, &FTabInfo::OnMouseDownHisory, FwdMenuAnchorPtr)
+				.OnMouseButtonDown(this, &FTabInfo::OnMouseDownHistory, FwdMenuAnchorPtr)
 				.BorderImage( FEditorStyle::GetBrush("NoBorder") )
 				[
 					SAssignNew(FwdMenuAnchorPtr, SMenuAnchor)
@@ -389,6 +417,8 @@ FDocumentTracker::FDocumentTracker()
 	// Make sure we know when tabs become active
 	OnActiveTabChangedDelegateHandle = FGlobalTabmanager::Get()->OnActiveTabChanged_Subscribe( FOnActiveTabChanged::FDelegate::CreateRaw( this, &FDocumentTracker::OnActiveTabChanged ) );
 	TabForegroundedDelegateHandle = FGlobalTabmanager::Get()->OnTabForegrounded_Subscribe(FOnActiveTabChanged::FDelegate::CreateRaw(this, &FDocumentTracker::OnTabForegrounded));
+
+	CurrentHistoryIndex = INDEX_NONE;
 }
 
 FDocumentTracker::~FDocumentTracker()
@@ -410,7 +440,7 @@ void FDocumentTracker::OnActiveTabChanged(TSharedPtr<SDockTab> PreviouslyActive,
 			TSharedPtr<SDockTab> Tab = (*ListIt)->GetTab().Pin();
 			if (Tab == NewlyActivated)
 			{
-				LastEditedTabInfo = *ListIt;
+				WeakLastEditedTabInfo = *ListIt;
 				Factory->OnTabActivated(Tab);
 			}
 		}
@@ -419,19 +449,21 @@ void FDocumentTracker::OnActiveTabChanged(TSharedPtr<SDockTab> PreviouslyActive,
 
 void FDocumentTracker::OnTabForegrounded(TSharedPtr<SDockTab> ForegroundedTab, TSharedPtr<SDockTab> BackgroundedTab)
 {
+	TSharedPtr<FTabInfo> NewTabInfo;
 	TSharedPtr<SDockTab> OwnedForeground, OwnedBackground;
 	TSharedPtr<FDocumentTabFactory> ForegroundFactory, BackgroundFactory;
 
 	FTabList& List = GetSpawnedList();
-	for ( auto ListIt = List.CreateIterator(); ListIt; ++ListIt )
+	for ( const TSharedPtr<FTabInfo> TabInfo : List )
 	{
 		// Get the factory (can't fail; the tabs had to come from somewhere; failing means a tab survived a mode transition to a mode where it is not allowed!)
-		TSharedPtr<FDocumentTabFactory> Factory = ( *ListIt )->GetFactory().Pin();
+		TSharedPtr<FDocumentTabFactory> Factory = TabInfo->GetFactory().Pin();
 		if ( ensure(Factory.IsValid()) )
 		{
-			TSharedPtr<SDockTab> Tab = ( *ListIt )->GetTab().Pin();
+			TSharedPtr<SDockTab> Tab = TabInfo->GetTab().Pin();
 			if ( Tab == ForegroundedTab )
 			{
+				NewTabInfo = TabInfo;
 				OwnedForeground = Tab;
 				ForegroundFactory = Factory;
 			}
@@ -452,6 +484,14 @@ void FDocumentTracker::OnTabForegrounded(TSharedPtr<SDockTab> ForegroundedTab, T
 	{
 		ForegroundFactory->OnTabForegrounded(OwnedForeground);
 	}
+
+	TSharedPtr<FGenericTabHistory> CurrentTabHistory = GetCurrentTabHistory();
+
+	if ( ForegroundFactory.IsValid() && NewTabInfo.IsValid() && NewTabInfo->GetCurrentHistory() != CurrentTabHistory )
+	{
+		// If a tab was manually foregrounded, need to add tab history
+		NewTabInfo->AddTabHistory(ForegroundFactory->CreateTabHistoryNode(NewTabInfo->GetPayload()), true);
+	}
 }
 
 FDocumentTracker::FTabList& FDocumentTracker::GetSpawnedList()
@@ -472,6 +512,18 @@ FDocumentTracker::FTabList& FDocumentTracker::GetSpawnedList()
 	return SpawnedTabs;
 }
 
+TSharedPtr<FTabInfo> FDocumentTracker::GetLastEditedTabInfo()
+{
+	TSharedPtr<FTabInfo> LastEditedTabInfo = WeakLastEditedTabInfo.Pin();
+	if (!LastEditedTabInfo.IsValid() || !LastEditedTabInfo->GetTab().IsValid())
+	{
+		// Clear if either the info or actual tab are gone
+		WeakLastEditedTabInfo = nullptr;
+		return nullptr;
+	}
+	return LastEditedTabInfo;
+}
+
 void FDocumentTracker::Initialize(TSharedPtr<FAssetEditorToolkit> InHostingApp )
 {
 	check(!HostingAppPtr.IsValid());
@@ -490,59 +542,53 @@ TSharedPtr<SDockTab> FDocumentTracker::OpenDocument(TSharedPtr<FTabPayload> InPa
 		InOpenCause = FDocumentTracker::ForceOpenNewDocument;
 	}
 
-	if(InOpenCause == NavigatingCurrentDocument || InOpenCause == QuickNavigateCurrentDocument)
+	if(InOpenCause == CreateHistoryEvent)
+	{
+		// Deprecated, all opens now save history
+		InOpenCause = FDocumentTracker::OpenNewDocument;
+	}
+
+	if(InOpenCause == NavigatingCurrentDocument || InOpenCause == QuickNavigateCurrentDocument || InOpenCause == NavigateBackwards || InOpenCause == NavigateForwards)
 	{
 		return NavigateCurrentTab(InPayload, InOpenCause);
 	}
-	else if(InOpenCause == OpenNewDocument || InOpenCause == CreateHistoryEvent)
+
+	// Spawning or restoring a tab, so the factory has to be valid
+	TSharedPtr<FDocumentTabFactory> Factory = FindSupportingFactory(InPayload.ToSharedRef());
+	if(!Factory.IsValid())
 	{
-		TSharedPtr< FTabInfo > TabInfo;
-		TSharedPtr<SDockTab> Tab;
-		bool bPrevTabMatches = false;
+		return nullptr;
+	}
 
+	if(InOpenCause == OpenNewDocument || InOpenCause == RestorePreviousDocument)
+	{
 		// If the current tab matches we'll re-use it.
-		TSharedPtr<FTabInfo> LastEditedTabInfoPinned = LastEditedTabInfo.Pin();
-		if(LastEditedTabInfoPinned.IsValid() && LastEditedTabInfoPinned->PayloadMatches(InPayload))
+		TSharedPtr<FTabInfo> LastEditedTabInfo = GetLastEditedTabInfo();
+		if(LastEditedTabInfo.IsValid() && LastEditedTabInfo->PayloadMatches(InPayload))
 		{
-			TabInfo = LastEditedTabInfoPinned;
-			Tab = TabInfo->GetTab().Pin();
-			bPrevTabMatches = true;
+			LastEditedTabInfo->AddTabHistory(Factory->CreateTabHistoryNode(InPayload), true);
+			return LastEditedTabInfo->GetTab().Pin();
 		}
 		else
 		{
 			// Check if the payload is currently open in any tab
 			FTabList& List = GetSpawnedList();
-			for (auto ListIt = List.CreateIterator(); ListIt; ++ListIt)
+			for (const TSharedPtr<FTabInfo>& TabInfo : List)
 			{
-				if ((*ListIt)->PayloadMatches(InPayload))
+				TSharedPtr<SDockTab> Tab = TabInfo->GetTab().Pin();
+				if (TabInfo->PayloadMatches(InPayload))
 				{
-					TabInfo = (*ListIt);
-					Tab = TabInfo->GetTab().Pin();
-					TabManager->DrawAttention( TabInfo->GetTab().Pin().ToSharedRef() );
+					// Manually opening an existing tab, add to history
+					TabInfo->AddTabHistory(Factory->CreateTabHistoryNode(InPayload), true);
+					TabManager->DrawAttention(Tab.ToSharedRef());
+					return Tab;
 				}
 			}
 		}
-		
-		if(!TabInfo.IsValid())
-		{
-			// If no tab was found with the payload
-			Tab = OpenNewTab(InPayload, OpenNewDocument);
-		}
-		else if(InOpenCause == CreateHistoryEvent)
-		{
-			TSharedPtr<FDocumentTabFactory> Factory = FindSupportingFactory(InPayload.ToSharedRef());
-			TabInfo->AddTabHistory(Factory->CreateTabHistoryNode(InPayload), true, true);
-		}
-		
-		return Tab;
-	}
-	else if(InOpenCause == NavigateBackwards || InOpenCause == NavigateForwards)
-	{
-		return NavigateCurrentTab(InPayload, InOpenCause);
 	}
 
-	// Occurs only when forcing open a new tab.
-	return OpenNewTab(InPayload, InOpenCause);
+	// Occurs when forcing open a new tab, or if it failed to find existing one
+	return OpenNewTab(Factory->CreateTabHistoryNode(InPayload), InOpenCause);
 }
 
 TWeakPtr< FTabInfo > FDocumentTracker::FindTabInForeground()
@@ -557,7 +603,68 @@ TWeakPtr< FTabInfo > FDocumentTracker::FindTabInForeground()
 		}
 	}
 
-	return NULL;
+	return nullptr;
+}
+
+bool FDocumentTracker::NavigateToTabHistory(int32 InHistoryIdx)
+{
+	if (History.IsValidIndex(InHistoryIdx))
+	{
+		TSharedPtr<FGenericTabHistory> NewHistory = History[InHistoryIdx];
+		if (NewHistory->IsHistoryValid())
+		{
+			CurrentHistoryIndex = InHistoryIdx;
+
+			TSharedPtr<FTabInfo> FoundPayloadTab, FoundBoundTab;
+			FTabList& List = GetSpawnedList();
+			for (TSharedPtr<FTabInfo>& TabInfo : List)
+			{
+				if (TabInfo->GetCurrentHistory() == NewHistory)
+				{
+					// If it's literally the same history item that is already active, just focus the tab
+					TabManager->DrawAttention(TabInfo->GetTab().Pin().ToSharedRef());
+					return true;
+				}
+				else if (TabInfo->GetCurrentHistory()->GetBoundTab() == NewHistory->GetBoundTab())
+				{
+					FoundBoundTab = TabInfo;
+				}
+				else if (TabInfo->PayloadMatches(NewHistory->GetPayload()))
+				{
+					FoundPayloadTab = TabInfo;
+				}
+			}
+
+			// If no bound tab but found matching payload, use that
+			if (!FoundBoundTab.IsValid() && FoundPayloadTab.IsValid())
+			{
+				FoundBoundTab = FoundPayloadTab;
+			}
+
+			if (FoundBoundTab.IsValid())
+			{
+				// Found a mostly matching tab we want to restore in
+				FoundBoundTab->SetCurrentHistory(NewHistory, true, true);
+				TabManager->DrawAttention(FoundBoundTab->GetTab().Pin().ToSharedRef());
+				return true;
+			}
+
+			// The tab was closed, so spawn a new one
+			OpenNewTab(NewHistory, EOpenDocumentCause::NavigatingHistory);
+			return true;
+		}
+	}
+
+	return false;
+}
+
+TSharedPtr<FGenericTabHistory> FDocumentTracker::GetCurrentTabHistory()
+{
+	if (History.IsValidIndex(CurrentHistoryIndex))
+	{
+		return History[CurrentHistoryIndex];
+	}
+	return nullptr;
 }
 
 TSharedPtr<SDockTab> FDocumentTracker::NavigateCurrentTab(TSharedPtr<FTabPayload> InPayload, EOpenDocumentCause InNavigateCause)
@@ -568,67 +675,85 @@ TSharedPtr<SDockTab> FDocumentTracker::NavigateCurrentTab(TSharedPtr<FTabPayload
 	if(List.Num())
 	{
 		// Make sure we find an available tab to navigate, there are ones available.
+		TSharedPtr<FTabInfo> LastEditedTabInfo = GetLastEditedTabInfo();
 		if(!LastEditedTabInfo.IsValid())
 		{
-			LastEditedTabInfo = FindTabInForeground();
+			WeakLastEditedTabInfo = FindTabInForeground();
 
 			// Check if we still do not have a valid tab, if we do not, activate the first tab
-			if(!LastEditedTabInfo.IsValid())
+			// Any invalid tabs would have been deleted inside GetSpawnedList so we know 0 is valid
+			if(!WeakLastEditedTabInfo.IsValid())
 			{
-				LastEditedTabInfo = List[0];
-				LastEditedTabInfo.Pin()->GetTab().Pin()->ActivateInParent(ETabActivationCause::SetDirectly);
+				WeakLastEditedTabInfo = List[0];
+				List[0]->GetTab().Pin()->ActivateInParent(ETabActivationCause::SetDirectly);
 			}
+
+			LastEditedTabInfo = GetLastEditedTabInfo();
 		}
 
 		check(LastEditedTabInfo.IsValid());
-		TSharedPtr<FAssetEditorToolkit> HostingApp = HostingAppPtr.Pin();
-		FWorkflowTabSpawnInfo SpawnInfo;
-		SpawnInfo.Payload = InPayload;
-		SpawnInfo.TabInfo = LastEditedTabInfo.Pin();
+		TSharedPtr<SDockTab> LastEditedTab = LastEditedTabInfo->GetTab().Pin();
 
 		if(InNavigateCause == NavigatingCurrentDocument || InNavigateCause == QuickNavigateCurrentDocument)
 		{
 			TSharedPtr<FDocumentTabFactory> Factory = FindSupportingFactory(InPayload.ToSharedRef());
-			// If doing a Quick navigate of the document, do not save history data as it's likely still at the default values. The object is always saved
-			LastEditedTabInfo.Pin()->AddTabHistory(Factory->CreateTabHistoryNode(InPayload), InNavigateCause != QuickNavigateCurrentDocument);
-			// Ensure that the tab appears if the tab isn't currently in the foreground.
-			LastEditedTabInfo.Pin()->GetTab().Pin()->ActivateInParent(ETabActivationCause::SetDirectly);
+			if (Factory.IsValid())
+			{
+				// If doing a Quick navigate of the document, do not save history data as it's likely still at the default values. The object is always saved
+				LastEditedTabInfo->AddTabHistory(Factory->CreateTabHistoryNode(InPayload), InNavigateCause != QuickNavigateCurrentDocument);
+				// Ensure that the tab appears if the tab isn't currently in the foreground.
+				LastEditedTab->ActivateInParent(ETabActivationCause::SetDirectly);
+			}
 		}
 		else if(InNavigateCause == NavigateBackwards)
 		{
-			LastEditedTabInfo.Pin()->OnGoBackInHistory();
+			LastEditedTabInfo->OnGoBackInHistory();
 		}
 		else if(InNavigateCause == NavigateForwards)
 		{
-			LastEditedTabInfo.Pin()->OnGoForwardInHistory();
+			LastEditedTabInfo->OnGoForwardInHistory();
 		}
 
-		return LastEditedTabInfo.Pin()->GetTab().Pin();
+		return LastEditedTab;
 	}
 	
 	// Open in new tab
-	return OpenNewTab(InPayload, OpenNewDocument);
+	TSharedPtr<FDocumentTabFactory> Factory = FindSupportingFactory(InPayload.ToSharedRef());
+	if (Factory.IsValid())
+	{
+		return OpenNewTab(Factory->CreateTabHistoryNode(InPayload), OpenNewDocument);
+	}
+	
+	return nullptr;
 }
 
-TSharedPtr<SDockTab> FDocumentTracker::OpenNewTab(TSharedPtr<FTabPayload> InPayload, EOpenDocumentCause InOpenCause)
+TSharedPtr<SDockTab> FDocumentTracker::OpenNewTab(TSharedPtr<FGenericTabHistory> InTabHistory, EOpenDocumentCause InOpenCause)
 {
-	ensure(InOpenCause == ForceOpenNewDocument ||InOpenCause == OpenNewDocument || InOpenCause == RestorePreviousDocument);
-
-	TSharedPtr<FDocumentTabFactory> Factory = FindSupportingFactory(InPayload.ToSharedRef());
+	ensure(InOpenCause == ForceOpenNewDocument || InOpenCause == OpenNewDocument || InOpenCause == RestorePreviousDocument || InOpenCause == NavigatingHistory);
 
+	TSharedPtr<FDocumentTabFactory> Factory = InTabHistory->GetFactory().Pin();
 	TSharedPtr<SDockTab> NewTab;
 
 	if(Factory.IsValid())
 	{
 		TSharedPtr<FAssetEditorToolkit> HostingApp = HostingAppPtr.Pin();
 		FWorkflowTabSpawnInfo SpawnInfo;
-		SpawnInfo.Payload = InPayload;
+		SpawnInfo.Payload = InTabHistory->GetPayload();
 
 		NewTab = Factory->SpawnBlankTab();
 
-		TSharedPtr<FTabInfo> NewTabInfo = MakeShareable( new FTabInfo(NewTab.ToSharedRef(), Factory) );
+		TSharedPtr<FTabInfo> NewTabInfo = MakeShareable( new FTabInfo(NewTab.ToSharedRef(), Factory, AsShared()) );
 		SpawnedTabs.Add( NewTabInfo );
-		NewTabInfo->AddTabHistory(Factory->CreateTabHistoryNode(InPayload));
+
+		if (InOpenCause == NavigatingHistory)
+		{
+			// Don't want to create a new history entry when going back/forward
+			NewTabInfo->SetCurrentHistory(InTabHistory, true, true);
+		}
+		else
+		{
+			NewTabInfo->AddTabHistory(InTabHistory);
+		}
 
 		if (InOpenCause == ForceOpenNewDocument  || InOpenCause == OpenNewDocument)
 		{
@@ -637,6 +762,10 @@ TSharedPtr<SDockTab> FDocumentTracker::OpenNewTab(TSharedPtr<FTabPayload> InPayl
 		else if (InOpenCause == RestorePreviousDocument)
 		{
 			TabManager->RestoreDocumentTab( "Document", FTabManager::ESearchPreference::RequireClosedTab, NewTab.ToSharedRef() );
+
+			// Clear tab history before this so previous restores don't show up
+			History.RemoveAt(0, History.Num() - 1, true);
+			CurrentHistoryIndex = History.Num() - 1;
 		}
 	}
 
@@ -793,16 +922,16 @@ TArray< TSharedPtr<SDockTab> > FDocumentTracker::GetAllDocumentTabs() const
 
 TSharedPtr<SDockTab> FDocumentTracker::GetActiveTab() const
 {
-	if(LastEditedTabInfo.IsValid())
+	if(WeakLastEditedTabInfo.IsValid())
 	{
-		return LastEditedTabInfo.Pin()->GetTab().Pin();
+		return WeakLastEditedTabInfo.Pin()->GetTab().Pin();
 	}
 
-	return NULL;
+	return nullptr;
 }
 
 FReply FDocumentTracker::OnNavigateTab(FDocumentTracker::EOpenDocumentCause InCause)
 {
-	NavigateCurrentTab(FTabPayload_UObject::Make(NULL), InCause);
+	NavigateCurrentTab(FTabPayload_UObject::Make(nullptr), InCause);
 	return FReply::Handled();
 }
diff --git a/Engine/Source/Editor/Kismet/Public/WorkflowOrientedApp/WorkflowTabManager.h b/Engine/Source/Editor/Kismet/Public/WorkflowOrientedApp/WorkflowTabManager.h
--- a/Engine/Source/Editor/Kismet/Public/WorkflowOrientedApp/WorkflowTabManager.h
+++ b/Engine/Source/Editor/Kismet/Public/WorkflowOrientedApp/WorkflowTabManager.h
@@ -60,7 +60,7 @@ private:
 class FTabInfo : public TSharedFromThis<FTabInfo>
 {
 public:
-	FTabInfo(const TSharedRef<SDockTab>& InTab, const TSharedPtr<FDocumentTabFactory>& InSpawner);
+	FTabInfo(const TSharedRef<SDockTab>& InTab, const TSharedPtr<FDocumentTabFactory>& InSpawner, const TSharedPtr<class FDocumentTracker>& InTracker);
 
 	/** Returns TRUE if the payload used by this TabInfo is the same as passed in */
 	bool PayloadMatches(const TSharedPtr<FTabPayload> TestPayload) const;
@@ -73,7 +73,13 @@ public:
 	 * @param InHistoryNode		The history node to add
 	 * @param bInSaveHistory	TRUE if history should be saved
 	 */
-	void AddTabHistory(TSharedPtr< struct FGenericTabHistory > InHistoryNode, bool bInSaveHistory = true, bool bPrevTabMatches = false);
+	void AddTabHistory(TSharedPtr< struct FGenericTabHistory > InHistoryNode, bool bInSaveHistory = true);
+
+	/** Retrieves the history currently visible in this tab */
+	TSharedPtr<struct FGenericTabHistory> GetCurrentHistory() const;
+
+	/** Sets the current history and evoke it, this can either be an old or new history node */
+	void SetCurrentHistory(TSharedPtr<struct FGenericTabHistory> NewHistory, bool bInSaveHistory = true, bool bShouldRestore = true);
 
 	/** Single step forward in history */
 	FReply OnGoForwardInHistory();
@@ -114,7 +120,7 @@ private:
 	 *
 	 * @param InMenuAnchor		This is the anchor the menu will use for positioning
 	 */
-	FReply OnMouseDownHisory( const FGeometry& MyGeometry, const FPointerEvent& MouseEvent, TWeakPtr< SMenuAnchor > InMenuAnchor );
+	FReply OnMouseDownHistory( const FGeometry& MyGeometry, const FPointerEvent& MouseEvent, TWeakPtr< SMenuAnchor > InMenuAnchor );
 
 	/** 
 	 * Callback to create the history menu.
@@ -130,10 +136,10 @@ private:
 	TSharedPtr<SWidget> HistoryNavigationWidget;
 	/** Tab this info represents */
 	TWeakPtr<class SDockTab> Tab;
-	/** List of history data nodes */
-	TArray< TSharedPtr< struct FGenericTabHistory > > History;
-	/** Current history index */
-	int32 CurrentHistoryIndex;
+	/** Current tab history element in this tab, may not necessarily be in the global history */
+	TSharedPtr<struct FGenericTabHistory> CurrentHistory;
+	/** Pointer to document tracker that spawned us */
+	TWeakPtr<class FDocumentTracker> WeakTracker;
 };
 
 /////////////////////////////////////////////////////
@@ -209,12 +215,27 @@ public:
 		return FactoryPtr.Pin()->ConstructTabName(SpawnInfo);
 	}
 
+	/** Returns the tab this is bound to */
+	TSharedPtr<FTabInfo> GetBoundTab() const
+	{
+		return BoundTab.Pin();
+	}
+
+	/** Binds this to a specific tab */
+	void BindToTab(const TSharedPtr<FTabInfo>& InTab)
+	{
+		BoundTab = InTab;
+	}
+
 protected:
 	/** The factory used to generate widgets for the payload */
 	TWeakPtr<class FDocumentTabFactory> FactoryPtr;
 
 	/** The payload this history node represents */
 	TSharedPtr<FTabPayload> Payload;
+
+	/** Specific tab this is bound to, may be null */
+	TWeakPtr<FTabInfo> BoundTab;
 };
 
 /////////////////////////////////////////////////////
@@ -239,14 +260,25 @@ public:
 	/** Are we opening a new document, or restoring a previously opened document */
 	enum EOpenDocumentCause
 	{
+		// Opens a new document, will open in existing tab if one is found
 		OpenNewDocument,
+		// Always open in a new tab
 		ForceOpenNewDocument,
+		// Restoring tab state from a previous session
 		RestorePreviousDocument,
+		// Navigating within an open tab
 		NavigatingCurrentDocument,
+		// Navigate without saving current document position
 		QuickNavigateCurrentDocument,
+		// Key/Mouse binding for back, turns into NavigatingHistory
 		NavigateBackwards,
+		// Key/Mouse binding for forward
 		NavigateForwards,
-		CreateHistoryEvent
+		// Going forward or back in the history
+		NavigatingHistory,
+
+		// Deprecated, treated the same as OpenNewDocument
+		CreateHistoryEvent,
 	};
 
 	/**
@@ -306,7 +338,7 @@ private:
 	TSharedPtr<FTabManager> TabManager;
 
 	// The last active tab info to be selected, used for opening new documents in the most recently used tab
-	TWeakPtr<FTabInfo> LastEditedTabInfo;
+	TWeakPtr<FTabInfo> WeakLastEditedTabInfo;
 
 	/** Handle to the registered OnActiveTabChanged delegate */
 	FDelegateHandle OnActiveTabChangedDelegateHandle;
@@ -314,10 +346,19 @@ private:
 	/** Handle to the registered OnTabForegrounded delegate */
 	FDelegateHandle TabForegroundedDelegateHandle;
 
+	/** List of history data nodes */
+	TArray< TSharedPtr<FGenericTabHistory> > History;
+
+	/** Current history index */
+	int32 CurrentHistoryIndex;
+
 private:
 	// Clean the spawned list 
 	FTabList& GetSpawnedList();
 
+	// Cleans the LastEditedTabInfo and returns pinned
+	TSharedPtr<FTabInfo> GetLastEditedTabInfo();
+
 	// Called by the global active tab changed callback; dispatches to individually registered callbacks
 	void OnActiveTabChanged(TSharedPtr<SDockTab> PreviouslyActive, TSharedPtr<SDockTab> NewlyActivated);
 
@@ -330,6 +371,12 @@ private:
 	/** Helper function to find the tab in the foreground */
 	TWeakPtr< FTabInfo > FindTabInForeground();
 
+	/** Navigates to specific index in tab history, returns success if it was found */
+	bool NavigateToTabHistory(int32 InHistoryIdx);
+
+	/** Returns the currently active tab history, may be null */
+	TSharedPtr<FGenericTabHistory> GetCurrentTabHistory();
+
 	/**
 	 * Navigates current tab
 	 *
@@ -343,11 +390,13 @@ private:
 	/**
 	 * Opens a new tab
 	 *
-	 * @param InPayload					Payload to use for the tab
+	 * @param InTabHistory				New or existing tab history to open tab to
 	 * @param InOpenCause				The cause of the tab being opened
 	 *
 	 * @return							The tab opened and modified
 	 */
-	TSharedPtr<SDockTab> OpenNewTab(TSharedPtr<FTabPayload> InPayload, EOpenDocumentCause InOpenCause);
+	TSharedPtr<SDockTab> OpenNewTab(TSharedPtr<FGenericTabHistory> InTabHistory, EOpenDocumentCause InOpenCause);
+
+	friend class FTabInfo;
 };
 
