CIS fix (Linux).

#rnx

diff --git a/Engine/Source/Editor/Kismet/Private/WorkflowOrientedApp/WorkflowTabManager.cpp b/Engine/Source/Editor/Kismet/Private/WorkflowOrientedApp/WorkflowTabManager.cpp
--- a/Engine/Source/Editor/Kismet/Private/WorkflowOrientedApp/WorkflowTabManager.cpp
+++ b/Engine/Source/Editor/Kismet/Private/WorkflowOrientedApp/WorkflowTabManager.cpp
@@ -214,7 +214,7 @@ TSharedRef<SWidget> FTabInfo::CreateHistoryMenu(bool bInBackHistory) const
 			{
 				MenuBuilder.AddMenuEntry(HistoryItem->GetHistoryTitle().Get(), FText(), FSlateIcon(), 
 					FUIAction(
-					FExecuteAction::CreateRaw(this, &FTabInfo::GoToHistoryIndex, HistoryIdx)
+					FExecuteAction::CreateRaw(const_cast<FTabInfo*>(this), &FTabInfo::GoToHistoryIndex, HistoryIdx)
 					), 
 					NAME_None, EUserInterfaceActionType::Button);
 			}
@@ -232,7 +232,7 @@ TSharedRef<SWidget> FTabInfo::CreateHistoryMenu(bool bInBackHistory) const
 			{
 				MenuBuilder.AddMenuEntry(HistoryItem->GetHistoryTitle().Get(), FText(), FSlateIcon(), 
 					FUIAction(
-					FExecuteAction::CreateRaw(this, &FTabInfo::GoToHistoryIndex, HistoryIdx)
+					FExecuteAction::CreateRaw(const_cast<FTabInfo*>(this), &FTabInfo::GoToHistoryIndex, HistoryIdx)
 					), 
 					NAME_None, EUserInterfaceActionType::Button);
 			}
