Optimize GetDependentBlueprints by caching dependencies. The goal here is to save minutes of cook time in large projects. This should also speed up overall load times for blueprints when large numbers of blueprints are already in memory

diff --git a/Engine/Source/Editor/Kismet/Private/BlueprintCompilationManager.cpp b/Engine/Source/Editor/Kismet/Private/BlueprintCompilationManager.cpp
--- a/Engine/Source/Editor/Kismet/Private/BlueprintCompilationManager.cpp
+++ b/Engine/Source/Editor/Kismet/Private/BlueprintCompilationManager.cpp
@@ -515,17 +515,33 @@ void FBlueprintCompilationManagerImpl::FlushCompilationQueueImpl(bool bSuppressB
 		for(const FBPCompileRequestInternal& CompileJob : QueuedRequests)
 		{
 			if ((CompileJob.UserData.CompileOptions & 
-				(	EBlueprintCompileOptions::RegenerateSkeletonOnly|
-					EBlueprintCompileOptions::IsRegeneratingOnLoad)
+				(	EBlueprintCompileOptions::RegenerateSkeletonOnly)
+				) != EBlueprintCompileOptions::None)
+			{
+				continue;
+			}
+
+			UBlueprint* BP = CompileJob.UserData.BPToCompile;
+
+			if(!BP->bHasBeenRegenerated && BP->GetLinker())
+			{
+				// we may have cached dependencies before being fully loaded:
+				BP->bCachedDependenciesUpToDate = false;
+			}
+
+			FBlueprintEditorUtils::EnsureCachedDependenciesUpToDate(BP);
+
+			if ((CompileJob.UserData.CompileOptions & 
+				(	EBlueprintCompileOptions::IsRegeneratingOnLoad)
 				) != EBlueprintCompileOptions::None)
 			{
 				continue;
 			}
 			
-			if(CompileJob.UserData.BPToCompile->BlueprintType == BPTYPE_MacroLibrary)
+			if(BP->BlueprintType == BPTYPE_MacroLibrary)
 			{
 				TArray<UBlueprint*> DependentBlueprints;
-				FBlueprintEditorUtils::GetDependentBlueprints(CompileJob.UserData.BPToCompile, DependentBlueprints);
+				FBlueprintEditorUtils::GetDependentBlueprints(BP, DependentBlueprints);
 				for(UBlueprint* DependentBlueprint : DependentBlueprints)
 				{
 					if(!IsQueuedForCompilation(DependentBlueprint))
@@ -751,11 +767,6 @@ void FBlueprintCompilationManagerImpl::FlushCompilationQueueImpl(bool bSuppressB
 			BP->bBeingCompiled = true;
 			BP->CurrentMessageLog = CompilerData.ActiveResultsLog;
 			BP->bIsRegeneratingOnLoad = !BP->bHasBeenRegenerated && BP->GetLinker();
-			if(BP->bIsRegeneratingOnLoad)
-			{
-				// we may have cached dependencies before being fully loaded:
-				BP->bCachedDependenciesUpToDate = false;
-			}
 
 			if(CompilerData.ShouldResetErrorState())
 			{
@@ -1267,6 +1278,8 @@ void FBlueprintCompilationManagerImpl::FlushCompilationQueueImpl(bool bSuppressB
 				BPGC->ClassFlags &= ~CLASS_ReplicationDataIsSetUp;
 				BPGC->SetUpRuntimeReplicationData();
 			}
+			
+			FKismetCompilerUtilities::UpdateDependentBlueprints(BP);
 
 			ensure(BPGC == nullptr || BPGC->ClassDefaultObject->GetClass() == BPGC);
 		}
diff --git a/Engine/Source/Editor/KismetCompiler/Private/KismetCompilerMisc.cpp b/Engine/Source/Editor/KismetCompiler/Private/KismetCompilerMisc.cpp
--- a/Engine/Source/Editor/KismetCompiler/Private/KismetCompilerMisc.cpp
+++ b/Engine/Source/Editor/KismetCompiler/Private/KismetCompilerMisc.cpp
@@ -1612,6 +1612,40 @@ bool FKismetCompilerUtilities::IsIntermediateFunctionGraphTrivial(FName Function
 	return false;
 }
 
+void FKismetCompilerUtilities::UpdateDependentBlueprints(UBlueprint* ForBP)
+{
+	for(TWeakObjectPtr<UBlueprint> Dependency : ForBP->CachedDependencies)
+	{
+		if(UBlueprint* BP = Dependency.Get())
+		{
+			if(BP != ForBP) // avoid tautology
+			{
+				BP->CachedDependents.Add(ForBP);
+			}
+		}
+	}
+
+	// CachedDependencies may not include all function calls, e.g. because we're 
+	// calling a blueprint function library function via a macro (such a dependency
+	// cannot be detected until after graph expansion):
+	if(UBlueprintGeneratedClass* BPGC = Cast<UBlueprintGeneratedClass>(ForBP->GeneratedClass))
+	{
+		for(UFunction* Fn : BPGC->CalledFunctions)
+		{
+			if(UBlueprintGeneratedClass* OwningBPGC = Cast<UBlueprintGeneratedClass>(Fn->GetOwnerClass()))
+			{
+				if(UBlueprint* OwningBP = Cast<UBlueprint>(OwningBPGC->ClassGeneratedBy))
+				{
+					if(OwningBP != ForBP) // avoid tautology
+					{
+						OwningBP->CachedDependents.Add(ForBP);
+					}
+				}
+			}
+		}
+	}
+}
+
 //////////////////////////////////////////////////////////////////////////
 // FNodeHandlingFunctor
 
diff --git a/Engine/Source/Editor/KismetCompiler/Public/KismetCompilerMisc.h b/Engine/Source/Editor/KismetCompiler/Public/KismetCompilerMisc.h
--- a/Engine/Source/Editor/KismetCompiler/Public/KismetCompilerMisc.h
+++ b/Engine/Source/Editor/KismetCompiler/Public/KismetCompilerMisc.h
@@ -117,6 +117,9 @@ public:
 
 	/** @return true if the graph in question contains only an entry node or only an entry node and a call to its parent if the graph is an override */
 	static bool IsIntermediateFunctionGraphTrivial(FName FunctionName, const UEdGraph* FunctionGraph);
+
+	/** Add this BP to any BPs that it in*/
+	static void UpdateDependentBlueprints(UBlueprint* BP);
 };
 
 //////////////////////////////////////////////////////////////////////////
diff --git a/Engine/Source/Editor/UnrealEd/Private/Kismet2/BlueprintEditorUtils.cpp b/Engine/Source/Editor/UnrealEd/Private/Kismet2/BlueprintEditorUtils.cpp
--- a/Engine/Source/Editor/UnrealEd/Private/Kismet2/BlueprintEditorUtils.cpp
+++ b/Engine/Source/Editor/UnrealEd/Private/Kismet2/BlueprintEditorUtils.cpp
@@ -2944,41 +2944,15 @@ void FBlueprintEditorUtils::EnsureCachedDependenciesUpToDate(UBlueprint* Bluepri
 	}
 }
 
-void FBlueprintEditorUtils::GetDependentBlueprints(UBlueprint* Blueprint, TArray<UBlueprint*>& DependentBlueprints, bool bRemoveSelf)
+void FBlueprintEditorUtils::GetDependentBlueprints(UBlueprint* Blueprint, TArray<UBlueprint*>& DependentBlueprints)
 {
-	TArray<UObject*> AllBlueprints;
-	bool const bIncludeDerivedClasses = true;
-	GetObjectsOfClass(UBlueprint::StaticClass(), AllBlueprints, bIncludeDerivedClasses );
-
-	for (UObject* Obj : AllBlueprints)
+	for( TWeakObjectPtr<UBlueprint> DependentBPWeak : Blueprint->CachedDependents )
 	{
-		// we know the class is correct so a fast cast is ok here
-		UBlueprint* TestBP = static_cast<UBlueprint*>(Obj);
-		if (TestBP && !TestBP->IsPendingKill())
+		if(UBlueprint* DependentBP = DependentBPWeak.Get())
 		{
-			EnsureCachedDependenciesUpToDate(TestBP);
-
-			if (TestBP->CachedDependencies.Contains(Blueprint))
-			{
-				if (!DependentBlueprints.Contains(TestBP))
-				{
-					DependentBlueprints.Add(TestBP);
-
-					// When a Macro Library depends on this Blueprint, then any Blueprint that 
-					// depends on it must also depend on this Blueprint for re-compiling (bytecode, skeleton, full) purposes
-					if (TestBP->BlueprintType == BPTYPE_MacroLibrary)
-					{
-						GetDependentBlueprints(TestBP, DependentBlueprints, false);
-					}
-				}
-			}
+			DependentBlueprints.Add(DependentBP);
 		}
 	}
-
-	if (bRemoveSelf)
-	{
-		DependentBlueprints.RemoveSwap(Blueprint);
-	}
 }
 
 
diff --git a/Engine/Source/Editor/UnrealEd/Public/Kismet2/BlueprintEditorUtils.h b/Engine/Source/Editor/UnrealEd/Public/Kismet2/BlueprintEditorUtils.h
--- a/Engine/Source/Editor/UnrealEd/Public/Kismet2/BlueprintEditorUtils.h
+++ b/Engine/Source/Editor/UnrealEd/Public/Kismet2/BlueprintEditorUtils.h
@@ -524,7 +524,7 @@ public:
 	static void GatherDependencies(const UBlueprint* Blueprint, TSet<TWeakObjectPtr<UBlueprint>>& OutDependencies, TSet<TWeakObjectPtr<UStruct>>& OutUDSDependencies);
 
 	/** Returns a list of loaded Blueprints that are dependent on the given Blueprint. */
-	static void GetDependentBlueprints(UBlueprint* Blueprint, TArray<UBlueprint*>& DependentBlueprints, bool bRemoveSelf = true);
+	static void GetDependentBlueprints(UBlueprint* Blueprint, TArray<UBlueprint*>& DependentBlueprints);
 
 	/** Ensures, that CachedDependencies in BP are up to date */
 	static void EnsureCachedDependenciesUpToDate(UBlueprint* Blueprint);
diff --git a/Engine/Source/Runtime/Engine/Classes/Engine/Blueprint.h b/Engine/Source/Runtime/Engine/Classes/Engine/Blueprint.h
--- a/Engine/Source/Runtime/Engine/Classes/Engine/Blueprint.h
+++ b/Engine/Source/Runtime/Engine/Classes/Engine/Blueprint.h
@@ -697,8 +697,19 @@ public:
 	uint32 CrcLastCompiledSignature;
 
 	bool bCachedDependenciesUpToDate;
+	/**
+	 * Set of blueprints that we reference - i.e. blueprints that we have
+	 * some kind of reference to (variable of that blueprints type or function 
+	 * call
+	 */
 	TSet<TWeakObjectPtr<UBlueprint>> CachedDependencies;
 
+	/** 
+	 * Transient cache of dependent blueprints - i.e. blueprints that call
+	 * functions declared in this blueprint. Used to speed up compilation checks 
+	 */
+	TSet<TWeakObjectPtr<UBlueprint>> CachedDependents;
+
 	// User Defined Structures, the blueprint depends on
 	TSet<TWeakObjectPtr<UStruct>> CachedUDSDependencies;
 
