PR #6498: Remove inconsistent tooltip from the interface list on BP editor (Contributed by [REDACTED_PERSON])

diff --git a/Engine/Source/Editor/Kismet/Private/SBlueprintPalette.cpp b/Engine/Source/Editor/Kismet/Private/SBlueprintPalette.cpp
--- a/Engine/Source/Editor/Kismet/Private/SBlueprintPalette.cpp
+++ b/Engine/Source/Editor/Kismet/Private/SBlueprintPalette.cpp
@@ -242,8 +242,7 @@ static void GetSubGraphIcon(FEdGraphSchemaAction_K2Graph const* const ActionIn,
 			{
 				if (UEdGraphSchema_K2::FunctionCanBePlacedAsEvent(OverrideFunc))
 				{
-					Args.Add(TEXT("BaseTooltip"), ToolTipOut);
-					ToolTipOut = FText::Format(LOCTEXT("InterfaceFunctionExpectedAsEvent_Tooltip", "{BaseTooltip}\nInterface '{InterfaceName}' is already implemented as a function graph but is expected as an event. Remove the function graph and reimplement as an event."), Args);
+					ToolTipOut = FText::Format(LOCTEXT("EventFromInterface_Tooltip", "Event (from Interface '{InterfaceName}')"), Args);
 					ColorOut = FLinearColor::Yellow;
 				}
 			}
