Allow an interface function graph to be converted in the case where the interface has changed and the functino shuold now be placed as an event.

diff --git a/Engine/Source/Editor/Kismet/Private/BlueprintEditor.cpp b/Engine/Source/Editor/Kismet/Private/BlueprintEditor.cpp
--- a/Engine/Source/Editor/Kismet/Private/BlueprintEditor.cpp
+++ b/Engine/Source/Editor/Kismet/Private/BlueprintEditor.cpp
@@ -5459,12 +5459,6 @@ void FBlueprintEditor::OnConvertFunctionToEvent()
 				LogResults.Error(*LOCTEXT("FunctionHasOutput_Error", "A function can only be converted if it does not have any output parameters.").ToString());
 				SpecificErrorMessage = LOCTEXT("FunctionHasOutput_Error_Title", "Function cannot have output parameters");
 			}
-			// Make sure this is not an interface function
-			else if (FBlueprintEditorUtils::IsInterfaceFunction(GetBlueprintObj(), Func))
-			{
-				LogResults.Error(*LOCTEXT("FunctionIsFromInterface_Error", "Functions from interfaces cannot be converted to events.").ToString());
-				SpecificErrorMessage = LOCTEXT("FunctionIsFromInterface_Error_Title", "Cannot convert interface functions");
-			}
 			// Make sure this is not a blueprint/macro function library
 			else if (GetBlueprintObj()->BlueprintType == BPTYPE_FunctionLibrary || GetBlueprintObj()->BlueprintType == BPTYPE_MacroLibrary)
 			{
@@ -5548,7 +5542,9 @@ void FBlueprintEditor::ConvertFunctionToEvent(UK2Node_FunctionEntry* SelectedCal
 		bool bIsOverrideFunc = Func->GetSuperFunction() || (ParentFunction != nullptr);
 		const UEdGraphSchema_K2* K2Schema = GetDefault<UEdGraphSchema_K2>();
 
-		if (bIsOverrideFunc)
+		// Allow an interface function to be converted in the case where the interface has changed
+		// and the function should now be placed as an event (UE-85687)
+		if (bIsOverrideFunc || FBlueprintEditorUtils::IsInterfaceFunction(NodeBP, Func))
 		{
 			NewEventNode = FEdGraphSchemaAction_K2NewNode::SpawnNode<UK2Node_Event>(
 				EventGraph,
diff --git a/Engine/Source/Editor/KismetCompiler/Private/KismetCompiler.cpp b/Engine/Source/Editor/KismetCompiler/Private/KismetCompiler.cpp
--- a/Engine/Source/Editor/KismetCompiler/Private/KismetCompiler.cpp
+++ b/Engine/Source/Editor/KismetCompiler/Private/KismetCompiler.cpp
@@ -3305,7 +3305,6 @@ void FKismetCompilerContext::CreateAndProcessUbergraph()
 			bool bExistsAsGraph = false;
 
 			// Any function that can be implemented as an event needs to check to see if there is already an interface function graph
-			// If there is, we want to warn the user that this is unexpected but proceed to successfully compile the Blueprint
 			if (bCanImplementAsEvent)
 			{
 				for (UEdGraph* InterfaceGraph : InterfaceDesc.Graphs)
@@ -3313,9 +3312,6 @@ void FKismetCompilerContext::CreateAndProcessUbergraph()
 					if (InterfaceGraph->GetFName() == Function->GetFName())
 					{
 						bExistsAsGraph = true;
-
-						// Having an event override implemented as a function won't cause issues but is something the user should be aware of.
-						MessageLog.Warning(TEXT("Interface '@@' is already implemented as a function graph but is expected as an event. Remove the function graph and reimplement as an event."), InterfaceGraph);
 					}
 				}
 			}
