Remove the clearing of keyboard focus when showing the details panel because it causes errors when undoing while a property edit still has keyboard focus. Clearing the keyboard focus during the process of the undo would cause another transaction to start, breaking the current undo transaction.

diff --git a/Engine/Source/Editor/Kismet/Private/SKismetInspector.cpp b/Engine/Source/Editor/Kismet/Private/SKismetInspector.cpp
--- a/Engine/Source/Editor/Kismet/Private/SKismetInspector.cpp
+++ b/Engine/Source/Editor/Kismet/Private/SKismetInspector.cpp
@@ -475,19 +475,6 @@ void SKismetInspector::ShowDetailsForSingleObject(UObject* Object, const FShowDe
 
 void SKismetInspector::ShowDetailsForObjects(const TArray<UObject*>& PropertyObjects, const FShowDetailsOptions& Options)
 {
-	static bool bIsReentrant = false;
-	if ( !bIsReentrant )
-	{
-		bIsReentrant = true;
-		// When the selection is changed, we may be potentially actively editing a property,
-		// if this occurs we need need to immediately clear keyboard focus
-		if ( FSlateApplication::Get().HasFocusedDescendants(AsShared()) )
-		{
-			FSlateApplication::Get().ClearKeyboardFocus(EFocusCause::Mouse);
-		}
-		bIsReentrant = false;
-	}
-
 	// Refresh is being deferred until the next tick, this prevents batch operations from bombarding the details view with calls to refresh
 	RefreshPropertyObjects = PropertyObjects;
 	RefreshOptions = Options;
