Disable the delete node for any interface function in the MyBlueprint panel. The user can no longer get into a situation where they would ever need to delete the graph, so the reasoning behind this being enabled is no longer valid.

#rnx

diff --git a/Engine/Source/Editor/Kismet/Private/SMyBlueprint.cpp b/Engine/Source/Editor/Kismet/Private/SMyBlueprint.cpp
--- a/Engine/Source/Editor/Kismet/Private/SMyBlueprint.cpp
+++ b/Engine/Source/Editor/Kismet/Private/SMyBlueprint.cpp
@@ -2693,25 +2693,11 @@ bool SMyBlueprint::CanDeleteEntry() const
 
 	if (FEdGraphSchemaAction_K2Graph* GraphAction = SelectionAsGraph())
 	{
-		if (GraphAction->EdGraph != NULL)
-		{
-			// Allow the user to delete any graphs in the interface section if the function can be placed as an event, 
-			// this allows users to resolve warnings when a previously implemented graph has been changed to be an event.
-			if (GraphAction->GetSectionID() == NodeSectionID::INTERFACE)
-			{
-				UFunction* Function = GetBlueprintObj()->SkeletonGeneratedClass->FindFunctionByName(GraphAction->EdGraph->GetFName());
-				if (UEdGraphSchema_K2::FunctionCanBePlacedAsEvent(Function))
-				{
-					return true;
-				}
-			}
-			return GraphAction->EdGraph->bAllowDeletion;
-		}
-		return false;
+		return (GraphAction->EdGraph ? GraphAction->EdGraph->bAllowDeletion : false);
 	}
 	else if (FEdGraphSchemaAction_K2Delegate* DelegateAction = SelectionAsDelegate())
 	{
-		return (DelegateAction->EdGraph != NULL) && (DelegateAction->EdGraph->bAllowDeletion) && 
+		return (DelegateAction->EdGraph != nullptr) && (DelegateAction->EdGraph->bAllowDeletion) &&
 			FDeleteEntryHelper::CanDeleteVariable(GetBlueprintObj(), DelegateAction->GetDelegateName());
 	}
 	else if (FEdGraphSchemaAction_K2Var* VarAction = SelectionAsVar())
@@ -2720,7 +2706,7 @@ bool SMyBlueprint::CanDeleteEntry() const
 	}
 	else if (FEdGraphSchemaAction_K2Event* EventAction = SelectionAsEvent())
 	{
-		return EventAction->NodeTemplate != NULL;
+		return EventAction->NodeTemplate != nullptr;
 	}
 	else if (FEdGraphSchemaAction_K2LocalVar* LocalVariable = SelectionAsLocalVar())
 	{
