Conform a blueprint to an interface when it refreshes to ensure you cannot make graphs with the same name as an interface

diff --git a/Engine/Source/Editor/Kismet/Private/SMyBlueprint.cpp b/Engine/Source/Editor/Kismet/Private/SMyBlueprint.cpp
--- a/Engine/Source/Editor/Kismet/Private/SMyBlueprint.cpp
+++ b/Engine/Source/Editor/Kismet/Private/SMyBlueprint.cpp
@@ -966,6 +966,9 @@ void SMyBlueprint::Refresh()
 {
 	bNeedsRefresh = false;
 
+	// Conform to our interfaces here to ensure we catch any newly added functions
+	FBlueprintEditorUtils::ConformImplementedInterfaces(GetBlueprintObj());
+
 	GraphActionMenu->RefreshAllActions(/*bPreserveExpansion=*/ true);
 }
 
@@ -2165,11 +2168,16 @@ bool SMyBlueprint::CanOpenGraph() const
 
 void SMyBlueprint::OpenGraph(FDocumentTracker::EOpenDocumentCause InCause)
 {
-	UEdGraph* GraphToOpen = NULL;
+	UEdGraph* GraphToOpen = nullptr;
 
 	if (FEdGraphSchemaAction_K2Graph* GraphAction = SelectionAsGraph())
 	{
 		GraphToOpen = GraphAction->EdGraph;
+		// If we have no graph then this is an interface event, so focus on the event graph
+		if (!GraphToOpen)
+		{
+			GraphToOpen = FBlueprintEditorUtils::FindEventGraph(GetBlueprintObj());
+		}
 	}
 	else if (FEdGraphSchemaAction_K2Delegate* DelegateAction = SelectionAsDelegate())
 	{
@@ -2205,7 +2213,22 @@ bool SMyBlueprint::CanFocusOnNode() const
 {
 	FEdGraphSchemaAction_K2Event const* const EventAction = SelectionAsEvent();
 	FEdGraphSchemaAction_K2InputAction const* const InputAction = SelectionAsInputAction();
-	return (EventAction && EventAction->NodeTemplate) || (InputAction && InputAction->NodeTemplate);
+	UK2Node_Event* ExistingNode = nullptr;
+
+	if (FEdGraphSchemaAction_K2Graph* GraphAction = SelectionAsGraph())
+	{
+		// Is this an event implemented from an interface?
+		UBlueprint* BlueprintObj = GetBlueprintObj();		
+		UFunction* OverrideFunc = nullptr;
+		UClass* const OverrideFuncClass = FBlueprintEditorUtils::GetOverrideFunctionClass(BlueprintObj, GraphAction->FuncName, &OverrideFunc);
+		UEdGraph* EventGraph = FBlueprintEditorUtils::FindEventGraph(BlueprintObj);
+
+		// Add to event graph
+		FName EventName = OverrideFunc->GetFName();
+		ExistingNode = FBlueprintEditorUtils::FindOverrideForFunction(BlueprintObj, OverrideFuncClass, EventName);
+	}
+
+	return (EventAction && EventAction->NodeTemplate) || (InputAction && InputAction->NodeTemplate) || ExistingNode;
 }
 
 void SMyBlueprint::OnFocusNode()
@@ -2217,6 +2240,21 @@ void SMyBlueprint::OnFocusNode()
 		UK2Node* Node = EventAction ? EventAction->NodeTemplate : InputAction->NodeTemplate;
 		FKismetEditorUtilities::BringKismetToFocusAttentionOnObject(Node);
 	}
+	else if (FEdGraphSchemaAction_K2Graph* GraphAction = SelectionAsGraph())
+	{
+		// Is this an event implemented from an interface?
+		UBlueprint* BlueprintObj = GetBlueprintObj();
+		UFunction* OverrideFunc = nullptr;
+		UClass* const OverrideFuncClass = FBlueprintEditorUtils::GetOverrideFunctionClass(BlueprintObj, GraphAction->FuncName, &OverrideFunc);
+		UEdGraph* EventGraph = FBlueprintEditorUtils::FindEventGraph(BlueprintObj);
+
+		// Add to event graph
+		FName EventName = OverrideFunc->GetFName();
+		if (UK2Node_Event* ExistingNode = FBlueprintEditorUtils::FindOverrideForFunction(BlueprintObj, OverrideFuncClass, EventName))
+		{
+			FKismetEditorUtilities::BringKismetToFocusAttentionOnObject(ExistingNode);
+		}
+	}
 }
 
 void SMyBlueprint::OnFocusNodeInNewTab()
@@ -2228,7 +2266,7 @@ void SMyBlueprint::OnFocusNodeInNewTab()
 bool SMyBlueprint::CanImplementFunction() const
 {
 	FEdGraphSchemaAction_K2Graph* GraphAction = SelectionAsGraph();
-	return GraphAction && GraphAction->EdGraph == NULL;
+	return GraphAction && GraphAction->EdGraph == nullptr && !CanFocusOnNode();
 }
 
 void SMyBlueprint::OnImplementFunction()
@@ -2249,6 +2287,10 @@ void SMyBlueprint::ImplementFunction(FEdGraphSchemaAction_K2Graph* GraphAction)
 	UBlueprint* BlueprintObj = GetBlueprintObj();
 	check(BlueprintObj && BlueprintObj->SkeletonGeneratedClass);
 
+	// Ensure that we are conforming to all current interfaces so that if there has been an additional
+	// interface function added we just focus to it instead of creating a new one
+	FBlueprintEditorUtils::ConformImplementedInterfaces(BlueprintObj);
+
 	UFunction* OverrideFunc = nullptr;
 	UClass* const OverrideFuncClass = FBlueprintEditorUtils::GetOverrideFunctionClass(BlueprintObj, GraphAction->FuncName, &OverrideFunc);
 	check(OverrideFunc);
