Add the access specifiers to functions in the MyBlueprint window
- Additional blueprint editor font styles
- additional editor setting and view options in myblueprint for toggleing this
- by default it is off

diff --git a/Engine/Source/Editor/BlueprintGraph/Private/BlueprintEditorSettings.cpp b/Engine/Source/Editor/BlueprintGraph/Private/BlueprintEditorSettings.cpp
--- a/Engine/Source/Editor/BlueprintGraph/Private/BlueprintEditorSettings.cpp
+++ b/Engine/Source/Editor/BlueprintGraph/Private/BlueprintEditorSettings.cpp
@@ -28,6 +28,7 @@ UBlueprintEditorSettings::UBlueprintEditorSettings(const FObjectInitializer& Obj
 	, bAlwaysShowInterfacesInOverrides(true)
 	, bShowParentClassInOverrides(true)
 	, bShowEmptySections(true)
+	, bShowAccessSpecifier(false)
 	, bSpawnDefaultBlueprintNodes(true)
 	, bHideConstructionScriptComponentsInDetailsView(true)
 	, bHostFindInBlueprintsInGlobalTab(true)
diff --git a/Engine/Source/Editor/BlueprintGraph/Public/BlueprintEditorSettings.h b/Engine/Source/Editor/BlueprintGraph/Public/BlueprintEditorSettings.h
--- a/Engine/Source/Editor/BlueprintGraph/Public/BlueprintEditorSettings.h
+++ b/Engine/Source/Editor/BlueprintGraph/Public/BlueprintEditorSettings.h
@@ -86,6 +86,10 @@ public:
 	UPROPERTY(config)
 	bool bShowEmptySections;
 
+	/** If set we'll show the access specifier of functions in the My Blueprint view */
+	UPROPERTY(config)
+	bool bShowAccessSpecifier;
+
 	/** If set will spawn default nodes in new Blueprints */
 	UPROPERTY(EditAnywhere, config, Category=Workflow)
 	bool bSpawnDefaultBlueprintNodes;
diff --git a/Engine/Source/Editor/EditorStyle/Private/SlateEditorStyle.cpp b/Engine/Source/Editor/EditorStyle/Private/SlateEditorStyle.cpp
--- a/Engine/Source/Editor/EditorStyle/Private/SlateEditorStyle.cpp
+++ b/Engine/Source/Editor/EditorStyle/Private/SlateEditorStyle.cpp
@@ -5994,6 +5994,12 @@ void FSlateEditorStyle::FStyle::SetupPersonaStyle()
 			);
 	}
 
+	// Access modifier display in MyBlueprint window for functions/variables
+	{
+		Set("BlueprintEditor.AccessModifier.Public", FTextBlockStyle().SetFont(DEFAULT_FONT("Bold", 9) ));
+		Set("BlueprintEditor.AccessModifier.Default", FTextBlockStyle().SetFont(DEFAULT_FONT("Regular", 9)));
+	}
+
 	// Kismet linear expression display
 	{
 		Set( "KismetExpression.ReadVariable.Body", new BOX_BRUSH( "/Graph/Linear_VarNode_Background", FMargin(16.f/64.f, 12.f/28.f) ) );
diff --git a/Engine/Source/Editor/Kismet/Private/SBlueprintPalette.cpp b/Engine/Source/Editor/Kismet/Private/SBlueprintPalette.cpp
--- a/Engine/Source/Editor/Kismet/Private/SBlueprintPalette.cpp
+++ b/Engine/Source/Editor/Kismet/Private/SBlueprintPalette.cpp
@@ -1023,6 +1023,22 @@ void SBlueprintPaletteItem::Construct(const FArguments& InArgs, FCreateWidgetFor
 	TSharedRef<SWidget> IconWidget = CreateIconWidget(IconToolTip, IconBrush, IconColor, IconDocLink, IconDocExcerpt, SecondaryBrush, SecondaryIconColor);
 	IconWidget->SetEnabled(bIsEditingEnabled);
 
+	UBlueprintEditorSettings* Settings = GetMutableDefault<UBlueprintEditorSettings>();
+
+	// Enum representing the access specifier of this function or variable
+	enum class EAccessSpecifier : uint8
+	{
+		None		= 0,
+		Private		= 1,
+		Protected	= 2,
+		Public		= 3
+	};
+
+	// We should only bother checking for access if the setting is on and this is not an animation graph
+	const bool bShouldCheckForAccessSpec = Settings->bShowAccessSpecifier;
+
+	EAccessSpecifier ActionAccessSpecifier = EAccessSpecifier::None;	
+
 	// Setup a meta tag for this node
 	FTutorialMetaData TagMeta("PaletteItem"); 
 	if( ActionPtr.IsValid() )
@@ -1034,23 +1050,27 @@ void SBlueprintPaletteItem::Construct(const FArguments& InArgs, FCreateWidgetFor
 	FSlateFontInfo NameFont = FCoreStyle::GetDefaultFontStyle("Regular", 10);
 	TSharedRef<SWidget> NameSlotWidget = CreateTextSlotWidget( NameFont, InCreateData, bIsReadOnly );
 	
-	// For Variables and Local Variables, we will convert the icon widget into a pin type selector.
-	if (GraphAction->GetTypeId() == FEdGraphSchemaAction_K2Var::StaticGetTypeId() || GraphAction->GetTypeId() == FEdGraphSchemaAction_K2LocalVar::StaticGetTypeId())
+	// Will set the icon of this property to be a Pin Type selector. 
+	auto GenerateVariableSettings = [&](FProperty* VariableProp)
 	{
-		FProperty* VariableProp = nullptr;
-
-		if (GraphAction->GetTypeId() == FEdGraphSchemaAction_K2Var::StaticGetTypeId())
-		{
-			VariableProp = StaticCastSharedPtr<FEdGraphSchemaAction_K2Var>(GraphAction)->GetProperty();
-		}
-		else if (GraphAction->GetTypeId() == FEdGraphSchemaAction_K2LocalVar::StaticGetTypeId())
-		{
-			VariableProp = StaticCastSharedPtr<FEdGraphSchemaAction_K2LocalVar>(GraphAction)->GetProperty();
-		}
-
-		// If the variable is not a local variable or created by the current Blueprint, do not use the PinTypeSelector
 		if (VariableProp)
 		{
+			if (bShouldCheckForAccessSpec)
+			{
+				if (VariableProp->GetBoolMetaData(FBlueprintMetadata::MD_Private))
+				{
+					ActionAccessSpecifier = EAccessSpecifier::Private;
+				}
+				else if (VariableProp->GetBoolMetaData(FBlueprintMetadata::MD_Protected))
+				{
+					ActionAccessSpecifier = EAccessSpecifier::Protected;
+				}
+				else
+				{
+					ActionAccessSpecifier = EAccessSpecifier::Public;
+				}
+			}
+
 			if (FBlueprintEditorUtils::IsVariableCreatedByBlueprint(Blueprint, VariableProp) || VariableProp->GetOwner<UFunction>())
 			{
 				const UEdGraphSchema_K2* Schema = GetDefault<UEdGraphSchema_K2>();
@@ -1058,36 +1078,126 @@ void SBlueprintPaletteItem::Construct(const FArguments& InArgs, FCreateWidgetFor
 					.IsEnabled(bIsEditingEnabled);
 			}
 		}
+	};
+
+	// For Variables and Local Variables, we will convert the icon widget into a pin type selector.
+	if (GraphAction->GetTypeId() == FEdGraphSchemaAction_K2Var::StaticGetTypeId())
+	{	
+		GenerateVariableSettings(StaticCastSharedPtr<FEdGraphSchemaAction_K2Var>(GraphAction)->GetProperty());
+	}
+	else if (GraphAction->GetTypeId() == FEdGraphSchemaAction_K2LocalVar::StaticGetTypeId())
+	{
+		GenerateVariableSettings(StaticCastSharedPtr<FEdGraphSchemaAction_K2LocalVar>(GraphAction)->GetProperty());
 	}
+	// Determine the access level of this action if it is a function graph or for interface events
+	else if (bShouldCheckForAccessSpec && GraphAction->GetTypeId() == FEdGraphSchemaAction_K2Graph::StaticGetTypeId())
+	{
+		UFunction* FunctionToCheck = nullptr;
 
-	// now, create the actual widget
-	ChildSlot
-	[
-		SNew(SHorizontalBox)
-		.AddMetaData<FTutorialMetaData>(TagMeta)
-		// icon slot
-		+SHorizontalBox::Slot()
-			.AutoWidth()
+		if (FEdGraphSchemaAction_K2Graph* FuncGraphAction = (FEdGraphSchemaAction_K2Graph*)(GraphAction.Get()))
+		{
+			FunctionToCheck = FindField<UFunction>(Blueprint->SkeletonGeneratedClass, FuncGraphAction->FuncName);
+
+			// Handle override/interface functions
+			if(!FunctionToCheck)
+			{
+				FBlueprintEditorUtils::GetOverrideFunctionClass(Blueprint, FuncGraphAction->FuncName, &FunctionToCheck);			
+			}
+		}
+
+		// If we have found a function that matches this action name, then grab it's access specifier
+		if (FunctionToCheck)
+		{
+			if (FunctionToCheck->HasAnyFunctionFlags(FUNC_Protected))
+			{
+				ActionAccessSpecifier = EAccessSpecifier::Protected;
+			}
+			else if (FunctionToCheck->HasAnyFunctionFlags(FUNC_Private))
+			{
+				ActionAccessSpecifier = EAccessSpecifier::Private;
+			}
+			else
+			{
+				ActionAccessSpecifier = EAccessSpecifier::Public;
+			}
+		}
+	}
+
+	FText AccessModifierText = FText::GetEmpty();
+
+	switch (ActionAccessSpecifier)
+	{
+		case EAccessSpecifier::Public:
+		{
+			AccessModifierText = LOCTEXT("AccessModifierPublic", "public");
+		}
+		break;
+		case EAccessSpecifier::Protected:
+		{
+			AccessModifierText = LOCTEXT("AccessModifierProtected", "protected");
+		}
+		break;
+		case EAccessSpecifier::Private:
+		{
+			AccessModifierText = LOCTEXT("AccessModifierPrivate", "private");
+		}
+		break;
+	}
+
+	// Calculate a color so that the text gets brighter the more accessible the action is
+	const bool AccessSpecifierEnabled = (ActionAccessSpecifier != EAccessSpecifier::None) && bShouldCheckForAccessSpec;
+
+	// Create the widget with an icon
+	TSharedRef<SHorizontalBox> ActionBox = SNew(SHorizontalBox)		
+		.AddMetaData<FTutorialMetaData>(TagMeta);
+
+	ActionBox.Get().AddSlot()
+		.AutoWidth()
 		[
 			IconWidget
-		]
-		// name slot
-		+SHorizontalBox::Slot()
-			.FillWidth(1.f)
+		];
+
+	// Only add an access specifier if we have one
+	if (ActionAccessSpecifier != EAccessSpecifier::None)
+	{
+		ActionBox.Get().AddSlot()
+			.MaxWidth(50.f)
+			.FillWidth(AccessSpecifierEnabled ? 0.4f : 0.0f)
+			.Padding(FMargin(/* horizontal */ AccessSpecifierEnabled ? 6.0f : 0.0f, /* vertical */ 0.0f))
 			.VAlign(VAlign_Center)
-			.Padding(3,0)
+			.HAlign(HAlign_Right)
+			[
+				SNew(STextBlock)
+				// Will only display text if we have a modifier level
+					.IsEnabled(AccessSpecifierEnabled)
+					.Text(AccessModifierText)
+					.ColorAndOpacity(FSlateColor::UseSubduedForeground())
+					// Bold if public
+					.TextStyle(FEditorStyle::Get(), ActionAccessSpecifier == EAccessSpecifier::Public ? "BlueprintEditor.AccessModifier.Public" : "BlueprintEditor.AccessModifier.Default")
+			];
+	}
+
+	ActionBox.Get().AddSlot()
+		.FillWidth(1.f)
+		.VAlign(VAlign_Center)
+		.Padding(/* horizontal */ 3.0f, /* vertical */ 0.0f)
 		[
 			NameSlotWidget
-		]
-		// optional visibility slot
-		+SHorizontalBox::Slot()
-			.AutoWidth()
-			.Padding(FMargin(3,0))
-			.VAlign(VAlign_Center)
+		];
+
+	ActionBox.Get().AddSlot()
+		.AutoWidth()
+		.Padding(FMargin(3.0f, 0.0f))
+		.VAlign(VAlign_Center)
 		[
 			SNew(SPaletteItemVisibilityToggle, ActionPtr, InBlueprintEditor, InBlueprint)
 			.IsEnabled(bIsEditingEnabled)
-		]
+		];
+
+	// Now, create the actual widget
+	ChildSlot
+	[
+		ActionBox
 	];
 }
 
diff --git a/Engine/Source/Editor/Kismet/Private/SMyBlueprint.cpp b/Engine/Source/Editor/Kismet/Private/SMyBlueprint.cpp
--- a/Engine/Source/Editor/Kismet/Private/SMyBlueprint.cpp
+++ b/Engine/Source/Editor/Kismet/Private/SMyBlueprint.cpp
@@ -406,6 +406,20 @@ void SMyBlueprint::Construct(const FArguments& InArgs, TWeakPtr<FBlueprintEditor
 		TEXT("MyBlueprint_AlwaysShowInterfacesInOverrides")
 	);
 
+	ViewOptions.AddMenuEntry(
+		LOCTEXT("AlwaysShowAccessSpecifier", "Show access specifier in the My Blueprint View"),
+		LOCTEXT("AlwaysShowAccessSpecifierTooltip", "Should we always display the access specifier of functions in the function menu?"),
+		FSlateIcon(),
+		FUIAction(
+			FExecuteAction::CreateSP(this, &SMyBlueprint::OnToggleShowAccessSpecifier),
+			FCanExecuteAction(),
+			FIsActionChecked::CreateSP(this, &SMyBlueprint::GetShowAccessSpecifier)
+		),
+		NAME_None,
+		EUserInterfaceActionType::ToggleButton,
+		TEXT("MyBlueprint_AlwaysShowAccessSpecifier")
+	);
+
 	SAssignNew(FilterBox, SSearchBox)
 		.OnTextChanged( this, &SMyBlueprint::OnFilterTextChanged );
 
@@ -1601,6 +1615,20 @@ bool SMyBlueprint::GetShowParentClassInOverrides() const
 	return GetMutableDefault<UBlueprintEditorSettings>()->bShowParentClassInOverrides;
 }
 
+void SMyBlueprint::OnToggleShowAccessSpecifier()
+{
+	UBlueprintEditorSettings* Settings = GetMutableDefault<UBlueprintEditorSettings>();
+	Settings->bShowAccessSpecifier = !Settings->bShowAccessSpecifier;
+	Settings->PostEditChange();
+	Settings->SaveConfig();
+	Refresh();
+}
+
+bool SMyBlueprint::GetShowAccessSpecifier() const 
+{
+	return GetMutableDefault<UBlueprintEditorSettings>()->bShowAccessSpecifier;
+}
+
 bool SMyBlueprint::IsShowingReplicatedVariablesOnly() const
 {
 	return bShowReplicatedVariablesOnly;
diff --git a/Engine/Source/Editor/Kismet/Private/SMyBlueprint.h b/Engine/Source/Editor/Kismet/Private/SMyBlueprint.h
--- a/Engine/Source/Editor/Kismet/Private/SMyBlueprint.h
+++ b/Engine/Source/Editor/Kismet/Private/SMyBlueprint.h
@@ -188,10 +188,14 @@ private:
 	void OnToggleAlwaysShowInterfacesInOverrides();
 	bool GetAlwaysShowInterfacesInOverrides() const;
 
-	/** Support functions for view options for bShowParentClassInOverrides blueprint editor setting*/
+	/** Support functions for view options for bShowParentClassInOverrides blueprint editor setting */
 	void OnToggleShowParentClassInOverrides();
 	bool GetShowParentClassInOverrides() const;
 
+	/** Support functions for view options for bShowAccessSpecifier blueprint editor setting */
+	void OnToggleShowAccessSpecifier();
+	bool GetShowAccessSpecifier() const;
+
 	/** Helper function to open the selected graph */
 	void OpenGraph(FDocumentTracker::EOpenDocumentCause InCause);
 
