- Allow TArray, TSet and TMap elements to be visible in the SCS tree if the container is marked as Visible/Editable
- Make the behavior of the Kismet and Actor details panel consistent when editing subobjects that come from a container
- Add the component instead of it's owner when adding SelectionInfo allowing for more user friendly editing

diff --git a/Engine/Source/Editor/Kismet/Private/SKismetInspector.cpp b/Engine/Source/Editor/Kismet/Private/SKismetInspector.cpp
--- a/Engine/Source/Editor/Kismet/Private/SKismetInspector.cpp
+++ b/Engine/Source/Editor/Kismet/Private/SKismetInspector.cpp
@@ -33,6 +33,7 @@
 #include "IDetailCustomization.h"
 #include "Editor.h"
 #include "PropertyEditorModule.h"
+#include "Kismet2/ComponentEditorUtils.h"	// For CanEditNativeComponent()
 
 #include "IDetailsView.h"
 
@@ -254,7 +255,7 @@ FText SKismetInspector::GetContextualEditingWidgetTitle() const
 			}
 			else if (USCS_Node* SCSNode = Cast<USCS_Node>(Object))
 			{
-				if (SCSNode->ComponentTemplate != NULL)
+				if (SCSNode->ComponentTemplate != nullptr)
 				{
 					const FName VariableName = SCSNode->GetVariableName();
 					if (VariableName != NAME_None)
@@ -336,7 +337,7 @@ void SKismetInspector::Construct(const FArguments& InArgs)
 	// Create a property view
 	FPropertyEditorModule& EditModule = FModuleManager::Get().GetModuleChecked<FPropertyEditorModule>("PropertyEditor");
 
-	FNotifyHook* NotifyHook = NULL;
+	FNotifyHook* NotifyHook = nullptr;
 	if(InArgs._SetNotifyHook)
 	{
 		NotifyHook = Kismet2.Get();
@@ -470,7 +471,7 @@ void SKismetInspector::ShowDetailsForSingleObject(UObject* Object, const FShowDe
 {
 	TArray<UObject*> PropertyObjects;
 
-	if (Object != NULL)
+	if (Object != nullptr)
 	{
 		PropertyObjects.Add(Object);
 	}
@@ -512,7 +513,7 @@ void SKismetInspector::ShowSingleStruct(TSharedPtr<FStructOnScope> InStructToDis
 
 void SKismetInspector::AddPropertiesRecursive(FProperty* Property)
 {
-	if (Property != NULL)
+	if (Property != nullptr)
 	{
 		// Add this property
 		SelectedObjectProperties.Add(Property);
@@ -520,8 +521,7 @@ void SKismetInspector::AddPropertiesRecursive(FProperty* Property)
 		// If this is a struct or an array of structs, recursively add the child properties
 		FArrayProperty* ArrayProperty = CastField<FArrayProperty>(Property);
 		FStructProperty* StructProperty = CastField<FStructProperty>(Property);
-		if(	StructProperty != NULL && 
-			StructProperty->Struct != NULL)
+		if(	StructProperty != nullptr && StructProperty->Struct != nullptr)
 		{
 			for (TFieldIterator<FProperty> StructPropIt(StructProperty->Struct); StructPropIt; ++StructPropIt)
 			{
@@ -627,7 +627,7 @@ void SKismetInspector::UpdateFromObjects(const TArray<UObject*>& PropertyObjects
 			{
 				// Edit the component template
 				UActorComponent* NodeComponent = SCSNode->ComponentTemplate;
-				if (NodeComponent != NULL)
+				if (NodeComponent != nullptr)
 				{
 					SelectionInfo.ObjectsForPropertyEditing.Add(NodeComponent);
 					SelectionInfo.EditableComponentTemplates.Add(NodeComponent);
@@ -654,10 +654,9 @@ void SKismetInspector::UpdateFromObjects(const TArray<UObject*>& PropertyObjects
 			else if (UActorComponent* ActorComponent = Cast<UActorComponent>(Object))
 			{
 				AActor* Owner = ActorComponent->GetOwner();
-				if(Owner != NULL && Owner->HasAnyFlags(RF_ClassDefaultObject))
-				{
-					// We're editing a component that's owned by a CDO, so set the CDO to the property editor (so that propagation works) and then filter to just the component property that we want to edit
-					SelectionInfo.ObjectsForPropertyEditing.AddUnique(Owner);
+				if(Owner != nullptr && Owner->HasAnyFlags(RF_ClassDefaultObject))
+				{					
+					SelectionInfo.ObjectsForPropertyEditing.AddUnique(ActorComponent);
 					SelectionInfo.EditableComponentTemplates.Add(ActorComponent);
 				}
 				else
@@ -683,13 +682,13 @@ void SKismetInspector::UpdateFromObjects(const TArray<UObject*>& PropertyObjects
 		for (auto CompIt = SelectionInfo.EditableComponentTemplates.CreateIterator(); CompIt; ++CompIt)
 		{
 			UActorComponent* EditableComponentTemplate = *CompIt;
-			check(EditableComponentTemplate != NULL);
+			check(EditableComponentTemplate != nullptr);
 
 			// Add all properties belonging to the component template class
 			for (TFieldIterator<FProperty> PropIt(EditableComponentTemplate->GetClass()); PropIt; ++PropIt)
 			{
 				FProperty* Property = *PropIt;
-				check(Property != NULL);
+				check(Property != nullptr);
 
 				AddPropertiesRecursive(Property);
 			}
@@ -698,7 +697,7 @@ void SKismetInspector::UpdateFromObjects(const TArray<UObject*>& PropertyObjects
 			for (auto ObjIt = SelectionInfo.ObjectsForPropertyEditing.CreateIterator(); ObjIt; ++ObjIt)
 			{
 				UObject* Object = *ObjIt;
-				check(Object != NULL);
+				check(Object != nullptr);
 
 				if (Object != EditableComponentTemplate)
 				{
@@ -708,40 +707,12 @@ void SKismetInspector::UpdateFromObjects(const TArray<UObject*>& PropertyObjects
 					}
 					else
 					{
-						auto FindPropertyReferencingComponent = [Object](const UActorComponent* Component) -> FProperty*
-						{
-							if (AActor* Owner = Component->GetOwner())
-							{
-								if (UClass* OwnerClass = Owner->GetClass())
-								{
-									AActor* OwnerCDO = CastChecked<AActor>(OwnerClass->GetDefaultObject());
-									for (TFieldIterator<FObjectProperty> ObjPropIt(OwnerClass, EFieldIteratorFlags::IncludeSuper); ObjPropIt; ++ObjPropIt)
-									{
-										FObjectProperty* ObjectProperty = *ObjPropIt;
-										check(ObjectProperty != nullptr);
-
-										// If the property value matches the current archetype, add it as a selected property for filtering
-										if (Component->GetClass()->IsChildOf(ObjectProperty->PropertyClass)
-											&& Component == ObjectProperty->GetObjectPropertyValue_InContainer(OwnerCDO))
-										{
-											ObjectProperty = FindField<FObjectProperty>(Object->GetClass(), ObjectProperty->GetFName());
-											if (ObjectProperty != nullptr)
-											{
-												return ObjectProperty;
-											}
-										}
-									}
-								}
-							}
-							return nullptr;
-						};
-
-						FProperty* ReferencingProperty = FindPropertyReferencingComponent(EditableComponentTemplate);
+						FProperty* ReferencingProperty = FComponentEditorUtils::GetPropertyForEditableNativeComponent(EditableComponentTemplate);
 						if (ReferencingProperty == nullptr)
 						{
 							if (UActorComponent* Archetype = Cast<UActorComponent>(EditableComponentTemplate->GetArchetype()))
 							{
-								ReferencingProperty = FindPropertyReferencingComponent(Archetype);
+								ReferencingProperty = FComponentEditorUtils::GetPropertyForEditableNativeComponent(Archetype);
 							}
 						}
 						if (ReferencingProperty)
@@ -815,7 +786,7 @@ bool SKismetInspector::IsPropertyVisible( const FPropertyAndParent& PropertyAndP
 
 	if(const UClass* OwningClass = Property.GetOwner<UClass>())
 	{
-		const UBlueprint* BP = BlueprintEditorPtr.IsValid() ? BlueprintEditorPtr.Pin()->GetBlueprintObj() : NULL;
+		const UBlueprint* BP = BlueprintEditorPtr.IsValid() ? BlueprintEditorPtr.Pin()->GetBlueprintObj() : nullptr;
 		const bool VariableAddedInCurentBlueprint = (OwningClass->ClassGeneratedBy == BP);
 
 		// If we did not add this var, hide it!
@@ -830,9 +801,12 @@ bool SKismetInspector::IsPropertyVisible( const FPropertyAndParent& PropertyAndP
 
 	// figure out if this Blueprint variable is an Actor variable
 	const FArrayProperty* ArrayProperty = CastField<const FArrayProperty>(&Property);
+	const FSetProperty* SetProperty = CastField<const FSetProperty>(&Property);
+	const FMapProperty* MapProperty = CastField<const FMapProperty>(&Property);
+
 	const FProperty* TestProperty = ArrayProperty ? ArrayProperty->Inner : &Property;
 	const FObjectPropertyBase* ObjectProperty = CastField<const FObjectPropertyBase>(TestProperty);
-	bool bIsActorProperty = (ObjectProperty != NULL && ObjectProperty->PropertyClass->IsChildOf(AActor::StaticClass()));
+	bool bIsActorProperty = (ObjectProperty != nullptr && ObjectProperty->PropertyClass->IsChildOf(AActor::StaticClass()));
 
 	if (bEditOnTemplateDisabled && bIsActorProperty)
 	{
@@ -868,6 +842,19 @@ bool SKismetInspector::IsPropertyVisible( const FPropertyAndParent& PropertyAndP
 			return true;
 		}
 	}
+	else if (ArrayProperty || MapProperty || SetProperty)
+	{
+		// .Find won't work here because the items inside of the container properties are not FProperties
+		for (const TWeakFieldPtr<FProperty>& CurProp : SelectedObjectProperties)
+		{
+			if ((ArrayProperty && (ArrayProperty->PropertyFlags & CPF_Edit) && CurProp->GetFName() == ArrayProperty->GetFName()) ||
+				(MapProperty && (MapProperty->PropertyFlags & CPF_Edit) && CurProp->GetFName() == MapProperty->GetFName()) ||
+				(SetProperty && (SetProperty->PropertyFlags & CPF_Edit) && CurProp->GetFName() == SetProperty->GetFName()))
+			{
+				return true;
+			}
+		}
+	}
 
 
 	return SelectedObjectProperties.Num() == 0;
diff --git a/Engine/Source/Editor/Kismet/Private/SSCSEditor.cpp b/Engine/Source/Editor/Kismet/Private/SSCSEditor.cpp
--- a/Engine/Source/Editor/Kismet/Private/SSCSEditor.cpp
+++ b/Engine/Source/Editor/Kismet/Private/SSCSEditor.cpp
@@ -1252,7 +1252,7 @@ bool FSCSEditorTreeNodeComponent::CanEditDefaults() const
 	}
 	else if (UActorComponent* ComponentTemplate = GetComponentTemplate())
 	{
-		bCanEdit = FComponentEditorUtils::CanEditNativeComponent(ComponentTemplate);
+		bCanEdit = (FComponentEditorUtils::GetPropertyForEditableNativeComponent(ComponentTemplate) != nullptr);
 	}
 
 	return bCanEdit;
@@ -4708,7 +4708,7 @@ void SSCSEditor::UpdateTree(bool bRegenerateTreeNodes)
 						&& (!ActorComp->IsVisualizationComponent())
 						&& (ActorComp->CreationMethod != EComponentCreationMethod::UserConstructionScript || !bHideConstructionScriptComponentsInDetailsView)
 						&& (ParentSceneComp == nullptr || !ParentSceneComp->IsCreatedByConstructionScript() || !ActorComp->HasAnyFlags(RF_DefaultSubObject)))
-						&& (ActorComp->CreationMethod != EComponentCreationMethod::Native || FComponentEditorUtils::CanEditNativeComponent(ActorComp));
+						&& (ActorComp->CreationMethod != EComponentCreationMethod::Native || FComponentEditorUtils::GetPropertyForEditableNativeComponent(ActorComp));
 				};
 
 				for (auto It(ComponentsToAdd.CreateIterator()); It; ++It)
diff --git a/Engine/Source/Editor/UnrealEd/Private/Kismet2/ComponentEditorUtils.cpp b/Engine/Source/Editor/UnrealEd/Private/Kismet2/ComponentEditorUtils.cpp
--- a/Engine/Source/Editor/UnrealEd/Private/Kismet2/ComponentEditorUtils.cpp
+++ b/Engine/Source/Editor/UnrealEd/Private/Kismet2/ComponentEditorUtils.cpp
@@ -164,18 +164,17 @@ bool FComponentEditorUtils::CanEditComponentInstance(const UActorComponent* Acto
 		&& (!ActorComp->IsVisualizationComponent())
 		&& (ActorComp->CreationMethod != EComponentCreationMethod::UserConstructionScript || bAllowUserContructionScript)
 		&& (ParentSceneComp == nullptr || !ParentSceneComp->IsCreatedByConstructionScript() || !ActorComp->HasAnyFlags(RF_DefaultSubObject)))
-		&& (ActorComp->CreationMethod != EComponentCreationMethod::Native || FComponentEditorUtils::CanEditNativeComponent(ActorComp));
+		&& (ActorComp->CreationMethod != EComponentCreationMethod::Native || FComponentEditorUtils::GetPropertyForEditableNativeComponent(ActorComp));
 }
 
-bool FComponentEditorUtils::CanEditNativeComponent(const UActorComponent* NativeComponent)
+FProperty* FComponentEditorUtils::GetPropertyForEditableNativeComponent(const UActorComponent* NativeComponent)
 {
 	// A native component can be edited if it is bound to a member variable and that variable is marked as visible in the editor
-	// Note: We aren't concerned with whether the component is marked editable - the component itself is responsible for determining which of its properties are editable
-
-	bool bCanEdit = false;
-	
+	// Note: We aren't concerned with whether the component is marked editable - the component itself is responsible for determining which of its properties are editable	
 	UObject* ComponentOuter = (NativeComponent ? NativeComponent->GetOuter() : nullptr);
 	UClass* OwnerClass = (ComponentOuter ? ComponentOuter->GetClass() : nullptr);
+	AActor* OwnerCDO = (OwnerClass ? CastChecked<AActor>(OwnerClass->GetDefaultObject()) : nullptr);
+
 	if (OwnerClass != nullptr)
 	{
 		for (TFieldIterator<FObjectProperty> It(OwnerClass); It; ++It)
@@ -183,22 +182,111 @@ bool FComponentEditorUtils::CanEditNativeComponent(const UActorComponent* Native
 			FObjectProperty* ObjectProp = *It;
 
 			// Must be visible - note CPF_Edit is set for all properties that should be visible, not just those that are editable
-			if (( ObjectProp->PropertyFlags & ( CPF_Edit ) ) == 0)
+			if ((ObjectProp->PropertyFlags & (CPF_Edit)) == 0)
 			{
 				continue;
 			}
 
 			UObject* Object = ObjectProp->GetObjectPropertyValue(ObjectProp->ContainerPtrToValuePtr<void>(ComponentOuter));
-			bCanEdit = Object != nullptr && Object->GetFName() == NativeComponent->GetFName();
+			if (Object != nullptr && Object->GetFName() == NativeComponent->GetFName())
+			{
+				return ObjectProp;
+			}
+		}	
+	
+		// We have to check for array properties as well because they are not FObjectProperties and we want to be able to
+		// edit the inside of it
+		if (OwnerCDO != nullptr)
+		{
+			for (TFieldIterator<FArrayProperty> PropIt(OwnerClass, EFieldIteratorFlags::IncludeSuper); PropIt; ++PropIt)
+			{
+				FArrayProperty* TestProperty = *PropIt;
+				void* ArrayPropInstAddress = TestProperty->ContainerPtrToValuePtr<void>(OwnerCDO);
+
+				// Ensure that this property is valid
+				FObjectProperty* ArrayEntryProp = CastField<FObjectProperty>(TestProperty->Inner);
+				if ((ArrayEntryProp == nullptr) || !ArrayEntryProp->PropertyClass->IsChildOf<UActorComponent>() || ((TestProperty->PropertyFlags & CPF_Edit) == 0))
+				{
+					continue;
+				}
+
+				// For each object in this array
+				FScriptArrayHelper ArrayHelper(TestProperty, ArrayPropInstAddress);
+				for (int32 ComponentIndex = 0; ComponentIndex < ArrayHelper.Num(); ++ComponentIndex)
+				{
+					// If this object is the native component we are looking for, then we should be allowed to edit it
+					const uint8* ArrayValRawPtr = ArrayHelper.GetRawPtr(ComponentIndex);
+					UObject* ArrayElement = ArrayEntryProp->GetObjectPropertyValue(ArrayValRawPtr);			
+					
+					if (ArrayElement != nullptr && ArrayElement->GetFName() == NativeComponent->GetFName())
+					{
+						return ArrayEntryProp;
+					}
+				}				
+			}			
 
-			if (bCanEdit)
+			// Check for map properties as well
+			for (TFieldIterator<FMapProperty> PropIt(OwnerClass, EFieldIteratorFlags::IncludeSuper); PropIt; ++PropIt)
 			{
-				break;
+				FMapProperty* TestProperty = *PropIt;
+				void* MapPropInstAddress = TestProperty->ContainerPtrToValuePtr<void>(OwnerCDO);
+				
+				// Ensure that this property is valid and that it is marked as visible in the editor
+				FObjectProperty* MapValProp = CastField<FObjectProperty>(TestProperty->ValueProp);
+				if ((MapValProp == nullptr) || !MapValProp->PropertyClass->IsChildOf<UActorComponent>() || ((TestProperty->PropertyFlags & CPF_Edit) == 0))
+				{
+					continue;
+				}
+
+				FScriptMapHelper MapHelper(TestProperty, MapPropInstAddress);
+				for (int32 MapSparseIndex = 0; MapSparseIndex < MapHelper.GetMaxIndex(); ++MapSparseIndex)
+				{
+					// For each value in the map (don't bother checking the keys, they won't be what the user can edit in this case)
+					if (MapHelper.IsValidIndex(MapSparseIndex))
+					{
+						const uint8* MapValueData = MapHelper.GetValuePtr(MapSparseIndex);						
+						UObject* ValueElement = MapValProp->GetObjectPropertyValue(MapValueData);
+						if (ValueElement != nullptr && ValueElement->GetFName() == NativeComponent->GetFName())
+						{
+							return MapValProp;
+						}
+					}
+				}
+			}
+
+			// Finally we should check for set properties
+			for (TFieldIterator<FSetProperty> PropIt(OwnerClass, EFieldIteratorFlags::IncludeSuper); PropIt; ++PropIt)
+			{
+				FSetProperty* TestProperty = *PropIt;
+				void* SetPropInstAddress = TestProperty->ContainerPtrToValuePtr<void>(OwnerCDO);
+
+				// Ensure that this property is valid and that it is marked visible
+				FObjectProperty* SetValProp = CastField<FObjectProperty>(TestProperty->ElementProp);
+				if ((SetValProp == nullptr) || !SetValProp->PropertyClass->IsChildOf<UActorComponent>() || ((TestProperty->PropertyFlags & CPF_Edit) == 0))
+				{
+					continue;
+				}
+				
+				// For each item in the set
+				FScriptSetHelper SetHelper(TestProperty, SetPropInstAddress);
+				for (int32 i = 0; i < SetHelper.Num(); ++i)
+				{
+					if (SetHelper.IsValidIndex(i))
+					{
+						const uint8* SetValData = SetHelper.GetElementPtr(i);
+						UObject* SetValueElem = SetValProp->GetObjectPropertyValue(SetValData);
+						
+						if (SetValueElem != nullptr && SetValueElem->GetFName() == NativeComponent->GetFName())
+						{
+							return SetValProp;
+						}
+					}
+				}
 			}
 		}
-	}
+	}	
 
-	return bCanEdit;
+	return nullptr;
 }
 
 bool FComponentEditorUtils::IsValidVariableNameString(const UActorComponent* InComponent, const FString& InString)
diff --git a/Engine/Source/Editor/UnrealEd/Public/Kismet2/ComponentEditorUtils.h b/Engine/Source/Editor/UnrealEd/Public/Kismet2/ComponentEditorUtils.h
--- a/Engine/Source/Editor/UnrealEd/Public/Kismet2/ComponentEditorUtils.h
+++ b/Engine/Source/Editor/UnrealEd/Public/Kismet2/ComponentEditorUtils.h
@@ -18,8 +18,12 @@ public:
 	/** Is the instance component is editable */
 	static bool CanEditComponentInstance(const UActorComponent* ActorComp, const UActorComponent* ParentSceneComp, bool bAllowUserContructionScript);
 
-	/** Tests whether the native component is editable */
-	static bool CanEditNativeComponent(const UActorComponent* NativeComponent);
+	/** 
+	* Test if the native component is editable. If it is, return a valid pointer to it's FProperty
+	* Otherwise, return nullptr. A native component is editable if it is marked as EditAnywhere
+	* via meta data tags or is within an editable property container 
+	*/
+	static FProperty* GetPropertyForEditableNativeComponent(const UActorComponent* NativeComponent);
 
 	/** Test whether or not the given string is a valid variable name string for the given component instance */
 	static bool IsValidVariableNameString(const UActorComponent* InComponent, const FString& InString);
diff --git a/Engine/Source/Runtime/Engine/Private/ActorEditorUtils.cpp b/Engine/Source/Runtime/Engine/Private/ActorEditorUtils.cpp
--- a/Engine/Source/Runtime/Engine/Private/ActorEditorUtils.cpp
+++ b/Engine/Source/Runtime/Engine/Private/ActorEditorUtils.cpp
@@ -53,7 +53,7 @@ namespace FActorEditorUtils
 				if (Component->CreationMethod == EComponentCreationMethod::Native)
 				{
 					// Make sure it's an exposed native component
-					if (FComponentEditorUtils::CanEditNativeComponent(Component))
+					if (FComponentEditorUtils::GetPropertyForEditableNativeComponent(Component))
 					{
 						OutEditableComponents.Add(Component);
 					}
diff --git a/Engine/Source/Runtime/Engine/Private/Components/ActorComponent.cpp b/Engine/Source/Runtime/Engine/Private/Components/ActorComponent.cpp
--- a/Engine/Source/Runtime/Engine/Private/Components/ActorComponent.cpp
+++ b/Engine/Source/Runtime/Engine/Private/Components/ActorComponent.cpp
@@ -1897,7 +1897,7 @@ bool UActorComponent::IsEditableWhenInherited() const
 #if WITH_EDITOR
 		if (CreationMethod == EComponentCreationMethod::Native && !IsTemplate())
 		{
-			bCanEdit = FComponentEditorUtils::CanEditNativeComponent(this);
+			bCanEdit = FComponentEditorUtils::GetPropertyForEditableNativeComponent(this) != nullptr;
 		}
 		else
 #endif
