Prevent 'pending kill' blueprints from being recompiled and reinstanced

diff --git a/Engine/Source/Editor/Kismet/Private/BlueprintCompilationManager.cpp b/Engine/Source/Editor/Kismet/Private/BlueprintCompilationManager.cpp
--- a/Engine/Source/Editor/Kismet/Private/BlueprintCompilationManager.cpp
+++ b/Engine/Source/Editor/Kismet/Private/BlueprintCompilationManager.cpp
@@ -731,6 +731,42 @@ void FBlueprintCompilationManagerImpl::FlushCompilationQueueImpl(bool bSuppressB
 			}
 		}
 
+		/*	Prevent 'pending kill' blueprints from being recompiled. Dependency
+			gathering is currently done for the following reasons:
+			 * Update a caller's called functions when they are recreated
+			 * Update a child type's cached information about its superclass
+			 * Update a child type's class layout when a parent type layout changes
+			 * Update a reader/writers references to member variables when member variables are recreated
+		
+			Pending kill objects do not need these updates and StaticDuplicateObject
+			cannot duplicate them - so they cannot be updated as normal, anyway.
+
+			Ultimately pending kill UBlueprintGeneratedClass instances rely on the GetDerivedClasses/ReparentChild
+			calls in FBlueprintCompileReinstancer() to maintain accurate class layouts so that we 
+			don't leak or scribble memory.
+		*/
+		CurrentlyCompilingBPs.RemoveAll(
+			[](FCompilerData& Data)
+			{ 
+				if(Data.BP->IsPendingKill())
+				{
+					check(!Data.BP->bBeingCompiled);
+					check(Data.BP->CurrentMessageLog == nullptr);
+					if(UPackage* Package = Data.BP->GetOutermost())
+					{
+						Package->SetDirtyFlag(Data.bPackageWasDirty);
+					}
+					if(Data.ResultsLog)
+					{
+						Data.ResultsLog->EndEvent();
+					}
+					Data.BP->bQueuedForCompilation = false;
+					return true;
+				}
+				return false;
+			}
+		);
+
 		BlueprintsToRecompile.Empty();
 		QueuedRequests.Empty();
 
@@ -1853,7 +1889,10 @@ void FBlueprintCompilationManagerImpl::ReinstanceBatch(TArray<FReinstancingJob>&
 			
 			for(UClass* ClassToReinstance : ClassesToReinstanceList)
 			{
-				ClassesToReinstance.Add(ClassToReinstance);
+				if(!ClassToReinstance->IsPendingKill())
+				{
+					ClassesToReinstance.Add(ClassToReinstance);
+				}
 			}
 		}
 		else
@@ -1865,7 +1904,10 @@ void FBlueprintCompilationManagerImpl::ReinstanceBatch(TArray<FReinstancingJob>&
 			
 			for(UClass* ClassToReparent : ClassesToReparentList)
 			{
-				ClassesToReparent.Add(ClassToReparent);
+				if(!ClassToReparent->IsPendingKill())
+				{
+					ClassesToReparent.Add(ClassToReparent);
+				}
 			}
 		}
 	}
diff --git a/Engine/Source/Editor/UnrealEd/Private/Kismet2/KismetReinstanceUtilities.cpp b/Engine/Source/Editor/UnrealEd/Private/Kismet2/KismetReinstanceUtilities.cpp
--- a/Engine/Source/Editor/UnrealEd/Private/Kismet2/KismetReinstanceUtilities.cpp
+++ b/Engine/Source/Editor/UnrealEd/Private/Kismet2/KismetReinstanceUtilities.cpp
@@ -1536,6 +1536,8 @@ UClass* FBlueprintCompileReinstancer::MoveCDOToNewClass(UClass* OwnerClass, cons
 
 	UObject* OldCDO = OwnerClass->ClassDefaultObject;
 	const FName ReinstanceName = MakeUniqueObjectName(GetTransientPackage(), OwnerClass->GetClass(), *(FString(TEXT("REINST_")) + *OwnerClass->GetName()));
+
+	checkf(!OwnerClass->IsPendingKill(), TEXT("%s is PendingKill - will not duplicate successfully"), *(OwnerClass->GetName()));
 	UClass* CopyOfOwnerClass = CastChecked<UClass>(StaticDuplicateObject(OwnerClass, GetTransientPackage(), ReinstanceName, ~RF_Transactional));
 
 	CopyOfOwnerClass->RemoveFromRoot();
