Include both selected actors and selected components in the transaction to avoid inconsistency when undoing

diff --git a/Engine/Source/Editor/Kismet/Private/SSCSEditor.cpp b/Engine/Source/Editor/Kismet/Private/SSCSEditor.cpp
--- a/Engine/Source/Editor/Kismet/Private/SSCSEditor.cpp
+++ b/Engine/Source/Editor/Kismet/Private/SSCSEditor.cpp
@@ -6541,6 +6541,7 @@ void SSCSEditor::OnApplyChangesToBlueprint() const
 			const FScopedTransaction Transaction(LOCTEXT("PushToBlueprintDefaults_Transaction", "Apply Changes to Blueprint"));
 
 			// The component selection state should be maintained
+			GEditor->GetSelectedActors()->Modify();
 			GEditor->GetSelectedComponents()->Modify();
 
 			Actor->Modify();
