Include the reinstancing of the blueprints in the apply to instance transaction
#rnx

diff --git a/Engine/Source/Editor/Kismet/Private/SSCSEditor.cpp b/Engine/Source/Editor/Kismet/Private/SSCSEditor.cpp
--- a/Engine/Source/Editor/Kismet/Private/SSCSEditor.cpp
+++ b/Engine/Source/Editor/Kismet/Private/SSCSEditor.cpp
@@ -6575,15 +6575,15 @@ void SSCSEditor::OnApplyChangesToBlueprint() const
 					}
 				}
 			}
-		}
 
-		// Compile the BP outside of the transaction
- 		if (NumChangedProperties > 0)
- 		{
-			FBlueprintEditorUtils::MarkBlueprintAsModified(Blueprint);
-			FKismetEditorUtilities::CompileBlueprint(Blueprint);
-			RestoreSelectedInstanceComponent.Restore();
- 		}
+			// Compile the BP outside of the transaction
+			if (NumChangedProperties > 0)
+			{
+				FBlueprintEditorUtils::MarkBlueprintAsModified(Blueprint);
+				FKismetEditorUtilities::CompileBlueprint(Blueprint);
+				RestoreSelectedInstanceComponent.Restore();
+			}
+		}
 
 		// Set up a notification record to indicate success/failure
 		FNotificationInfo NotificationInfo(FText::GetEmpty());
