Fix for Prev-Next Buttons don't work on Blueprint Diff tool when a diff graph is active

diff --git a/Engine/Source/Editor/Kismet/Private/SBlueprintDiff.cpp b/Engine/Source/Editor/Kismet/Private/SBlueprintDiff.cpp
--- a/Engine/Source/Editor/Kismet/Private/SBlueprintDiff.cpp
+++ b/Engine/Source/Editor/Kismet/Private/SBlueprintDiff.cpp
@@ -1020,6 +1020,17 @@ void SBlueprintDiff::Construct( const FArguments& InArgs)
 			[
 				SNew(SOverlay)
 				+ SOverlay::Slot()
+				.VAlign(VAlign_Top)
+				[
+					TopRevisionInfoWidget.ToSharedRef()		
+				]
+				+ SOverlay::Slot()
+				.VAlign(VAlign_Top)
+				.Padding(0.0f, 6.0f, 0.0f, 4.0f)
+				[
+					GraphToolBarWidget.ToSharedRef()
+				]
+				+ SOverlay::Slot()
 				[
 					SNew(SVerticalBox)
 					+SVerticalBox::Slot()
@@ -1057,17 +1068,6 @@ void SBlueprintDiff::Construct( const FArguments& InArgs)
 						]
 					]
 				]
-				+ SOverlay::Slot()
-				.VAlign(VAlign_Top)
-				[
-					TopRevisionInfoWidget.ToSharedRef()		
-				]
-				+ SOverlay::Slot()
-				.VAlign(VAlign_Top)
-				.Padding(0.0f, 6.0f, 0.0f, 4.0f)
-				[
-					GraphToolBarWidget.ToSharedRef()
-				]
 			]
 		];
 
