Fix Find-in-Blueprint search thread deadlock caused by improper handling of unloaded LevelScript Blueprint assets.

#rnx

diff --git a/Engine/Source/Editor/Kismet/Private/BlueprintAssetHandler.cpp b/Engine/Source/Editor/Kismet/Private/BlueprintAssetHandler.cpp
--- a/Engine/Source/Editor/Kismet/Private/BlueprintAssetHandler.cpp
+++ b/Engine/Source/Editor/Kismet/Private/BlueprintAssetHandler.cpp
@@ -17,9 +17,8 @@ class FLevelBlueprintAssetHandler : public IBlueprintAssetHandler
 
 	virtual bool AssetContainsBlueprint(const FAssetData& InAssetData) const
 	{
-		static FName FiB("FiB");
 		// Worlds are only considered to contain a blueprint if they have FiB data
-		return InAssetData.TagsAndValues.Contains(FiB) || InAssetData.TagsAndValues.Contains(FBlueprintTags::FindInBlueprintsData);
+		return InAssetData.TagsAndValues.Contains(FBlueprintTags::FindInBlueprintsData) || InAssetData.TagsAndValues.Contains(FBlueprintTags::UnversionedFindInBlueprintsData);
 	}
 };
 
diff --git a/Engine/Source/Editor/Kismet/Private/FindInBlueprintManager.cpp b/Engine/Source/Editor/Kismet/Private/FindInBlueprintManager.cpp
--- a/Engine/Source/Editor/Kismet/Private/FindInBlueprintManager.cpp
+++ b/Engine/Source/Editor/Kismet/Private/FindInBlueprintManager.cpp
@@ -2097,36 +2097,41 @@ void FFindInBlueprintSearchManager::OnAssetAdded(const FAssetData& InAssetData)
 	}
 	else if (Handler->AssetContainsBlueprint(InAssetData))
 	{
-		// Check first for versioned FiB data (latest codepath)
-		FAssetDataTagMapSharedView::FFindTagResult Result = InAssetData.TagsAndValues.FindTag(FBlueprintTags::FindInBlueprintsData);
-		if(Result.IsSet())
+		AddUnloadedBlueprintSearchMetadata(InAssetData);
+	}
+}
+
+void FFindInBlueprintSearchManager::AddUnloadedBlueprintSearchMetadata(const FAssetData& InAssetData)
+{
+	// Check first for versioned FiB data (latest codepath)
+	FAssetDataTagMapSharedView::FFindTagResult Result = InAssetData.TagsAndValues.FindTag(FBlueprintTags::FindInBlueprintsData);
+	if (Result.IsSet())
+	{
+		const FString& FiBVersionedSearchData = Result.GetValue();
+		if (FiBVersionedSearchData.Len() == 0)
 		{
-			const FString& FiBVersionedSearchData = Result.GetValue();
-			if (FiBVersionedSearchData.Len() == 0)
-			{
-				UnindexedAssets.Add(InAssetData.ObjectPath);
-			}
-			else
-			{
-				ExtractUnloadedFiBData(InAssetData, FiBVersionedSearchData, EFiBVersion::FIB_VER_NONE);
-			}
+			UnindexedAssets.Add(InAssetData.ObjectPath);
 		}
 		else
 		{
-			// Check for legacy (unversioned) FiB data
-			FAssetDataTagMapSharedView::FFindTagResult ResultLegacy = InAssetData.TagsAndValues.FindTag("FiB");
-			if (ResultLegacy.IsSet())
-			{
-				ExtractUnloadedFiBData(InAssetData, ResultLegacy.GetValue(), EFiBVersion::FIB_VER_BASE);
-			}
-			// The asset has no FiB data, keep track of it so we can inform the user
-			else
-			{
-				UnindexedAssets.Add(InAssetData.ObjectPath);
-			}
-
+			ExtractUnloadedFiBData(InAssetData, FiBVersionedSearchData, EFiBVersion::FIB_VER_NONE);
 		}
 	}
+	else
+	{
+		// Check for legacy (unversioned) FiB data
+		FAssetDataTagMapSharedView::FFindTagResult ResultLegacy = InAssetData.TagsAndValues.FindTag(FBlueprintTags::UnversionedFindInBlueprintsData);
+		if (ResultLegacy.IsSet())
+		{
+			ExtractUnloadedFiBData(InAssetData, ResultLegacy.GetValue(), EFiBVersion::FIB_VER_BASE);
+		}
+		// The asset has no FiB data, keep track of it so we can inform the user
+		else
+		{
+			UnindexedAssets.Add(InAssetData.ObjectPath);
+		}
+
+	}
 }
 
 void FFindInBlueprintSearchManager::ExtractUnloadedFiBData(const FAssetData& InAssetData, const FString& InFiBData, EFiBVersion InFiBDataVersion)
@@ -2134,9 +2139,9 @@ void FFindInBlueprintSearchManager::ExtractUnloadedFiBData(const FAssetData& InA
 	CSV_SCOPED_TIMING_STAT(FindInBlueprint, ExtractUnloadedFiBData);
 	CSV_CUSTOM_STAT(FindInBlueprint, ExtractUnloadedCountThisFrame, 1, ECsvCustomStatOp::Accumulate);
 
-	// Check whether this asset has already had its search data cached.
+	// Check whether this asset has already had its search data cached. If marked for deletion, we will replace it with a new entry.
 	FSearchData SearchData = GetSearchDataForAssetPath(InAssetData.ObjectPath);
-	if (SearchData.IsValid())
+	if (SearchData.IsValid() && !SearchData.IsMarkedForDeletion())
 	{
 		return;
 	}
@@ -2347,7 +2352,71 @@ void FFindInBlueprintSearchManager::OnAssetLoaded(UObject* InAsset)
 
 void FFindInBlueprintSearchManager::OnBlueprintUnloaded(UBlueprint* InBlueprint)
 {
-	RemoveBlueprintByPath(*InBlueprint->GetPathName());
+	if (InBlueprint)
+	{
+		// Transient objects likely represent trashed or temporary Blueprint assets; we're not going to try and re-index those.
+		if (InBlueprint->HasAnyFlags(RF_Transient))
+		{
+			FScopeLock ScopeLock(&SafeModifyCacheCriticalSection);
+
+			// Invalidate any existing database entry that matches the reference. It's ok if we don't find a match here.
+			for (int32 SearchIdx = 0; SearchIdx < SearchArray.Num(); ++SearchIdx)
+			{
+				if (SearchArray[SearchIdx].Blueprint.Get() == InBlueprint)
+				{
+					// Remove it from the lookup table so it can no longer be found indirectly.
+					SearchMap.Remove(SearchArray[SearchIdx].AssetPath);
+
+					// Stale entries are flagged to be removed later, which happens after a GC pass.
+					SearchArray[SearchIdx].StateFlags |= ESearchDataStateFlags::WasRemoved;
+					break;
+				}
+			}
+
+			return;
+		}
+
+		// Reloaded assets will be re-indexed on load through the registry delegates; there's no need to handle any removal here.
+		if (InBlueprint->HasAnyFlags(RF_NewerVersionExists))
+		{
+			return;
+		}
+
+		// Otherwise, the Blueprint is about to be unloaded. We need to re-index it as an unloaded asset.
+		if(const UObject* AssetObject = GetAssetObject(InBlueprint))
+		{
+			// Mark any existing entry for deletion. This will allow the entry to be updated below.
+			const FName AssetPath = *AssetObject->GetPathName();
+			RemoveBlueprintByPath(AssetPath);
+
+			// Add or update an existing entry to one that represents the data for the asset on disk, and re-index it.
+			if (ensure(AssetRegistryModule != nullptr))
+			{
+				const bool bIncludeOnlyOnDiskAssets = true;
+				FAssetData AssetData = AssetRegistryModule->Get().GetAssetByObjectPath(AssetPath, bIncludeOnlyOnDiskAssets);
+				if (AssetData.IsValid() && AssetData.IsAssetLoaded())
+				{
+					// Re-scan the asset file on disk to ensure that the updated entry will be based on the serialized FiB tag.
+					const FString PackageName = AssetData.PackageName.ToString();
+					if (FPackageName::IsValidLongPackageName(PackageName))
+					{
+						FString PackageFilename;
+						if (FPackageName::DoesPackageExist(PackageName, nullptr, &PackageFilename))
+						{
+							TArray<FString> FilesToScan = { PackageFilename };
+							AssetRegistryModule->Get().ScanModifiedAssetFiles(FilesToScan);
+
+							AssetData = AssetRegistryModule->Get().GetAssetByObjectPath(AssetPath, bIncludeOnlyOnDiskAssets);
+							if (AssetData.IsValid())
+							{
+								AddUnloadedBlueprintSearchMetadata(AssetData);
+							}
+						}
+					}
+				}
+			}
+		}
+	}
 }
 
 void FFindInBlueprintSearchManager::OnHotReload(bool bWasTriggeredAutomatically)
@@ -2746,37 +2815,22 @@ void FFindInBlueprintSearchManager::CleanCache()
 	{
 		// Here it builds the new map/array, clean of deleted content.
 
-		// If the database item is not marked for deletion and not pending kill (if loaded), keep it in the database
-		if (!SearchArray[SearchValuePair.Value].IsMarkedForDeletion() && !(SearchArray[SearchValuePair.Value].Blueprint.IsValid() && SearchArray[SearchValuePair.Value].Blueprint->IsPendingKill()))
+		// If the database item is invalid, marked for deletion, stale or pending kill (if loaded), remove it from the database
+		const bool bEvenIfPendingKill = true;
+		const FSearchData& SearchData = SearchArray[SearchValuePair.Value];
+		if (!SearchData.IsValid()
+			|| SearchData.IsMarkedForDeletion()
+			|| SearchData.Blueprint.IsStale()
+			|| (SearchData.Blueprint.IsValid(bEvenIfPendingKill) && SearchData.Blueprint->IsPendingKill()))
 		{
-			// Build the new map/array
-			NewSearchMap.Add(SearchValuePair.Key, NewSearchArray.Add(MoveTemp(SearchArray[SearchValuePair.Value])) );
+			// Also remove it from the list of loaded assets that require indexing
+			PendingAssets.Remove(SearchData.AssetPath);
+			UnindexedAssets.Remove(SearchData.AssetPath);
 		}
 		else
 		{
-			// Remove it from the list of loaded assets that require indexing
-			PendingAssets.Remove(SearchArray[SearchValuePair.Value].AssetPath);
-
-			// Level Blueprints are destroyed when you open a new level, we need to re-add it as an unloaded asset so long as they were not marked for deletion
-			if (!SearchArray[SearchValuePair.Value].IsMarkedForDeletion() && FModuleManager::Get().IsModuleLoaded(TEXT("AssetRegistry")))
-			{
-				SearchArray[SearchValuePair.Value].Blueprint = nullptr;
-
-				AssetRegistryModule = &FModuleManager::LoadModuleChecked<FAssetRegistryModule>(TEXT("AssetRegistry"));
-
-				// The asset was not user deleted, so this should usually find the asset. New levels can be deleted if they were not saved
-				FAssetData AssetData = AssetRegistryModule->Get().GetAssetByObjectPath(SearchArray[SearchValuePair.Value].AssetPath);
-				if(AssetData.IsValid())
-				{
-					FAssetDataTagMapSharedView::FFindTagResult ResultLegacy = AssetData.TagsAndValues.FindTag("FiB");
-					if (ResultLegacy.IsSet())
-					{
-						SearchArray[SearchValuePair.Value].Value = ResultLegacy.GetValue();
-					}
-					// Build the new map/array
-					NewSearchMap.Add(SearchValuePair.Key, NewSearchArray.Add(SearchArray[SearchValuePair.Value]) );
-				}
-			}
+			// Build the new map/array
+			NewSearchMap.Add(SearchValuePair.Key, NewSearchArray.Add(SearchData));
 		}
 	}
 
diff --git a/Engine/Source/Editor/Kismet/Public/FindInBlueprintManager.h b/Engine/Source/Editor/Kismet/Public/FindInBlueprintManager.h
--- a/Engine/Source/Editor/Kismet/Public/FindInBlueprintManager.h
+++ b/Engine/Source/Editor/Kismet/Public/FindInBlueprintManager.h
@@ -701,7 +701,10 @@ private:
 	/** Removes a Blueprint from being managed by the FiB system by passing in the UBlueprint's path */
 	void RemoveBlueprintByPath(FName InPath);
 
-	/** Begins the process of extracting unloaded FiB data */
+	/** Adds a new search database entry for unloaded asset data */
+	void AddUnloadedBlueprintSearchMetadata(const FAssetData& InAssetData);
+
+	/** Begins the process of extracting FiB data from an unloaded asset */
 	void ExtractUnloadedFiBData(const FAssetData& InAssetData, const FString& InFiBData, EFiBVersion InFiBDataVersion);
 
 	/** Determines the global find results tab label */
diff --git a/Engine/Source/Editor/UnrealEd/Private/EditorServer.cpp b/Engine/Source/Editor/UnrealEd/Private/EditorServer.cpp
--- a/Engine/Source/Editor/UnrealEd/Private/EditorServer.cpp
+++ b/Engine/Source/Editor/UnrealEd/Private/EditorServer.cpp
@@ -34,6 +34,7 @@
 #include "Engine/EngineTypes.h"
 #include "Engine/EngineBaseTypes.h"
 #include "Engine/Level.h"
+#include "Engine/LevelScriptBlueprint.h"
 #include "Components/ActorComponent.h"
 #include "Components/SceneComponent.h"
 #include "GameFramework/Actor.h"
@@ -2080,10 +2081,20 @@ void UEditorEngine::EditorDestroyWorld( FWorldContext & Context, const FText& Cl
 	FEditorSupportDelegates::PrepareToCleanseEditorObject.Broadcast(ContextWorld);
 	for (ULevel* Level : ContextWorld->GetLevels())
 	{
-		UWorld* LevelWorld = Level->GetTypedOuter<UWorld>();
-		if (ensureAlways(LevelWorld) && LevelWorld != ContextWorld && LevelWorld != NewWorld)
+		if (ensureAlways(Level))
 		{
-			FEditorSupportDelegates::PrepareToCleanseEditorObject.Broadcast(LevelWorld);
+			const bool bDontCreate = true;
+			if (ULevelScriptBlueprint* LSBP = Level->GetLevelScriptBlueprint(bDontCreate))
+			{
+				// Signals that the associated LSBP is about to be unloaded.
+				LSBP->ClearEditorReferences();
+			}
+
+			UWorld* LevelWorld = Level->GetTypedOuter<UWorld>();
+			if (ensureAlways(LevelWorld) && LevelWorld != ContextWorld && LevelWorld != NewWorld)
+			{
+				FEditorSupportDelegates::PrepareToCleanseEditorObject.Broadcast(LevelWorld);
+			}
 		}
 	}
 
diff --git a/Engine/Source/Runtime/CoreUObject/Private/Blueprint/BlueprintSupport.cpp b/Engine/Source/Runtime/CoreUObject/Private/Blueprint/BlueprintSupport.cpp
--- a/Engine/Source/Runtime/CoreUObject/Private/Blueprint/BlueprintSupport.cpp
+++ b/Engine/Source/Runtime/CoreUObject/Private/Blueprint/BlueprintSupport.cpp
@@ -38,6 +38,7 @@ const FName FBlueprintTags::BlueprintDisplayName(TEXT("BlueprintDisplayName"));
 const FName FBlueprintTags::IsDataOnly(TEXT("IsDataOnly"));
 const FName FBlueprintTags::ImplementedInterfaces(TEXT("ImplementedInterfaces"));
 const FName FBlueprintTags::FindInBlueprintsData(TEXT("FiBData"));
+const FName FBlueprintTags::UnversionedFindInBlueprintsData(TEXT("FiB"));
 const FName FBlueprintTags::NumReplicatedProperties(TEXT("NumReplicatedProperties"));
 const FName FBlueprintTags::NumNativeComponents(TEXT("NativeComponents"));
 const FName FBlueprintTags::NumBlueprintComponents(TEXT("BlueprintComponents"));
diff --git a/Engine/Source/Runtime/CoreUObject/Public/Blueprint/BlueprintSupport.h b/Engine/Source/Runtime/CoreUObject/Public/Blueprint/BlueprintSupport.h
--- a/Engine/Source/Runtime/CoreUObject/Public/Blueprint/BlueprintSupport.h
+++ b/Engine/Source/Runtime/CoreUObject/Public/Blueprint/BlueprintSupport.h
@@ -36,6 +36,8 @@ struct COREUOBJECT_API FBlueprintTags
 	static const FName ImplementedInterfaces;
 	/** Very large string used to store find in blueprint data for the editor */
 	static const FName FindInBlueprintsData;
+	/** (Deprecated) Legacy tag that was initially used to store find in blueprint data for the editor */
+	static const FName UnversionedFindInBlueprintsData;
 	/** Number of replicated properties */
 	static const FName NumReplicatedProperties;
 	/** Number of native components */
