Add null check when determining if we can focus on a node to prevent crash

#rnx

diff --git a/Engine/Source/Editor/Kismet/Private/SMyBlueprint.cpp b/Engine/Source/Editor/Kismet/Private/SMyBlueprint.cpp
--- a/Engine/Source/Editor/Kismet/Private/SMyBlueprint.cpp
+++ b/Engine/Source/Editor/Kismet/Private/SMyBlueprint.cpp
@@ -2249,11 +2249,13 @@ bool SMyBlueprint::CanFocusOnNode() const
 		UBlueprint* BlueprintObj = GetBlueprintObj();		
 		UFunction* OverrideFunc = nullptr;
 		UClass* const OverrideFuncClass = FBlueprintEditorUtils::GetOverrideFunctionClass(BlueprintObj, GraphAction->FuncName, &OverrideFunc);
-		UEdGraph* EventGraph = FBlueprintEditorUtils::FindEventGraph(BlueprintObj);
 
-		// Add to event graph
-		FName EventName = OverrideFunc->GetFName();
-		ExistingNode = FBlueprintEditorUtils::FindOverrideForFunction(BlueprintObj, OverrideFuncClass, EventName);
+		if (OverrideFunc)
+		{
+			// Add to event graph
+			FName EventName = OverrideFunc->GetFName();
+			ExistingNode = FBlueprintEditorUtils::FindOverrideForFunction(BlueprintObj, OverrideFuncClass, EventName);
+		}
 	}
 
 	return (EventAction && EventAction->NodeTemplate) || (InputAction && InputAction->NodeTemplate) || ExistingNode;
@@ -2274,13 +2276,15 @@ void SMyBlueprint::OnFocusNode()
 		UBlueprint* BlueprintObj = GetBlueprintObj();
 		UFunction* OverrideFunc = nullptr;
 		UClass* const OverrideFuncClass = FBlueprintEditorUtils::GetOverrideFunctionClass(BlueprintObj, GraphAction->FuncName, &OverrideFunc);
-		UEdGraph* EventGraph = FBlueprintEditorUtils::FindEventGraph(BlueprintObj);
 
-		// Add to event graph
-		FName EventName = OverrideFunc->GetFName();
-		if (UK2Node_Event* ExistingNode = FBlueprintEditorUtils::FindOverrideForFunction(BlueprintObj, OverrideFuncClass, EventName))
+		if (OverrideFunc)
 		{
-			FKismetEditorUtilities::BringKismetToFocusAttentionOnObject(ExistingNode);
+			// Add to event graph
+			FName EventName = OverrideFunc->GetFName();
+			if (UK2Node_Event* ExistingNode = FBlueprintEditorUtils::FindOverrideForFunction(BlueprintObj, OverrideFuncClass, EventName))
+			{
+				FKismetEditorUtilities::BringKismetToFocusAttentionOnObject(ExistingNode);
+			}
 		}
 	}
 }
