Remove inaccurate comment
#rnx

diff --git a/Engine/Source/Editor/Kismet/Private/SSCSEditor.cpp b/Engine/Source/Editor/Kismet/Private/SSCSEditor.cpp
--- a/Engine/Source/Editor/Kismet/Private/SSCSEditor.cpp
+++ b/Engine/Source/Editor/Kismet/Private/SSCSEditor.cpp
@@ -6576,7 +6576,6 @@ void SSCSEditor::OnApplyChangesToBlueprint() const
 				}
 			}
 
-			// Compile the BP outside of the transaction
 			if (NumChangedProperties > 0)
 			{
 				FBlueprintEditorUtils::MarkBlueprintAsModified(Blueprint);
