Fix crash in BP editor

I'm not sure how to reproduce the crash, but I reached FBlueprintEditor::CanAddNewLocalVariable with a null FocusedGraphEdPtr after I canceled an "add component" menu. Seems like checking the pointer can't hurt.

#rnx

diff --git a/Engine/Source/Editor/Kismet/Private/BlueprintEditor.cpp b/Engine/Source/Editor/Kismet/Private/BlueprintEditor.cpp
--- a/Engine/Source/Editor/Kismet/Private/BlueprintEditor.cpp
+++ b/Engine/Source/Editor/Kismet/Private/BlueprintEditor.cpp
@@ -8169,8 +8169,13 @@ bool FBlueprintEditor::CanAddNewLocalVariable() const
 {
 	if (InEditingMode())
 	{
-		UEdGraph* TargetGraph = FBlueprintEditorUtils::GetTopLevelGraph(FocusedGraphEdPtr.Pin()->GetCurrentGraph());
+		TSharedPtr<SGraphEditor> FocusedGraphEd = FocusedGraphEdPtr.Pin();
+		if (!FocusedGraphEd.IsValid())
+		{
+			return false;
+		}
 
+		UEdGraph* TargetGraph = FBlueprintEditorUtils::GetTopLevelGraph(FocusedGraphEd->GetCurrentGraph());
 		return TargetGraph->GetSchema()->GetGraphType(TargetGraph) == GT_Function;
 	}
 	return false;
