Fix potential for a deadlock during a Blueprint search while asynchronously indexing loaded assets after editor startup.

#rnx

diff --git a/Engine/Source/Editor/Kismet/Private/FindInBlueprintManager.cpp b/Engine/Source/Editor/Kismet/Private/FindInBlueprintManager.cpp
--- a/Engine/Source/Editor/Kismet/Private/FindInBlueprintManager.cpp
+++ b/Engine/Source/Editor/Kismet/Private/FindInBlueprintManager.cpp
@@ -2569,8 +2569,16 @@ FSearchData FFindInBlueprintSearchManager::GetNextSearchDataForQuery(FActiveSear
 			}
 			else if (AssetPath == FirstAssetPath)
 			{
-				// Yield to allow indexing to progress a bit further.
-				FPlatformProcess::Sleep(0.1f);
+				if (bIsPausing)
+				{
+					// Wait here until the search query is unpaused.
+					BlockSearchQueryIfPaused();
+				}
+				else
+				{
+					// Yield to allow indexing to progress a bit further.
+					FPlatformProcess::Sleep(0.1f);
+				}
 
 				// Check for a new entry in case the cache has grown in size.
 				SearchData = GetNextSearchDataForQuery(SearchQuery, /*bCheckDeferredList = */false);
@@ -2613,6 +2621,18 @@ FSearchData FFindInBlueprintSearchManager::GetNextSearchDataForQuery(FActiveSear
 	return SearchData;
 }
 
+void FFindInBlueprintSearchManager::BlockSearchQueryIfPaused()
+{
+	// Check if the thread has been told to pause, this occurs for the Garbage Collector and for saving to disk
+	if (bIsPausing == true)
+	{
+		// Pause all searching, the GC is running and we will also be saving the database
+		ActiveSearchCounter.Decrement();
+		FScopeLock ScopeLock(&PauseThreadsCriticalSection);
+		ActiveSearchCounter.Increment();
+	}
+}
+
 void FFindInBlueprintSearchManager::BeginSearchQuery(const FStreamSearch* InSearchOriginator)
 {
 	if (AssetRegistryModule == nullptr)
@@ -2636,14 +2656,8 @@ void FFindInBlueprintSearchManager::BeginSearchQuery(const FStreamSearch* InSear
 
 bool FFindInBlueprintSearchManager::ContinueSearchQuery(const FStreamSearch* InSearchOriginator, FSearchData& OutSearchData)
 {
-	// Check if the thread has been told to pause, this occurs for the Garbage Collector and for saving to disk
-	if (bIsPausing == true)
-	{
-		// Pause all searching, the GC is running and we will also be saving the database
-		ActiveSearchCounter.Decrement();
-		FScopeLock ScopeLock(&PauseThreadsCriticalSection);
-		ActiveSearchCounter.Increment();
-	}
+	// If paused, wait here until searching is resumed.
+	BlockSearchQueryIfPaused();
 
 	FActiveSearchQueryPtr SearchQuery = FindSearchQuery(InSearchOriginator);
 	if (!SearchQuery.IsValid())
@@ -2797,11 +2811,19 @@ void FFindInBlueprintSearchManager::CleanCache()
 	 	const FStreamSearch* ActiveSearch = It.Key();
 	 	check(ActiveSearch);
 	 	{
-			FSearchData SearchData;
-	 		ContinueSearchQuery(ActiveSearch, SearchData);
-
-			FName CachePath = SearchData.AssetPath;
-	 		CacheQueries.Add(ActiveSearch, CachePath);
+			FActiveSearchQueryPtr SearchQuery = FindSearchQuery(ActiveSearch);
+			if (SearchQuery.IsValid())
+			{
+				// Don't check the deferred list here; in this case we are only fixing up the query's index position (state) if
+				// we haven't yet reached the end of the database, and we don't care about assets that are waiting to be indexed.
+				const bool bCheckDeferredList = false;
+				FSearchData SearchData = GetNextSearchDataForQuery(SearchQuery, bCheckDeferredList);
+				if (SearchData.IsValid())
+				{
+					FName CachePath = SearchData.AssetPath;
+					CacheQueries.Add(ActiveSearch, CachePath);
+				}
+			}
 	 	}
 	}
 
diff --git a/Engine/Source/Editor/Kismet/Public/FindInBlueprintManager.h b/Engine/Source/Editor/Kismet/Public/FindInBlueprintManager.h
--- a/Engine/Source/Editor/Kismet/Public/FindInBlueprintManager.h
+++ b/Engine/Source/Editor/Kismet/Public/FindInBlueprintManager.h
@@ -742,6 +742,9 @@ protected:
 	/** Returns the next pending search data for the given query and advances the index to the next entry */
 	FSearchData GetNextSearchDataForQuery(FActiveSearchQueryPtr InSearchQueryPtr, bool bCheckDeferredList);
 
+	/** If searches are paused, blocks the calling thread until searching is resumed */
+	void BlockSearchQueryIfPaused();
+
 private:
 	/** Maps the Blueprint paths to their index in the SearchArray */
 	TMap<FName, int32> SearchMap;
