Fix subcomponents not appearing in the details panel when added as an instance until deselecting and reselecting the actor.

diff --git a/Engine/Source/Editor/Kismet/Private/SSCSEditor.cpp b/Engine/Source/Editor/Kismet/Private/SSCSEditor.cpp
--- a/Engine/Source/Editor/Kismet/Private/SSCSEditor.cpp
+++ b/Engine/Source/Editor/Kismet/Private/SSCSEditor.cpp
@@ -4356,6 +4356,8 @@ UActorComponent* SSCSEditor::PerformComboAddClass(TSubclassOf<UActorComponent> C
 			// As the SCS splits up the scene and actor components, can now add directly
 			NewComponent = AddNewComponent(NewClass, AssetOverride);
 		}
+
+		UpdateTree();
 	}
 
 	return NewComponent;
