Copy up from DS to Main

diff --git a/Engine/Source/Developer/MessageLog/Private/UserInterface/SMessageLogListing.cpp b/Engine/Source/Developer/MessageLog/Private/UserInterface/SMessageLogListing.cpp
--- a/Engine/Source/Developer/MessageLog/Private/UserInterface/SMessageLogListing.cpp
+++ b/Engine/Source/Developer/MessageLog/Private/UserInterface/SMessageLogListing.cpp
@@ -69,6 +69,7 @@ void SMessageLogListing::Construct( const FArguments& InArgs, const TSharedRef<
 										.OnSelectionChanged(this, &SMessageLogListing::OnLineSelectionChanged)
 										.ExternalScrollbar(ScrollBar)
 										.ItemHeight(24.0f)
+										.ConsumeMouseWheel(EConsumeMouseWheel::Always)
 									]
 								]
 								+SHorizontalBox::Slot()
diff --git a/Engine/Source/Editor/Kismet/Private/BlueprintCompilationManager.cpp b/Engine/Source/Editor/Kismet/Private/BlueprintCompilationManager.cpp
--- a/Engine/Source/Editor/Kismet/Private/BlueprintCompilationManager.cpp
+++ b/Engine/Source/Editor/Kismet/Private/BlueprintCompilationManager.cpp
@@ -106,6 +106,14 @@ struct FBPCompileRequestInternal
 	FBPCompilationManagerCPPResults CPPResults;
 };
 
+enum class EReparentClassOptions
+{
+	None = 0x0,
+
+	ReplaceReferencesToOldClasses = 0x1,
+};
+ENUM_CLASS_FLAGS(EReparentClassOptions)
+
 struct FBlueprintCompilationManagerImpl : public FGCObject
 {
 	FBlueprintCompilationManagerImpl();
@@ -126,7 +134,7 @@ struct FBlueprintCompilationManagerImpl : public FGCObject
 	bool IsGeneratedClassLayoutReady() const;
 	void GetDefaultValue(const UClass* ForClass, const FProperty* Property, FString& OutDefaultValueAsString) const;
 
-	static void ReparentHierarchies(const TMap<UClass*, UClass*>& OldClassToNewClass);
+	static void ReparentHierarchies(const TMap<UClass*, UClass*>& OldClassToNewClass, EReparentClassOptions Options);
 	static void BuildDSOMap(UObject* OldObject, UObject* NewObject, TMap<UObject*, UObject*>& OutOldToNewDSO);
 	static void ReinstanceBatch(TArray<FReinstancingJob>& Reinstancers, TMap< UClass*, UClass* >& InOutOldToNewClassMap, FUObjectSerializeContext* InLoadContext);
 	static UClass* FastGenerateSkeletonClass(UBlueprint* BP, FKismetCompilerContext& CompilerContext, bool bIsSkeletonOnly, TArray<FSkeletonFixupData>& OutSkeletonFixupData);
@@ -171,16 +179,22 @@ struct FBlueprintCompilationManagerImpl : public FGCObject
 // free function that we use to cross a module boundary (from CoreUObject to here)
 void FlushReinstancingQueueImplWrapper();
 void MoveSkelCDOAside(UClass* Class, TMap<UClass*, UClass*>& OldToNewMap);
+void ReparentHierarchiesWrapper(const TMap<UClass*, UClass*>& OldToNewMap)
+{
+	FBlueprintCompilationManagerImpl::ReparentHierarchies(OldToNewMap, EReparentClassOptions::ReplaceReferencesToOldClasses);
+}
 
 FBlueprintCompilationManagerImpl::FBlueprintCompilationManagerImpl()
 {
 	FBlueprintSupport::SetFlushReinstancingQueueFPtr(&FlushReinstancingQueueImplWrapper);
+	FBlueprintSupport::SetClassReparentingFPtr(&ReparentHierarchiesWrapper);
 	bGeneratedClassLayoutReady = true;
 }
 
 FBlueprintCompilationManagerImpl::~FBlueprintCompilationManagerImpl() 
 { 
 	FBlueprintSupport::SetFlushReinstancingQueueFPtr(nullptr); 
+	FBlueprintSupport::SetClassReparentingFPtr(nullptr);
 }
 
 void FBlueprintCompilationManagerImpl::AddReferencedObjects(FReferenceCollector& Collector)
@@ -1678,8 +1692,10 @@ void FBlueprintCompilationManagerImpl::GetDefaultValue(const UClass* ForClass, c
 	}
 }
 
-void FBlueprintCompilationManagerImpl::ReparentHierarchies(const TMap<UClass*, UClass*>& OldToNewClasses)
+void FBlueprintCompilationManagerImpl::ReparentHierarchies(const TMap<UClass*, UClass*>& OldToNewClasses, EReparentClassOptions Options)
 {
+	const bool bReplaceReferencesToOldClasses = (Options & EReparentClassOptions::ReplaceReferencesToOldClasses) != EReparentClassOptions::None;
+
 	// something has decided to replace instances of a class. We need to update all the children of those types:
 	TArray< UClass* > ClassesOrdered;
 	// Map used to distinguish between new classes and classes that need to be reinstanced (reparented) via a new reinstancer:
@@ -1710,7 +1726,10 @@ void FBlueprintCompilationManagerImpl::ReparentHierarchies(const TMap<UClass*, U
 					continue;
 				}
 
-				Classes.Add(DerivedClass);
+				if (OldToNewClasses.Find(DerivedClass) == nullptr)
+				{
+					Classes.Add(DerivedClass);
+				}
 			}
 		}
 
@@ -1753,20 +1772,24 @@ void FBlueprintCompilationManagerImpl::ReparentHierarchies(const TMap<UClass*, U
 	{
 		UClass* ClassToReinstance = ReinstancingJob.OldToNew.Value;
 
-		ClassToReinstance->ClassConstructor = nullptr;
-		ClassToReinstance->ClassVTableHelperCtorCaller = nullptr;
-		ClassToReinstance->ClassAddReferencedObjects = nullptr;
-		
-		// check to see if we're a direct parent of one of the new classes:
-		UClass* const* NewParent = OldToNewClasses.Find(ClassToReinstance->GetSuperClass());
-		if(NewParent)
+		// We only need to relink if we've reparented the class to a new type:
+		if (ReinstancingJob.Reinstancer.IsValid())
 		{
-			ClassToReinstance->SetSuperStruct(*NewParent);
-		}
+			ClassToReinstance->ClassConstructor = nullptr;
+			ClassToReinstance->ClassVTableHelperCtorCaller = nullptr;
+			ClassToReinstance->ClassAddReferencedObjects = nullptr;
+		
+			// check to see if we're a direct parent of one of the new classes:
+			UClass* const* NewParent = OldToNewClasses.Find(ClassToReinstance->GetSuperClass());
+			if(NewParent)
+			{
+				ClassToReinstance->SetSuperStruct(*NewParent);
+			}
 
-		ClassToReinstance->Bind();
-		ClassToReinstance->ClearFunctionMapsCaches();
-		ClassToReinstance->StaticLink(true);
+			ClassToReinstance->Bind();
+			ClassToReinstance->ClearFunctionMapsCaches();
+			ClassToReinstance->StaticLink(true);
+		}
 
 		OldClassToNewClassIncludingChildren.Add(ReinstancingJob.OldToNew.Key, ClassToReinstance);
 	}
@@ -1776,18 +1799,27 @@ void FBlueprintCompilationManagerImpl::ReparentHierarchies(const TMap<UClass*, U
 
 	// Reinstance (non archetype) instances
 	TMap<UClass*, UClass*> OldClassToNewClassDerivedTypes;
+	
+	if (bReplaceReferencesToOldClasses)
+	{
+		OldClassToNewClassDerivedTypes = OldClassToNewClassIncludingChildren;
+	}
+	
 	for(const FReinstancingJob& ReinstancingJob : Reinstancers)
 	{
 		OldClassToNewClassDerivedTypes.Add(ReinstancingJob.OldToNew);
 	}
 	TGuardValue<bool> ReinstancingGuard(GIsReinstancing, true);
-	FBatchReplaceInstancesOfClassParameters Options;
-	Options.bArchetypesAreUpToDate = true;
+	FBatchReplaceInstancesOfClassParameters BatchOptions;
+	BatchOptions.bArchetypesAreUpToDate = true;
+	BatchOptions.bReplaceReferencesToOldClasses = bReplaceReferencesToOldClasses;
 
 	// Make sure we don't replace old instances that are in the *callers* old to new TMap!
 	TSet<UObject*> OldObjects;
 	for(TPair<UClass*, UClass*> OldToNew : OldClassToNewClassDerivedTypes)
 	{
+		ensure(OldToNew.Value->HasAnyClassFlags(CLASS_TokenStreamAssembled));
+
 		TArray< UObject* > OldObjectsOfType;
 		GetObjectsOfClass(OldToNew.Key, OldObjectsOfType);
 
@@ -1799,10 +1831,10 @@ void FBlueprintCompilationManagerImpl::ReparentHierarchies(const TMap<UClass*, U
 			}
 		}
 	}
-	Options.ObjectsThatShouldUseOldStuff = &OldObjects;
-	Options.InstancesThatShouldUseOldClass = &OldObjects;
+	BatchOptions.ObjectsThatShouldUseOldStuff = &OldObjects;
+	BatchOptions.InstancesThatShouldUseOldClass = &OldObjects;
 
-	FBlueprintCompileReinstancer::BatchReplaceInstancesOfClass( OldClassToNewClassDerivedTypes, Options );
+	FBlueprintCompileReinstancer::BatchReplaceInstancesOfClass( OldClassToNewClassDerivedTypes, BatchOptions);
 }
 
 
@@ -2013,7 +2045,7 @@ void FBlueprintCompilationManagerImpl::ReinstanceBatch(TArray<FReinstancingJob>&
 			}
 		}
 
-		if (UBlueprintGeneratedClass* BPGClass = CastChecked<UBlueprintGeneratedClass>(ReinstancingJob.OldToNew.Value))
+		if (UBlueprintGeneratedClass* BPGClass = Cast<UBlueprintGeneratedClass>(ReinstancingJob.OldToNew.Value))
 		{
 			BPGClass->UpdateCustomPropertyListForPostConstruction();
 
@@ -3104,7 +3136,7 @@ bool FBlueprintCompilationManager::GetDefaultValue(const UClass* ForClass, const
 
 void FBlueprintCompilationManager::ReparentHierarchies(const TMap<UClass*, UClass*>& OldClassToNewClass)
 {
-	FBlueprintCompilationManagerImpl::ReparentHierarchies(OldClassToNewClass);
+	FBlueprintCompilationManagerImpl::ReparentHierarchies(OldClassToNewClass, EReparentClassOptions::None);
 }
 
 void FBlueprintCompilationManager::RegisterCompilerExtension(TSubclassOf<UBlueprint> BlueprintType, UBlueprintCompilerExtension* Extension)
diff --git a/Engine/Source/Editor/UnrealEd/Private/Cooker/PackageNameCache.h b/Engine/Source/Editor/UnrealEd/Private/Cooker/PackageNameCache.h
--- a/Engine/Source/Editor/UnrealEd/Private/Cooker/PackageNameCache.h
+++ b/Engine/Source/Editor/UnrealEd/Private/Cooker/PackageNameCache.h
@@ -97,7 +97,8 @@ bool FPackageNameCache::DoesPackageExist(const FName& PackageName, FString* OutF
 	}
 
 	TArray<FAssetData> Assets;
-	AssetRegistry->GetAssetsByPackageName(PackageName, Assets, /*bIncludeOnlyDiskAssets*/ true);
+	bool bIncludeOnlyDiskAssets = !FPackageName::IsExtraPackage(PackageName.ToString());
+	AssetRegistry->GetAssetsByPackageName(PackageName, Assets, bIncludeOnlyDiskAssets);
 
 	if (Assets.Num() <= 0)
 	{
diff --git a/Engine/Source/Editor/UnrealEd/Private/Kismet2/BlueprintEditorUtils.cpp b/Engine/Source/Editor/UnrealEd/Private/Kismet2/BlueprintEditorUtils.cpp
--- a/Engine/Source/Editor/UnrealEd/Private/Kismet2/BlueprintEditorUtils.cpp
+++ b/Engine/Source/Editor/UnrealEd/Private/Kismet2/BlueprintEditorUtils.cpp
@@ -2129,14 +2129,16 @@ UClass* FBlueprintEditorUtils::GetSkeletonClass(UClass* FromClass)
 
 UClass* FBlueprintEditorUtils::GetMostUpToDateClass(UClass* FromClass)
 {
-	if(!FromClass || FromClass->HasAnyClassFlags(CLASS_Native) || FromClass->bCooked)
+	if (!FromClass || FromClass->HasAnyClassFlags(CLASS_Native) || FromClass->bCooked)
 	{
 		return FromClass;
 	}
 
 	// It's really not safe/coherent to try and dig out the 'right' class. Things that need the 'most up to date'
 	// version of a class should always be looking at the skeleton:
-	return GetSkeletonClass(FromClass);
+	UClass* SkeletonClass = GetSkeletonClass(FromClass);
+
+	return SkeletonClass ? SkeletonClass : FromClass;
 }
 
 const UClass* FBlueprintEditorUtils::GetMostUpToDateClass(const UClass* FromClass)
diff --git a/Engine/Source/Editor/UnrealEd/Private/Kismet2/KismetReinstanceUtilities.cpp b/Engine/Source/Editor/UnrealEd/Private/Kismet2/KismetReinstanceUtilities.cpp
--- a/Engine/Source/Editor/UnrealEd/Private/Kismet2/KismetReinstanceUtilities.cpp
+++ b/Engine/Source/Editor/UnrealEd/Private/Kismet2/KismetReinstanceUtilities.cpp
@@ -348,7 +348,6 @@ FBlueprintCompileReinstancer::FBlueprintCompileReinstancer(UClass* InClassToRein
 
 		// Pull the blueprint that generated this reinstance target, and gather the blueprints that are dependent on it
 		UBlueprint* GeneratingBP = Cast<UBlueprint>(ClassToReinstance->ClassGeneratedBy);
-		check(GeneratingBP || GIsAutomationTesting);
 		if(!IsReinstancingSkeleton() && GeneratingBP)
 		{
 			ClassToReinstanceDefaultValuesCRC = GeneratingBP->CrcLastCompiledCDO;
@@ -1508,7 +1507,7 @@ void FBlueprintCompileReinstancer::BatchReplaceInstancesOfClass(TMap<UClass*, UC
 		return;
 	}
 
-	ReplaceInstancesOfClass_Inner(InOldToNewClassMap, nullptr, Options.ObjectsThatShouldUseOldStuff, false /*bClassObjectReplaced*/, true /*bPreserveRootComponent*/, Options.bArchetypesAreUpToDate, Options.InstancesThatShouldUseOldClass);
+	ReplaceInstancesOfClass_Inner(InOldToNewClassMap, nullptr, Options.ObjectsThatShouldUseOldStuff, false /*bClassObjectReplaced*/, true /*bPreserveRootComponent*/, Options.bArchetypesAreUpToDate, Options.InstancesThatShouldUseOldClass, Options.bReplaceReferencesToOldClasses);
 }
 
 UClass* FBlueprintCompileReinstancer::MoveCDOToNewClass(UClass* OwnerClass, const TMap<UClass*, UClass*>& OldToNewMap, bool bAvoidCDODuplication)
@@ -1886,7 +1885,7 @@ static void ReplaceActorHelper(AActor* OldActor, UClass* OldClass, UObject*& New
 	OldToNewInstanceMap.Add(OldActor, NewActor);
 }
 
-void FBlueprintCompileReinstancer::ReplaceInstancesOfClass_Inner(TMap<UClass*, UClass*>& InOldToNewClassMap, UObject* InOriginalCDO, TSet<UObject*>* ObjectsThatShouldUseOldStuff, bool bClassObjectReplaced, bool bPreserveRootComponent, bool bArchetypesAreUpToDate, const TSet<UObject*>* InstancesThatShouldUseOldClass )
+void FBlueprintCompileReinstancer::ReplaceInstancesOfClass_Inner(TMap<UClass*, UClass*>& InOldToNewClassMap, UObject* InOriginalCDO, TSet<UObject*>* ObjectsThatShouldUseOldStuff, bool bClassObjectReplaced, bool bPreserveRootComponent, bool bArchetypesAreUpToDate, const TSet<UObject*>* InstancesThatShouldUseOldClass, bool bReplaceReferencesToOldClasses)
 {
 	// If there is an original CDO, we are only reinstancing a single class
 	check((InOriginalCDO != nullptr && InOldToNewClassMap.Num() == 1) || InOriginalCDO == nullptr); // (InOldToNewClassMap.Num() > 1 && InOriginalCDO == nullptr) || (InOldToNewClassMap.Num() == 1 && InOriginalCDO != nullptr));
@@ -2198,6 +2197,23 @@ void FBlueprintCompileReinstancer::ReplaceInstancesOfClass_Inner(TMap<UClass*, U
 		OldToNewInstanceMap.Append(ObjectRemappingHelper.ReplacedObjects);
 	}
 
+	if(bReplaceReferencesToOldClasses)
+	{
+		check(ObjectsThatShouldUseOldStuff);
+
+		for (TPair<UClass*, UClass*> OldToNew : InOldToNewClassMap)
+		{
+			ObjectsThatShouldUseOldStuff->Add(OldToNew.Key);
+
+			TArray<UObject*> OldFunctions;
+			GetObjectsWithOuter(OldToNew.Key, OldFunctions);
+			ObjectsThatShouldUseOldStuff->Append(OldFunctions);
+			
+			OldToNewInstanceMap.Add(OldToNew.Key, OldToNew.Value);
+			SourceObjects.Add(OldToNew.Key);
+		}
+	}
+
 	FReplaceReferenceHelper::FindAndReplaceReferences(SourceObjects, ObjectsThatShouldUseOldStuff, ObjectsReplaced, OldToNewInstanceMap, ReinstancedObjectsWeakReferenceMap);
 	
 	for (UObject* Obj : ObjectsReplaced)
diff --git a/Engine/Source/Editor/UnrealEd/Public/Kismet2/KismetReinstanceUtilities.h b/Engine/Source/Editor/UnrealEd/Public/Kismet2/KismetReinstanceUtilities.h
--- a/Engine/Source/Editor/UnrealEd/Public/Kismet2/KismetReinstanceUtilities.h
+++ b/Engine/Source/Editor/UnrealEd/Public/Kismet2/KismetReinstanceUtilities.h
@@ -69,6 +69,14 @@ struct UNREALED_API FBatchReplaceInstancesOfClassParameters
 	const TSet<UObject*>* InstancesThatShouldUseOldClass = nullptr;
 
 	bool bArchetypesAreUpToDate = false; 
+	
+	/** 
+	 * Blueprints reuses its UClass* from compile to compile, but it's more 
+	 * intuitive to just replace a UClass* with a new instance (e.g. from a 
+	 * package reload). This flag tells the reinstancer to replace references
+	 * to old classes with references to new classes.
+	 */
+	bool bReplaceReferencesToOldClasses = false;
 };
 
 class UNREALED_API FBlueprintCompileReinstancer : public TSharedFromThis<FBlueprintCompileReinstancer>, public FGCObject
@@ -223,5 +231,5 @@ protected:
 
 private:
 	/** Handles the work of ReplaceInstancesOfClass, handling both normal replacement of instances and batch */
-	static void ReplaceInstancesOfClass_Inner(TMap<UClass*, UClass*>& InOldToNewClassMap, UObject* InOriginalCDO, TSet<UObject*>* ObjectsThatShouldUseOldStuff = NULL, bool bClassObjectReplaced = false, bool bPreserveRootComponent = true, bool bArchetypesAreUpToDate = false, const TSet<UObject*>* InstancesThatShouldUseOldClass = nullptr);
+	static void ReplaceInstancesOfClass_Inner(TMap<UClass*, UClass*>& InOldToNewClassMap, UObject* InOriginalCDO, TSet<UObject*>* ObjectsThatShouldUseOldStuff = NULL, bool bClassObjectReplaced = false, bool bPreserveRootComponent = true, bool bArchetypesAreUpToDate = false, const TSet<UObject*>* InstancesThatShouldUseOldClass = nullptr, bool bReplaceReferencesToOldClasses = false);
 };
diff --git a/Engine/Source/Programs/UnrealBuildTool/Configuration/ModuleRules.cs b/Engine/Source/Programs/UnrealBuildTool/Configuration/ModuleRules.cs
--- a/Engine/Source/Programs/UnrealBuildTool/Configuration/ModuleRules.cs
+++ b/Engine/Source/Programs/UnrealBuildTool/Configuration/ModuleRules.cs
@@ -473,6 +473,17 @@ namespace UnrealBuildTool
 		/// </summary>
 		internal PluginInfo Plugin;
 
+		/// <summary>
+		/// True if a Plugin contains this module
+		/// </summary>
+		public bool IsPlugin
+		{
+			get
+			{
+				return Plugin != null;
+			}
+		}
+
 		/// <summary>
 		/// The rules context for this instance
 		/// </summary>
@@ -828,6 +839,11 @@ namespace UnrealBuildTool
 		/// </summary>
 		public List<string> PublicAdditionalLibraries = new List<string>();
 
+		/// <summary>
+		/// List of additional pre-build libraries (names of the .lib files including extension) - typically used for additional targets which are still built, but using either TargetRules.PreBuildSteps or TargetRules.PreBuildTargets.
+		/// </summary>
+		public List<string> PublicPreBuildLibraries = new List<string>();
+
 		/// <summary>
 		/// List of system libraries to use - these are typically referenced via name and then found via the system paths. If you need to reference a .lib file use the PublicAdditionalLibraries instead
 		/// </summary>
@@ -945,6 +961,11 @@ namespace UnrealBuildTool
 		/// </summary>
 		public List<string> ExternalDependencies = new List<string>();
 
+		/// <summary>
+		/// External directories containing generated interop files.
+		/// </summary>
+		public List<string> AdditionalCodeGenDirectories = new List<string>();
+
 		/// <summary>
 		/// Subclass rules files which invalidate the makefile if modified.
 		/// </summary>
diff --git a/Engine/Source/Programs/UnrealBuildTool/Configuration/TargetDescriptor.cs b/Engine/Source/Programs/UnrealBuildTool/Configuration/TargetDescriptor.cs
--- a/Engine/Source/Programs/UnrealBuildTool/Configuration/TargetDescriptor.cs
+++ b/Engine/Source/Programs/UnrealBuildTool/Configuration/TargetDescriptor.cs
@@ -132,6 +132,11 @@ namespace UnrealBuildTool
 			this.AdditionalArguments = new CommandLineArguments(AdditionalArguments.ToArray());
 		}
 
+		public static TargetDescriptor FromTargetInfo(TargetInfo Info)
+		{
+			return new TargetDescriptor(Info.ProjectFile, Info.Name, Info.Platform, Info.Configuration, Info.Architecture, Info.Arguments);
+		}
+
 		/// <summary>
 		/// Parse a list of target descriptors from the command line
 		/// </summary>
@@ -430,5 +435,29 @@ namespace UnrealBuildTool
 			}
 			return Result.ToString();
 		}
+
+		public override int GetHashCode()
+		{
+			return ProjectFile.GetHashCode() + 
+				Name.GetHashCode() + 
+				Platform.GetHashCode() + 
+				Configuration.GetHashCode() + 
+				Architecture.GetHashCode();
+		}
+
+		public override bool Equals(object Obj) 
+		{
+			TargetDescriptor Other = Obj as TargetDescriptor;
+			if (Other != null)
+			{
+				return
+					ProjectFile == Other.ProjectFile &&
+					Name == Other.Name &&
+					Platform == Other.Platform &&
+					Configuration == Other.Configuration &&
+					Architecture == Other.Architecture;
+			}
+			return false;
+		}
 	}
 }
diff --git a/Engine/Source/Programs/UnrealBuildTool/Configuration/TargetInfo.cs b/Engine/Source/Programs/UnrealBuildTool/Configuration/TargetInfo.cs
--- a/Engine/Source/Programs/UnrealBuildTool/Configuration/TargetInfo.cs
+++ b/Engine/Source/Programs/UnrealBuildTool/Configuration/TargetInfo.cs
@@ -71,5 +71,41 @@ namespace UnrealBuildTool
 			this.ProjectFile = ProjectFile;
 			this.Arguments = Arguments;
 		}
+
+		/// <summary>
+		/// Construct a TargetInfo from an archive on disk
+		/// </summary>
+		/// <param name="Reader">Archive to read from</param>
+		public TargetInfo(BinaryArchiveReader Reader)
+		{
+			this.Name = Reader.ReadString();
+			this.Platform = UnrealTargetPlatform.Parse(Reader.ReadString());
+			string ConfigurationStr = Reader.ReadString();
+			this.Architecture = Reader.ReadString();
+			this.ProjectFile = Reader.ReadFileReference();
+			string[] ArgumentStrs = Reader.ReadArray(() => Reader.ReadString());
+
+			if (!UnrealTargetConfiguration.TryParse(ConfigurationStr, out Configuration))
+			{
+				throw new BuildException(string.Format("The configration name {0} is not a valid configration name. Valid names are ({1})", Name,
+					string.Join(",", Enum.GetValues(typeof(UnrealTargetConfiguration)).Cast<UnrealTargetConfiguration>().Select(x => x.ToString()))));
+			}
+
+			Arguments = new CommandLineArguments(ArgumentStrs);
+		}
+
+		/// <summary>
+		/// Write a TargetInfo to an archive on disk
+		/// </summary>
+		/// <param name="Writer">Archive to write to</param>
+		public void Write(BinaryArchiveWriter Writer)
+		{
+			Writer.WriteString(Name);
+			Writer.WriteString(Platform.ToString());
+			Writer.WriteString(Configuration.ToString());
+			Writer.WriteString(Architecture);
+			Writer.WriteFileReference(ProjectFile);
+			Writer.WriteArray(Arguments.GetRawArray(), Item => Writer.WriteString(Item));
+		}
 	}
 }
diff --git a/Engine/Source/Programs/UnrealBuildTool/Configuration/TargetRules.cs b/Engine/Source/Programs/UnrealBuildTool/Configuration/TargetRules.cs
--- a/Engine/Source/Programs/UnrealBuildTool/Configuration/TargetRules.cs
+++ b/Engine/Source/Programs/UnrealBuildTool/Configuration/TargetRules.cs
@@ -1347,6 +1347,11 @@ namespace UnrealBuildTool
 		[CommandLine("-OverrideBuildEnvironment")]
 		public bool bOverrideBuildEnvironment = false;
 
+		/// <summary>
+		/// Specifies a list of targets which should be built before this target is built.
+		/// </summary>
+		public List<TargetInfo> PreBuildTargets = new List<TargetInfo>();
+
 		/// <summary>
 		/// Specifies a list of steps which should be executed before this target is built, in the context of the host platform's shell.
 		/// The following variables will be expanded before execution:
@@ -2561,6 +2566,11 @@ namespace UnrealBuildTool
 			get { return Inner.bOverrideBuildEnvironment; }
 		}
 
+		public IReadOnlyList<TargetInfo> PreBuildTargets
+		{
+			get { return Inner.PreBuildTargets; }
+		}
+
 		public IReadOnlyList<string> PreBuildSteps
 		{
 			get { return Inner.PreBuildSteps; }
diff --git a/Engine/Source/Programs/UnrealBuildTool/Configuration/UEBuildModule.cs b/Engine/Source/Programs/UnrealBuildTool/Configuration/UEBuildModule.cs
--- a/Engine/Source/Programs/UnrealBuildTool/Configuration/UEBuildModule.cs
+++ b/Engine/Source/Programs/UnrealBuildTool/Configuration/UEBuildModule.cs
@@ -178,11 +178,12 @@ namespace UnrealBuildTool
 
 			ModuleApiDefine = Name.ToUpperInvariant() + "_API";
 
+			HashSet<string> PublicPreBuildLibraries = HashSetFromOptionalEnumerableStringParameter(Rules.PublicPreBuildLibraries);
 			PublicDefinitions = HashSetFromOptionalEnumerableStringParameter(Rules.PublicDefinitions);
 			PublicIncludePaths = CreateDirectoryHashSet(Rules.PublicIncludePaths);
 			PublicSystemIncludePaths = CreateDirectoryHashSet(Rules.PublicSystemIncludePaths);
 			PublicSystemLibraryPaths = CreateDirectoryHashSet(Rules.PublicSystemLibraryPaths);
-			PublicAdditionalLibraries = HashSetFromOptionalEnumerableStringParameter(Rules.PublicAdditionalLibraries);
+			PublicAdditionalLibraries = HashSetFromOptionalEnumerableStringParameter(Rules.PublicAdditionalLibraries.Union(PublicPreBuildLibraries));
 			PublicSystemLibraries = HashSetFromOptionalEnumerableStringParameter(Rules.PublicSystemLibraries);
 			PublicFrameworks = HashSetFromOptionalEnumerableStringParameter(Rules.PublicFrameworks);
 			PublicWeakFrameworks = HashSetFromOptionalEnumerableStringParameter(Rules.PublicWeakFrameworks);
@@ -195,8 +196,15 @@ namespace UnrealBuildTool
 					continue;
 				}
 
-				// the library path does not seem to be resolvable as is, lets warn about it as dependency checking will not work for it
-				Log.TraceWarning("Library '{0}' was not resolvable to a file when used in Module '{1}', assuming it is a filename and will search library paths for it. This is slow and dependency checking will not work for it. Please update reference to be fully qualified alternatively use PublicSystemLibraryPaths if you do intended to use this slow path to suppress this warning. ", LibraryName, Name);
+				if (PublicPreBuildLibraries.Contains(LibraryName))
+				{
+					Log.TraceLog("Library '{0}' was not resolvable to a file when used in Module '{1}'.  Be sure to add either a TargetRules.PreBuildSteps entry or a TargetRules.PreBuildTargets entry to assure it is built for your target.", LibraryName, Name);
+				}
+				else
+				{
+					// the library path does not seem to be resolvable as is, lets warn about it as dependency checking will not work for it
+					Log.TraceWarning("Library '{0}' was not resolvable to a file when used in Module '{1}', assuming it is a filename and will search library paths for it. This is slow and dependency checking will not work for it. Please update reference to be fully qualified alternatively use PublicSystemLibraryPaths if you do intended to use this slow path to suppress this warning. ", LibraryName, Name);
+				}
 			}
 
 			PublicAdditionalFrameworks = new HashSet<UEBuildFramework>();
diff --git a/Engine/Source/Programs/UnrealBuildTool/Configuration/UEBuildModuleCPP.cs b/Engine/Source/Programs/UnrealBuildTool/Configuration/UEBuildModuleCPP.cs
--- a/Engine/Source/Programs/UnrealBuildTool/Configuration/UEBuildModuleCPP.cs
+++ b/Engine/Source/Programs/UnrealBuildTool/Configuration/UEBuildModuleCPP.cs
@@ -43,9 +43,9 @@ namespace UnrealBuildTool
 		public bool bAddGeneratedCodeIncludePath;
 
 		/// <summary>
-		/// Wildcard matching the *.gen.cpp files for this module.  If this is null then this module doesn't have any UHT-produced code.
+		/// Paths containing *.gen.cpp files for this module.  If this is null then this module doesn't have any generated code.
 		/// </summary>
-		public string GeneratedCodeWildcard;
+		public List<string> GeneratedCppDirectories;
 
 		/// <summary>
 		/// List of invalid include directives. These are buffered up and output before we start compiling.
@@ -259,6 +259,17 @@ namespace UnrealBuildTool
 			if(bAddGeneratedCodeIncludePath || (ProjectFileGenerator.bGenerateProjectFiles && GeneratedCodeDirectory != null))
 			{
 				IncludePaths.Add(GeneratedCodeDirectory);
+
+				if (Rules.AdditionalCodeGenDirectories != null)
+				{
+					foreach (string CodeGenDir in Rules.AdditionalCodeGenDirectories)
+					{
+						if (Directory.Exists(CodeGenDir))
+						{
+							IncludePaths.Add(new DirectoryReference(CodeGenDir));
+						}
+					}
+				}
 			}
 
 			base.AddModuleToCompileEnvironment(SourceBinary, IncludePaths, SystemIncludePaths, Definitions, AdditionalFrameworks, AdditionalPrerequisites, bLegacyPublicIncludePaths);
@@ -421,10 +432,15 @@ namespace UnrealBuildTool
 			}
 
 			// Compile all the generated CPP files
-			if (GeneratedCodeWildcard != null && !CompileEnvironment.bHackHeaderGenerator)
+			if (GeneratedCppDirectories != null && !CompileEnvironment.bHackHeaderGenerator && SingleFileToCompile == null)
 			{
-				string[] GeneratedFiles = Directory.GetFiles(Path.GetDirectoryName(GeneratedCodeWildcard), Path.GetFileName(GeneratedCodeWildcard));
-				if(GeneratedFiles.Length > 0)
+				List<string> GeneratedFiles = new List<string>();
+				foreach (string GeneratedDir in GeneratedCppDirectories)
+				{
+					GeneratedFiles.AddRange(Directory.GetFiles(GeneratedDir, "*.gen.cpp"));
+				}
+
+				if(GeneratedFiles.Count > 0)
 				{
 					// Create a compile environment for the generated files. We can disable creating debug info here to improve link times.
 					CppCompileEnvironment GeneratedCPPCompileEnvironment = CompileEnvironment;
diff --git a/Engine/Source/Programs/UnrealBuildTool/Configuration/UEBuildTarget.cs b/Engine/Source/Programs/UnrealBuildTool/Configuration/UEBuildTarget.cs
--- a/Engine/Source/Programs/UnrealBuildTool/Configuration/UEBuildTarget.cs
+++ b/Engine/Source/Programs/UnrealBuildTool/Configuration/UEBuildTarget.cs
@@ -1619,6 +1619,9 @@ namespace UnrealBuildTool
 			TargetToolChain.GetVersionInfo(Makefile.Diagnostics);
 			Rules.GetBuildSettingsInfo(Makefile.Diagnostics);
 
+			// Get any pre-build targets.
+			Makefile.PreBuildTargets = Rules.PreBuildTargets.ToArray();
+
 			// Setup the hot reload module list
 			Makefile.HotReloadModuleNames = GetHotReloadModuleNames();
 
@@ -1968,6 +1971,16 @@ namespace UnrealBuildTool
 			}
 			Makefile.ExternalDependencies.UnionWith(Makefile.PluginFiles);
 
+			// Add any leaf dependencies (eg. response files) to the dependencies list
+			IEnumerable<FileItem> LeafPrerequisiteItems = Makefile.Actions.SelectMany(x => x.PrerequisiteItems).Except(Makefile.Actions.SelectMany(x => x.ProducedItems));
+			foreach (FileItem LeafPrerequisiteItem in LeafPrerequisiteItems)
+			{
+				if (LeafPrerequisiteItem.Exists)
+				{
+					Makefile.AdditionalDependencies.Add(LeafPrerequisiteItem);
+				}
+			}
+
 			// Write a header containing public definitions for this target
 			if (Rules.ExportPublicHeader != null)
 			{
diff --git a/Engine/Source/Programs/UnrealBuildTool/Modes/BuildMode.cs b/Engine/Source/Programs/UnrealBuildTool/Modes/BuildMode.cs
--- a/Engine/Source/Programs/UnrealBuildTool/Modes/BuildMode.cs
+++ b/Engine/Source/Programs/UnrealBuildTool/Modes/BuildMode.cs
@@ -64,6 +64,12 @@ namespace UnrealBuildTool
 		[CommandLine("-SkipBuild")]
 		public bool bSkipBuild = false;
 
+		/// <summary>
+		/// Skip pre build targets; just do the main target.
+		/// </summary>
+		[CommandLine("-SkipPreBuildTargets")]
+		public bool bSkipPreBuildTargets = false;
+
 		/// <summary>
 		/// Whether we should just export the XGE XML and pretend it succeeded
 		/// </summary>
@@ -151,8 +157,9 @@ namespace UnrealBuildTool
 			// Parse and build the targets
 			try
 			{
-				// Parse all the target descriptors
 				List<TargetDescriptor> TargetDescriptors;
+
+				// Parse all the target descriptors
 				using(Timeline.ScopeEvent("TargetDescriptor.ParseCommandLine()"))
 				{
 					TargetDescriptors = TargetDescriptor.ParseCommandLine(Arguments, BuildConfiguration.bUsePrecompiled, BuildConfiguration.bSkipRulesCompile);
@@ -165,16 +172,16 @@ namespace UnrealBuildTool
 				}
 
 				// Handle remote builds
-				for(int Idx = 0; Idx < TargetDescriptors.Count; Idx++)
+				for (int Idx = 0; Idx < TargetDescriptors.Count; ++Idx)
 				{
 					TargetDescriptor TargetDesc = TargetDescriptors[Idx];
-					if(RemoteMac.HandlesTargetPlatform(TargetDesc.Platform))
+					if (RemoteMac.HandlesTargetPlatform(TargetDesc.Platform))
 					{
 						FileReference BaseLogFile = Log.OutputFile ?? new FileReference(BaseLogFileName);
 						FileReference RemoteLogFile = FileReference.Combine(BaseLogFile.Directory, BaseLogFile.GetFileNameWithoutExtension() + "_Remote.txt");
 
 						RemoteMac RemoteMac = new RemoteMac(TargetDesc.ProjectFile);
-						if(!RemoteMac.Build(TargetDesc, RemoteLogFile))
+						if (!RemoteMac.Build(TargetDesc, RemoteLogFile, bSkipPreBuildTargets))
 						{
 							return (int)CompilationResult.Unknown;
 						}
@@ -184,13 +191,13 @@ namespace UnrealBuildTool
 				}
 
 				// Handle local builds
-				if(TargetDescriptors.Count > 0)
+				if (TargetDescriptors.Count > 0)
 				{
 					// Get a set of all the project directories
 					HashSet<DirectoryReference> ProjectDirs = new HashSet<DirectoryReference>();
-					foreach(TargetDescriptor TargetDesc in TargetDescriptors)
+					foreach (TargetDescriptor TargetDesc in TargetDescriptors)
 					{
-						if(TargetDesc.ProjectFile != null)
+						if (TargetDesc.ProjectFile != null)
 						{
 							DirectoryReference ProjectDirectory = TargetDesc.ProjectFile.Directory;
 							FileMetadataPrefetch.QueueProjectDirectory(ProjectDirectory);
@@ -200,11 +207,11 @@ namespace UnrealBuildTool
 
 					// Get all the build options
 					BuildOptions Options = BuildOptions.None;
-					if(bSkipBuild)
+					if (bSkipBuild)
 					{
 						Options |= BuildOptions.SkipBuild;
 					}
-					if(bXGEExport)
+					if (bXGEExport)
 					{
 						Options |= BuildOptions.XGEExport;
 					}
@@ -213,10 +220,10 @@ namespace UnrealBuildTool
 						Options |= BuildOptions.NoEngineChanges;
 					}
 
-					// Create the working set provider
+					// Create the working set provider per group.
 					using (ISourceFileWorkingSet WorkingSet = SourceFileWorkingSet.Create(UnrealBuildTool.RootDirectory, ProjectDirs))
 					{
-						Build(TargetDescriptors, BuildConfiguration, WorkingSet, Options, WriteOutdatedActionsFile);
+						Build(TargetDescriptors, BuildConfiguration, WorkingSet, Options, WriteOutdatedActionsFile, bSkipPreBuildTargets);
 					}
 				}
 			}
@@ -237,16 +244,45 @@ namespace UnrealBuildTool
 		/// <param name="WorkingSet">The source file working set</param>
 		/// <param name="Options">Additional options for the build</param>
 		/// <param name="WriteOutdatedActionsFile">Files to write the list of outdated actions to (rather than building them)</param>
+		/// <param name="bSkipPreBuildTargets">If true then only the current target descriptors will be built.</param>
 		/// <returns>Result from the compilation</returns>
-		public static void Build(List<TargetDescriptor> TargetDescriptors, BuildConfiguration BuildConfiguration, ISourceFileWorkingSet WorkingSet, BuildOptions Options, FileReference WriteOutdatedActionsFile)
+		public static void Build(List<TargetDescriptor> TargetDescriptors, BuildConfiguration BuildConfiguration, ISourceFileWorkingSet WorkingSet, BuildOptions Options, FileReference WriteOutdatedActionsFile, bool bSkipPreBuildTargets = false)
 		{
-			// Create a makefile for each target
-			TargetMakefile[] Makefiles = new TargetMakefile[TargetDescriptors.Count];
-			for(int TargetIdx = 0; TargetIdx < TargetDescriptors.Count; TargetIdx++)
+
+			List<TargetMakefile> TargetMakefiles = new List<TargetMakefile>();
+
+			for (int Idx = 0; Idx < TargetDescriptors.Count; ++Idx)
 			{
-				Makefiles[TargetIdx] = CreateMakefile(BuildConfiguration, TargetDescriptors[TargetIdx], WorkingSet);
+				TargetMakefile NewMakefile = CreateMakefile(BuildConfiguration, TargetDescriptors[Idx], WorkingSet);
+				TargetMakefiles.Add(NewMakefile);
+				if (!bSkipPreBuildTargets)
+				{
+					foreach (TargetInfo PreBuildTarget in NewMakefile.PreBuildTargets)
+					{
+						TargetDescriptor NewTarget = TargetDescriptor.FromTargetInfo(PreBuildTarget);
+						if (!TargetDescriptors.Contains(NewTarget))
+						{
+							TargetDescriptors.Add(NewTarget);
+						}
+					}
+				}
 			}
 
+			Build(TargetMakefiles.ToArray(), TargetDescriptors, BuildConfiguration, WorkingSet, Options, WriteOutdatedActionsFile);
+		}
+
+		/// <summary>
+		/// Build a list of targets with a given set of makefiles.
+		/// </summary>
+		/// <param name="Makefiles">Makefiles created with CreateMakefiles</param>
+		/// <param name="TargetDescriptors">Target descriptors</param>
+		/// <param name="BuildConfiguration">Current build configuration</param>
+		/// <param name="WorkingSet">The source file working set</param>
+		/// <param name="Options">Additional options for the build</param>
+		/// <param name="WriteOutdatedActionsFile">Files to write the list of outdated actions to (rather than building them)</param>
+		/// <returns>Result from the compilation</returns>
+		static void Build(TargetMakefile[] Makefiles, List<TargetDescriptor> TargetDescriptors, BuildConfiguration BuildConfiguration, ISourceFileWorkingSet WorkingSet, BuildOptions Options, FileReference WriteOutdatedActionsFile)
+		{
 			// Export the actions for each target
 			for (int TargetIdx = 0; TargetIdx < TargetDescriptors.Count; TargetIdx++)
 			{
diff --git a/Engine/Source/Programs/UnrealBuildTool/Modes/CleanMode.cs b/Engine/Source/Programs/UnrealBuildTool/Modes/CleanMode.cs
--- a/Engine/Source/Programs/UnrealBuildTool/Modes/CleanMode.cs
+++ b/Engine/Source/Programs/UnrealBuildTool/Modes/CleanMode.cs
@@ -24,6 +24,12 @@ namespace UnrealBuildTool
 		[CommandLine("-SkipRulesCompile")]
 		bool bSkipRulesCompile = false;
 
+		/// <summary>
+		/// Skip pre build targets; just do the main target.
+		/// </summary>
+		[CommandLine("-SkipPreBuildTargets")]
+		public bool bSkipPreBuildTargets = false;
+
 		/// <summary>
 		/// Main entry point
 		/// </summary>
@@ -40,27 +46,27 @@ namespace UnrealBuildTool
 
 			// Parse all the targets being built
 			List<TargetDescriptor> TargetDescriptors = TargetDescriptor.ParseCommandLine(Arguments, BuildConfiguration.bUsePrecompiled, bSkipRulesCompile);
-			if(TargetDescriptors.Count == 0)
+			if (TargetDescriptors.Count == 0)
 			{
 				throw new BuildException("No targets specified to clean");
 			}
 
 			// Also add implicit descriptors for cleaning UnrealBuildTool
-			if(!BuildConfiguration.bDoNotBuildUHT)
+			if (!BuildConfiguration.bDoNotBuildUHT)
 			{
 				const string UnrealHeaderToolTarget = "UnrealHeaderTool";
 
 				// Get a list of project files to clean UHT for
 				List<FileReference> ProjectFiles = new List<FileReference>();
-				foreach(TargetDescriptor TargetDesc in TargetDescriptors)
+				foreach (TargetDescriptor TargetDesc in TargetDescriptors)
 				{
-					if(TargetDesc.Name != UnrealHeaderToolTarget && !RemoteMac.HandlesTargetPlatform(TargetDesc.Platform))
+					if (TargetDesc.Name != UnrealHeaderToolTarget && !RemoteMac.HandlesTargetPlatform(TargetDesc.Platform))
 					{
-						if(ProjectFiles.Count == 0)
+						if (ProjectFiles.Count == 0)
 						{
 							ProjectFiles.Add(null);
 						}
-						if(TargetDesc.ProjectFile != null && !ProjectFiles.Contains(TargetDesc.ProjectFile))
+						if (TargetDesc.ProjectFile != null && !ProjectFiles.Contains(TargetDesc.ProjectFile))
 						{
 							ProjectFiles.Add(TargetDesc.ProjectFile);
 						}
@@ -68,11 +74,11 @@ namespace UnrealBuildTool
 				}
 
 				// Add descriptors for cleaning UHT with all these projects
-				if(ProjectFiles.Count > 0)
+				if (ProjectFiles.Count > 0)
 				{
 					UnrealTargetConfiguration Configuration = BuildConfiguration.bForceDebugUnrealHeaderTool ? UnrealTargetConfiguration.Debug : UnrealTargetConfiguration.Development;
 					string Architecture = UEBuildPlatform.GetBuildPlatform(BuildHostPlatform.Current.Platform).GetDefaultArchitecture(null);
-					foreach(FileReference ProjectFile in ProjectFiles)
+					foreach (FileReference ProjectFile in ProjectFiles)
 					{
 						TargetDescriptors.Add(new TargetDescriptor(ProjectFile, UnrealHeaderToolTarget, BuildHostPlatform.Current.Platform, Configuration, Architecture, null));
 					}
@@ -85,14 +91,29 @@ namespace UnrealBuildTool
 			// Loop through all the targets, and clean them all
 			HashSet<FileReference> FilesToDelete = new HashSet<FileReference>();
 			HashSet<DirectoryReference> DirectoriesToDelete = new HashSet<DirectoryReference>();
-			foreach(TargetDescriptor TargetDescriptor in TargetDescriptors)
+
+			for (int Idx = 0; Idx < TargetDescriptors.Count; ++Idx)
 			{
+				TargetDescriptor TargetDescriptor = TargetDescriptors[Idx];
+
 				// Create the rules assembly
 				RulesAssembly RulesAssembly = RulesCompiler.CreateTargetRulesAssembly(TargetDescriptor.ProjectFile, TargetDescriptor.Name, bSkipRulesCompile, BuildConfiguration.bUsePrecompiled, TargetDescriptor.ForeignPlugin);
 
 				// Create the rules object
 				ReadOnlyTargetRules Target = new ReadOnlyTargetRules(RulesAssembly.CreateTargetRules(TargetDescriptor.Name, TargetDescriptor.Platform, TargetDescriptor.Configuration, TargetDescriptor.Architecture, TargetDescriptor.ProjectFile, TargetDescriptor.AdditionalArguments));
 
+				if (!bSkipPreBuildTargets && Target.PreBuildTargets.Count > 0)
+				{
+					foreach (TargetInfo PreBuildTarget in Target.PreBuildTargets)
+					{
+						TargetDescriptor NewTarget = TargetDescriptor.FromTargetInfo(PreBuildTarget);
+						if (!TargetDescriptors.Contains(NewTarget))
+						{
+							TargetDescriptors.Add(NewTarget);
+						}
+					}
+				}
+
 				// Find the base folders that can contain binaries
 				List<DirectoryReference> BaseDirs = new List<DirectoryReference>();
 				BaseDirs.Add(UnrealBuildTool.EngineDirectory);
@@ -163,7 +184,7 @@ namespace UnrealBuildTool
 				foreach (DirectoryReference BaseDir in BaseDirs)
 				{
 					DirectoryReference BinariesDir = DirectoryReference.Combine(BaseDir, "Binaries", Target.Platform.ToString());
-					if(DirectoryReference.Exists(BinariesDir))
+					if (DirectoryReference.Exists(BinariesDir))
 					{
 						UEBuildPlatform.GetBuildPlatform(Target.Platform).FindBuildProductsToClean(BinariesDir, NamePrefixesArray, NameSuffixesArray, AdditionalFilesToDelete, AdditionalDirectoriesToDelete);
 					}
@@ -211,10 +232,10 @@ namespace UnrealBuildTool
 			}
 
 			// Also clean all the remote targets
-			for(int Idx = 0; Idx < TargetDescriptors.Count; Idx++)
+			for (int Idx = 0; Idx < TargetDescriptors.Count; Idx++)
 			{
 				TargetDescriptor TargetDescriptor = TargetDescriptors[Idx];
-				if(RemoteMac.HandlesTargetPlatform(TargetDescriptor.Platform))
+				if (RemoteMac.HandlesTargetPlatform(TargetDescriptor.Platform))
 				{
 					RemoteMac RemoteMac = new RemoteMac(TargetDescriptor.ProjectFile);
 					RemoteMac.Clean(TargetDescriptor);
diff --git a/Engine/Source/Programs/UnrealBuildTool/Platform/Mac/MacToolChain.cs b/Engine/Source/Programs/UnrealBuildTool/Platform/Mac/MacToolChain.cs
--- a/Engine/Source/Programs/UnrealBuildTool/Platform/Mac/MacToolChain.cs
+++ b/Engine/Source/Programs/UnrealBuildTool/Platform/Mac/MacToolChain.cs
@@ -1346,7 +1346,9 @@ namespace UnrealBuildTool
 
 			FixDylibOutputFile = FixDylibDependencies(BinaryLinkEnvironment, Executable, Graph);
 			OutputFiles.Add(FixDylibOutputFile);
-			if (!BinaryLinkEnvironment.bIsBuildingConsoleApplication)
+
+			bool bIsBuildingAppBundle = !BinaryLinkEnvironment.bIsBuildingDLL && !BinaryLinkEnvironment.bIsBuildingLibrary && !BinaryLinkEnvironment.bIsBuildingConsoleApplication;
+			if (bIsBuildingAppBundle)
 			{
 				OutputFiles.Add(FinalizeAppBundle(BinaryLinkEnvironment, Executable, FixDylibOutputFile, Graph));
 			}
diff --git a/Engine/Source/Programs/UnrealBuildTool/ProjectFiles/Xcode/XcodeProject.cs b/Engine/Source/Programs/UnrealBuildTool/ProjectFiles/Xcode/XcodeProject.cs
--- a/Engine/Source/Programs/UnrealBuildTool/ProjectFiles/Xcode/XcodeProject.cs
+++ b/Engine/Source/Programs/UnrealBuildTool/ProjectFiles/Xcode/XcodeProject.cs
@@ -1791,7 +1791,7 @@ namespace UnrealBuildTool
 				{
 					bHasEditorConfiguration = true;
 				}
-				else if (Config.ProjectTarget.TargetRules.bIsBuildingConsoleApplication)
+				else if (Config.ProjectTarget.TargetRules.bIsBuildingConsoleApplication || Config.ProjectTarget.TargetRules.bShouldCompileAsDLL)
 				{
 					bIsAppBundle = false;
 				}
diff --git a/Engine/Source/Programs/UnrealBuildTool/System/ExternalExecution.cs b/Engine/Source/Programs/UnrealBuildTool/System/ExternalExecution.cs
--- a/Engine/Source/Programs/UnrealBuildTool/System/ExternalExecution.cs
+++ b/Engine/Source/Programs/UnrealBuildTool/System/ExternalExecution.cs
@@ -599,7 +599,19 @@ namespace UnrealBuildTool
 					Info.GeneratedCPPFilenameBase = Path.Combine(Module.GeneratedCodeDirectory.FullName, Info.ModuleName) + ".gen";
 					if (!Module.Rules.bUsePrecompiled)
 					{
-						Module.GeneratedCodeWildcard = Path.Combine(Module.GeneratedCodeDirectory.FullName, "*.gen.cpp");
+						Module.GeneratedCppDirectories = new List<string>();
+						Module.GeneratedCppDirectories.Add(Module.GeneratedCodeDirectory.FullName);
+
+						if (Module.Rules.AdditionalCodeGenDirectories != null)
+						{
+							foreach (string DirPath in Module.Rules.AdditionalCodeGenDirectories)
+							{
+								if (Directory.Exists(DirPath))
+								{
+									Module.GeneratedCppDirectories.Add(Path.GetFullPath(DirPath));
+								}
+							}
+						}
 					}
 
 					UObjectModules.Add(Info);
diff --git a/Engine/Source/Programs/UnrealBuildTool/System/TargetMakefile.cs b/Engine/Source/Programs/UnrealBuildTool/System/TargetMakefile.cs
--- a/Engine/Source/Programs/UnrealBuildTool/System/TargetMakefile.cs
+++ b/Engine/Source/Programs/UnrealBuildTool/System/TargetMakefile.cs
@@ -88,6 +88,11 @@ namespace UnrealBuildTool
 		/// </summary>
 		public FileReference[] PreBuildScripts;
 
+		/// <summary>
+		/// Targets which should be build before building anything
+		/// </summary>
+		public TargetInfo[] PreBuildTargets;
+
 		/// <summary>
 		/// Every action in the action graph
 		/// </summary>
@@ -223,6 +228,7 @@ namespace UnrealBuildTool
 			bHasProjectScriptPlugin = Reader.ReadBool();
 			AdditionalArguments = Reader.ReadArray(() => Reader.ReadString());
 			PreBuildScripts = Reader.ReadArray(() => Reader.ReadFileReference());
+			PreBuildTargets = Reader.ReadArray(() => new TargetInfo(Reader));
 			Actions = Reader.ReadList(() => new Action(Reader));
 			EnvironmentVariables = Reader.ReadList(() => Tuple.Create(Reader.ReadString(), Reader.ReadString()));
 			OutputItems = Reader.ReadList(() => Reader.ReadFileItem());
@@ -257,6 +263,7 @@ namespace UnrealBuildTool
 			Writer.WriteBool(bHasProjectScriptPlugin);
 			Writer.WriteArray(AdditionalArguments, Item => Writer.WriteString(Item));
 			Writer.WriteArray(PreBuildScripts, Item => Writer.WriteFileReference(Item));
+			Writer.WriteArray(PreBuildTargets, Item => Item.Write(Writer));
 			Writer.WriteList(Actions, Action => Action.Write(Writer));
 			Writer.WriteList(EnvironmentVariables, x => { Writer.WriteString(x.Item1); Writer.WriteString(x.Item2); });
 			Writer.WriteList(OutputItems, Item => Writer.WriteFileItem(Item));
diff --git a/Engine/Source/Programs/UnrealBuildTool/ToolChain/RemoteMac.cs b/Engine/Source/Programs/UnrealBuildTool/ToolChain/RemoteMac.cs
--- a/Engine/Source/Programs/UnrealBuildTool/ToolChain/RemoteMac.cs
+++ b/Engine/Source/Programs/UnrealBuildTool/ToolChain/RemoteMac.cs
@@ -404,9 +404,30 @@ namespace UnrealBuildTool
 		/// </summary>
 		/// <param name="TargetDesc">Descriptor for the target to build</param>
 		/// <param name="RemoteLogFile">Path to store the remote log file</param>
+		/// <param name="bSkipPreBuildTargets">If true then any PreBuildTargets will be skipped</param>
 		/// <returns>True if the build succeeded, false otherwise</returns>
-		public bool Build(TargetDescriptor TargetDesc, FileReference RemoteLogFile)
+		public bool Build(TargetDescriptor TargetDesc, FileReference RemoteLogFile, bool bSkipPreBuildTargets)
 		{
+			// Compile the rules assembly
+			RulesAssembly RulesAssembly = RulesCompiler.CreateTargetRulesAssembly(TargetDesc.ProjectFile, TargetDesc.Name, false, false, TargetDesc.ForeignPlugin);
+
+			// Create the target rules
+			TargetRules Rules = RulesAssembly.CreateTargetRules(TargetDesc.Name, TargetDesc.Platform, TargetDesc.Configuration, TargetDesc.Architecture, TargetDesc.ProjectFile, TargetDesc.AdditionalArguments);
+			if (!bSkipPreBuildTargets)
+			{
+				foreach (TargetInfo PreBuildTargetInfo in Rules.PreBuildTargets)
+				{
+					RemoteMac PreBuildTargetRemoteMac = new RemoteMac(ProjectFile);
+					TargetDescriptor PreBuildTargetDesc = new TargetDescriptor(PreBuildTargetInfo.ProjectFile, PreBuildTargetInfo.Name, PreBuildTargetInfo.Platform, PreBuildTargetInfo.Configuration, PreBuildTargetInfo.Architecture, PreBuildTargetInfo.Arguments);
+
+					Log.TraceInformation("[Remote] Building pre target [{0}] for [{1}] ", PreBuildTargetDesc.ToString(), TargetDesc.ToString());
+					if (!PreBuildTargetRemoteMac.Build(PreBuildTargetDesc, RemoteLogFile, false))
+					{
+						return false;
+					}
+				}
+			}
+
 			// Get the directory for working files
 			DirectoryReference TempDir = CreateTempDirectory(TargetDesc);
 
@@ -425,12 +446,6 @@ namespace UnrealBuildTool
 				Mappings.Add(new RemoteMapping(RemoteLogFile.Directory, GetRemotePath(RemoteLogFile.Directory)));
 			}
 
-			// Compile the rules assembly
-			RulesAssembly RulesAssembly = RulesCompiler.CreateTargetRulesAssembly(TargetDesc.ProjectFile, TargetDesc.Name, false, false, TargetDesc.ForeignPlugin);
-
-			// Create the target rules
-			TargetRules Rules = RulesAssembly.CreateTargetRules(TargetDesc.Name, TargetDesc.Platform, TargetDesc.Configuration, TargetDesc.Architecture, TargetDesc.ProjectFile, TargetDesc.AdditionalArguments);
-
 			// Check if we need to enable a nativized plugin, and compile the assembly for that if we do
 			FileReference NativizedPluginFile = Rules.GetNativizedPlugin();
 			if(NativizedPluginFile != null)
@@ -448,6 +463,7 @@ namespace UnrealBuildTool
 			List<string> RemoteArguments = GetRemoteArgumentsForTarget(TargetDesc, LocalManifestFiles);
 			RemoteArguments.Add(String.Format("-Log={0}", GetRemotePath(RemoteLogFile)));
 			RemoteArguments.Add(String.Format("-Manifest={0}", GetRemotePath(RemoteManifestFile)));
+			RemoteArguments.Add(String.Format("-SkipPreBuildTargets"));
 
 			// Handle any per-platform setup that is required
 			if(TargetDesc.Platform == UnrealTargetPlatform.IOS || TargetDesc.Platform == UnrealTargetPlatform.TVOS)
diff --git a/Engine/Source/Runtime/Core/Private/Misc/AssertionMacros.cpp b/Engine/Source/Runtime/Core/Private/Misc/AssertionMacros.cpp
--- a/Engine/Source/Runtime/Core/Private/Misc/AssertionMacros.cpp
+++ b/Engine/Source/Runtime/Core/Private/Misc/AssertionMacros.cpp
@@ -4,6 +4,7 @@
 #include "Misc/VarArgs.h"
 #include "HAL/UnrealMemory.h"
 #include "Templates/UnrealTemplate.h"
+#include "Templates/Atomic.h"
 #include "Misc/CString.h"
 #include "Misc/Crc.h"
 #include "Containers/UnrealString.h"
@@ -27,6 +28,7 @@ namespace
 {
 	// used to track state of assets/ensures
 	bool bHasAsserted = false;
+	TAtomic<SIZE_T> NumEnsureFailures {0};
 	int32 ActiveEnsureCount = 0;
 
 	/** Lock used to synchronize the fail debug calls. */
@@ -157,6 +159,10 @@ bool FDebug::IsEnsuring()
 {
 	return ActiveEnsureCount > 0;
 }
+SIZE_T FDebug::GetNumEnsureFailures()
+{
+	return NumEnsureFailures.Load();
+}
 
 void FDebug::LogFormattedMessageWithCallstack(const FName& InLogName, const ANSICHAR* File, int32 Line, const TCHAR* Heading, const TCHAR* Message, ELogVerbosity::Type Verbosity)
 {
@@ -274,6 +280,8 @@ FORCENOINLINE void FDebug::EnsureFailed(const ANSICHAR* Expr, const ANSICHAR* Fi
 		return;
 	}
 	
+	++NumEnsureFailures;
+
 #if STATS
 	FString EnsureFailedPerfMessage = FString::Printf(TEXT("FDebug::EnsureFailed"));
 	SCOPE_LOG_TIME_IN_SECONDS(*EnsureFailedPerfMessage, nullptr)
diff --git a/Engine/Source/Runtime/Core/Private/Tests/Misc/CharTest.cpp b/Engine/Source/Runtime/Core/Private/Tests/Misc/CharTest.cpp
--- a/Engine/Source/Runtime/Core/Private/Tests/Misc/CharTest.cpp
+++ b/Engine/Source/Runtime/Core/Private/Tests/Misc/CharTest.cpp
@@ -33,9 +33,11 @@ void RunCharTests(FAutomationTestBase& Test, uint32 MaxChar)
 bool TCharTest::RunTest(const FString& Parameters)
 {
 	const char* CurrentLocale = setlocale(LC_CTYPE, nullptr);
-	check(CurrentLocale);
-
-	if (strcmp("C", CurrentLocale))
+	if (CurrentLocale == nullptr)
+	{
+		AddError(FString::Printf(TEXT("Locale is null but should be \"C\". Did something call setlocale()?")));
+	}
+	else if (strcmp("C", CurrentLocale))
 	{
 		AddError(FString::Printf(TEXT("Locale is \"%s\" but should be \"C\". Did something call setlocale()?"), ANSI_TO_TCHAR(CurrentLocale)));
 	}
diff --git a/Engine/Source/Runtime/Core/Public/Misc/AssertionMacros.h b/Engine/Source/Runtime/Core/Public/Misc/AssertionMacros.h
--- a/Engine/Source/Runtime/Core/Public/Misc/AssertionMacros.h
+++ b/Engine/Source/Runtime/Core/Public/Misc/AssertionMacros.h
@@ -47,6 +47,9 @@ struct CORE_API FDebug
 	// returns true if an ensure is currently in progress (e.g. the RenderThread is ensuring)
 	static bool IsEnsuring();
 
+	// returns the number of times an ensure has failed in this instance.
+	static SIZE_T GetNumEnsureFailures();
+
 	/** Dumps the stack trace into the log, meant to be used for debugging purposes. */
 	static void DumpStackTraceToLog();
 
diff --git a/Engine/Source/Runtime/CoreUObject/Private/Blueprint/BlueprintSupport.cpp b/Engine/Source/Runtime/CoreUObject/Private/Blueprint/BlueprintSupport.cpp
--- a/Engine/Source/Runtime/CoreUObject/Private/Blueprint/BlueprintSupport.cpp
+++ b/Engine/Source/Runtime/CoreUObject/Private/Blueprint/BlueprintSupport.cpp
@@ -117,6 +117,7 @@ bool FBlueprintSupport::IsDeferredCDOInitializationDisabled()
 }
 
 static FFlushReinstancingQueueFPtr FlushReinstancingQueueFPtr = nullptr;
+static FClassReparentingFPtr ClassReparentingFPtr = nullptr;
 
 void FBlueprintSupport::FlushReinstancingQueue()
 {
@@ -126,11 +127,24 @@ void FBlueprintSupport::FlushReinstancingQueue()
 	}
 }
 
+void FBlueprintSupport::ReparentHierarchies(const TMap<UClass*, UClass*>& OldClassToNewClass)
+{
+	if (ClassReparentingFPtr)
+	{
+		(*ClassReparentingFPtr)(OldClassToNewClass);
+	}
+}
+
 void FBlueprintSupport::SetFlushReinstancingQueueFPtr(FFlushReinstancingQueueFPtr Ptr)
 {
 	FlushReinstancingQueueFPtr = Ptr;
 }
 
+void FBlueprintSupport::SetClassReparentingFPtr(FClassReparentingFPtr Ptr)
+{
+	ClassReparentingFPtr = Ptr;
+}
+
 bool FBlueprintSupport::IsDeferredDependencyPlaceholder(UObject* LoadedObj)
 {
 	return LoadedObj && ( LoadedObj->IsA<ULinkerPlaceholderClass>() ||
diff --git a/Engine/Source/Runtime/CoreUObject/Private/Misc/PackageName.cpp b/Engine/Source/Runtime/CoreUObject/Private/Misc/PackageName.cpp
--- a/Engine/Source/Runtime/CoreUObject/Private/Misc/PackageName.cpp
+++ b/Engine/Source/Runtime/CoreUObject/Private/Misc/PackageName.cpp
@@ -161,6 +161,7 @@ struct FLongPackagePathsSingleton
 	FString EngineRootPath;
 	FString GameRootPath;
 	FString ScriptRootPath;
+	FString ExtraRootPath;
 	FString MemoryRootPath;
 	FString TempRootPath;
 	TArray<FString> MountPointRootPaths;
@@ -172,10 +173,12 @@ struct FLongPackagePathsSingleton
 	FString GameContentPath;
 	FString GameConfigPath;
 	FString GameScriptPath;
+	FString GameExtraPath;
 	FString GameSavedPath;
 	FString GameContentPathRebased;
 	FString GameConfigPathRebased;
 	FString GameScriptPathRebased;
+	FString GameExtraPathRebased;
 	FString GameSavedPathRebased;
 
 	//@TODO: Can probably consolidate these into a single array, if it weren't for EngineContentPathShort
@@ -203,6 +206,7 @@ struct FLongPackagePathsSingleton
 		{
 			OutRoots.Add(ConfigRootPath);
 			OutRoots.Add(ScriptRootPath);
+			OutRoots.Add(ExtraRootPath);
 			OutRoots.Add(MemoryRootPath);
 			OutRoots.Add(TempRootPath);
 		}
@@ -295,6 +299,7 @@ private:
 		EngineRootPath = TEXT("/Engine/");
 		GameRootPath   = TEXT("/Game/");
 		ScriptRootPath = TEXT("/Script/");
+		ExtraRootPath  = TEXT("/Extra/");
 		MemoryRootPath = TEXT("/Memory/");
 		TempRootPath   = TEXT("/Temp/");
 
@@ -305,6 +310,7 @@ private:
 		GameContentPath        = FPaths::ProjectContentDir();
 		GameConfigPath         = FPaths::ProjectConfigDir();
 		GameScriptPath         = FPaths::ProjectDir() / TEXT("Script/");
+		GameExtraPath          = FPaths::ProjectDir() / TEXT("Extra/");
 		GameSavedPath          = FPaths::ProjectSavedDir();
 
 		FString RebasedGameDir = FString::Printf(TEXT("../../../%s/"), FApp::GetProjectName());
@@ -312,11 +318,12 @@ private:
 		GameContentPathRebased = RebasedGameDir / TEXT("Content/");
 		GameConfigPathRebased  = RebasedGameDir / TEXT("Config/");
 		GameScriptPathRebased  = RebasedGameDir / TEXT("Script/");
+		GameExtraPathRebased   = RebasedGameDir / TEXT("Extra/");
 		GameSavedPathRebased   = RebasedGameDir / TEXT("Saved/");
 		
 		FScopeLock ScopeLock(&ContentMountPointCriticalSection);
 
-		ContentPathToRoot.Empty(11);
+		ContentPathToRoot.Empty(12);
 		ContentPathToRoot.Emplace(EngineRootPath, EngineContentPath);
 		if (FPaths::IsSamePath(GameContentPath, ContentPathShort))
 		{
@@ -330,14 +337,14 @@ private:
 		ContentPathToRoot.Emplace(EngineRootPath, EngineShadersPathShort);
 		ContentPathToRoot.Emplace(GameRootPath,   GameContentPath);
 		ContentPathToRoot.Emplace(ScriptRootPath, GameScriptPath);
-		ContentPathToRoot.Emplace(ScriptRootPath, GameScriptPath);
 		ContentPathToRoot.Emplace(TempRootPath,   GameSavedPath);
 		ContentPathToRoot.Emplace(GameRootPath,   GameContentPathRebased);
 		ContentPathToRoot.Emplace(ScriptRootPath, GameScriptPathRebased);
 		ContentPathToRoot.Emplace(TempRootPath,   GameSavedPathRebased);
 		ContentPathToRoot.Emplace(ConfigRootPath, GameConfigPath);
+		ContentPathToRoot.Emplace(ExtraRootPath,  GameExtraPath);
 
-		ContentRootToPath.Empty(9);
+		ContentRootToPath.Empty(10);
 		ContentRootToPath.Emplace(EngineRootPath, EngineContentPath);
 		ContentRootToPath.Emplace(EngineRootPath, EngineShadersPath);
 		ContentRootToPath.Emplace(GameRootPath,   GameContentPath);
@@ -345,6 +352,7 @@ private:
 		ContentRootToPath.Emplace(TempRootPath,   GameSavedPath);
 		ContentRootToPath.Emplace(GameRootPath,   GameContentPathRebased);
 		ContentRootToPath.Emplace(ScriptRootPath, GameScriptPathRebased);
+		ContentRootToPath.Emplace(ExtraRootPath,  GameExtraPathRebased);
 		ContentRootToPath.Emplace(TempRootPath,   GameSavedPathRebased);
 		ContentRootToPath.Emplace(ConfigRootPath, GameConfigPathRebased);
 
@@ -979,7 +987,10 @@ bool FPackageName::DoesPackageExist(const FString& LongPackageName, const FGuid*
 	// Convert to filename (no extension yet).
 	FString Filename = LongPackageNameToFilename(PackageName, TEXT(""));
 	// Find the filename (with extension).
-	bFoundFile = FindPackageFileWithoutExtension(Filename, Filename, InAllowTextFormats);
+	if (!WITH_EDITORONLY_DATA || !IsExtraPackage(PackageName))
+	{
+		bFoundFile = FindPackageFileWithoutExtension(Filename, Filename, InAllowTextFormats);
+	}
 
 	// On consoles, we don't support package downloading, so no need to waste any extra cycles/disk io dealing with it
 	if (!FPlatformProperties::RequiresCookedData() && bFoundFile && Guid != NULL)
@@ -1004,7 +1015,7 @@ bool FPackageName::DoesPackageExist(const FString& LongPackageName, const FGuid*
 		delete PackageReader;
 	}
 
-	if (bFoundFile && OutFilename)
+	if (OutFilename && (bFoundFile || (WITH_EDITORONLY_DATA && IsExtraPackage(PackageName))))
 	{
 		*OutFilename = Filename;
 	}
@@ -1554,6 +1565,11 @@ FStringView FPackageName::ObjectPathToObjectName(FStringView InObjectPath)
 	return ObjectPathToObjectNameImpl(InObjectPath);
 }
 
+bool FPackageName::IsExtraPackage(const FString& InPackageName)
+{
+	return InPackageName.StartsWith(FLongPackagePathsSingleton::Get().ExtraRootPath);
+}
+
 bool FPackageName::IsScriptPackage(const FString& InPackageName)
 {
 	return InPackageName.StartsWith(FLongPackagePathsSingleton::Get().ScriptRootPath);
diff --git a/Engine/Source/Runtime/CoreUObject/Private/UObject/Class.cpp b/Engine/Source/Runtime/CoreUObject/Private/UObject/Class.cpp
--- a/Engine/Source/Runtime/CoreUObject/Private/UObject/Class.cpp
+++ b/Engine/Source/Runtime/CoreUObject/Private/UObject/Class.cpp
@@ -3609,7 +3609,7 @@ UObject* UClass::CreateDefaultObject()
 				FString PackageName;
 				FString CDOName;
 				bool bDoNotify = false;
-				if (GIsInitialLoad && GetOutermost()->HasAnyPackageFlags(PKG_CompiledIn))
+				if (GIsInitialLoad && GetOutermost()->HasAnyPackageFlags(PKG_CompiledIn) && !GetOutermost()->HasAnyPackageFlags(PKG_RuntimeGenerated))
 				{
 					PackageName = GetOutermost()->GetFName().ToString();
 					CDOName = GetDefaultObjectName().ToString();
@@ -5167,7 +5167,10 @@ void UClass::AssembleReferenceTokenStreams()
 		{
 			// Force the default object to be created (except when we're in the middle of exit purge -
 			// this may happen if we exited PreInit early because of error).
-			if (!GExitPurge)
+			// 
+			// Keep from handling script generated classes here, as those systems handle CDO 
+			// instantiation themselves.
+			if (!GExitPurge && !Class->HasAnyFlags(RF_BeingRegenerated))
 			{
 				Class->GetDefaultObject(); // Force the default object to be constructed if it isn't already
 			}
diff --git a/Engine/Source/Runtime/CoreUObject/Public/Blueprint/BlueprintSupport.h b/Engine/Source/Runtime/CoreUObject/Public/Blueprint/BlueprintSupport.h
--- a/Engine/Source/Runtime/CoreUObject/Public/Blueprint/BlueprintSupport.h
+++ b/Engine/Source/Runtime/CoreUObject/Public/Blueprint/BlueprintSupport.h
@@ -61,6 +61,7 @@ struct FBlueprintWarningDeclaration
 };
 
 typedef void (*FFlushReinstancingQueueFPtr)();
+typedef void (*FClassReparentingFPtr)(const TMap<UClass*, UClass*>&);
 
 /** 
  * This set of functions contains blueprint related UObject functionality.
@@ -88,7 +89,10 @@ struct FBlueprintSupport
 
 	/** Checks for any old instances and reinstances them: */
 	static void FlushReinstancingQueue();
-	COREUOBJECT_API static void SetFlushReinstancingQueueFPtr(FFlushReinstancingQueueFPtr Ptr);
+	COREUOBJECT_API static void SetFlushReinstancingQueueFPtr(FFlushReinstancingQueueFPtr Ptr);	
+
+	COREUOBJECT_API static void ReparentHierarchies(const TMap<UClass*, UClass*>& OldClassToNewClass);
+	COREUOBJECT_API static void SetClassReparentingFPtr(FClassReparentingFPtr Ptr);
 
 	/** Tells if the specified object is one of the many flavors of FLinkerPlaceholderBase that we have. */
 	COREUOBJECT_API static bool IsDeferredDependencyPlaceholder(UObject* LoadedObj);
diff --git a/Engine/Source/Runtime/CoreUObject/Public/Misc/PackageName.h b/Engine/Source/Runtime/CoreUObject/Public/Misc/PackageName.h
--- a/Engine/Source/Runtime/CoreUObject/Public/Misc/PackageName.h
+++ b/Engine/Source/Runtime/CoreUObject/Public/Misc/PackageName.h
@@ -459,6 +459,11 @@ public:
 	static FStringView ObjectPathToObjectName(FStringView InObjectPath);
 	static FString ObjectPathToObjectName(const FString& InObjectPath);
 
+	/**
+	 * Checks the root of the package's path to see if it's an extra package
+	 */
+	static bool IsExtraPackage(const FString& InPackageName);
+
 	/**
 	 * Checks the root of the package's path to see if it is a script package
 	 * @return true if the root of the path matches the script path
diff --git a/Engine/Source/Runtime/CoreUObject/Public/UObject/ObjectMacros.h b/Engine/Source/Runtime/CoreUObject/Public/UObject/ObjectMacros.h
--- a/Engine/Source/Runtime/CoreUObject/Public/UObject/ObjectMacros.h
+++ b/Engine/Source/Runtime/CoreUObject/Public/UObject/ObjectMacros.h
@@ -130,7 +130,7 @@ enum EPackageFlags
 //	PKG_Unused						= 0x04000000,
 //	PKG_Unused						= 0x08000000,	
 //	PKG_Unused						= 0x10000000,	
-//	PKG_Unused						= 0x20000000,
+	PKG_RuntimeGenerated			= 0x20000000,	///< This package contains elements that are runtime generated, and may not follow standard loading order rules
 	PKG_ReloadingForCooker			= 0x40000000,   ///< This package is reloading in the cooker, try to avoid getting data we will never need. We won't save this package.
 	PKG_FilterEditorOnly			= 0x80000000,	///< Package has editor-only data filtered out
 };
@@ -502,7 +502,7 @@ enum EObjectFlags
 	RF_FinishDestroyed			=0x00010000,	///< FinishDestroy has been called on the object.
 
 	// Misc. Flags
-	RF_BeingRegenerated			=0x00020000,	///< Flagged on UObjects that are used to create UClasses (e.g. Blueprints) while they are regenerating their UClass on load (See FLinkerLoad::CreateExport())
+	RF_BeingRegenerated			=0x00020000,	///< Flagged on UObjects that are used to create UClasses (e.g. Blueprints) while they are regenerating their UClass on load (See FLinkerLoad::CreateExport()), as well as UClass objects in the midst of being created
 	RF_DefaultSubObject			=0x00040000,	///< Flagged on subobjects that are defaults
 	RF_WasLoaded				=0x00080000,	///< Flagged on UObjects that were loaded
 	RF_TextExportTransient		=0x00100000,	///< Do not export object to text form (e.g. copy/paste). Generally used for sub-objects that can be regenerated from data in their parent object.
diff --git a/Engine/Source/Runtime/CoreUObject/Public/UObject/UObjectBase.h b/Engine/Source/Runtime/CoreUObject/Public/UObject/UObjectBase.h
--- a/Engine/Source/Runtime/CoreUObject/Public/UObject/UObjectBase.h
+++ b/Engine/Source/Runtime/CoreUObject/Public/UObject/UObjectBase.h
@@ -214,6 +214,7 @@ private:
 	UObject*						OuterPrivate;
 	
 	friend class FBlueprintCompileReinstancer;
+	friend class FContextObjectManager;
 
 	/** This is used by the reinstancer to re-class and re-archetype the current instances of a class before recompiling */
 	void SetClass(UClass* NewClass);
diff --git a/Engine/Source/Runtime/Engine/Public/Physics/ImmediatePhysics/ImmediatePhysicsActorHandle.h b/Engine/Source/Runtime/Engine/Public/Physics/ImmediatePhysics/ImmediatePhysicsActorHandle.h
--- a/Engine/Source/Runtime/Engine/Public/Physics/ImmediatePhysics/ImmediatePhysicsActorHandle.h
+++ b/Engine/Source/Runtime/Engine/Public/Physics/ImmediatePhysics/ImmediatePhysicsActorHandle.h
@@ -1,4 +1,4 @@
-// Copyright Epic Games, Inc. All Rights Reserved.
+// Copyright 1998-2019 Epic Games, Inc. All Rights Reserved.
 
 #pragma once
 
diff --git a/Engine/Source/Runtime/Engine/Public/Physics/ImmediatePhysics/ImmediatePhysicsSimulation.h b/Engine/Source/Runtime/Engine/Public/Physics/ImmediatePhysics/ImmediatePhysicsSimulation.h
--- a/Engine/Source/Runtime/Engine/Public/Physics/ImmediatePhysics/ImmediatePhysicsSimulation.h
+++ b/Engine/Source/Runtime/Engine/Public/Physics/ImmediatePhysics/ImmediatePhysicsSimulation.h
@@ -1,4 +1,4 @@
-// Copyright Epic Games, Inc. All Rights Reserved.
+// Copyright 1998-2019 Epic Games, Inc. All Rights Reserved.
 
 #pragma once
 
diff --git a/Engine/Source/Runtime/Experimental/Chaos/Public/Chaos/PBDRigidParticles.h b/Engine/Source/Runtime/Experimental/Chaos/Public/Chaos/PBDRigidParticles.h
--- a/Engine/Source/Runtime/Experimental/Chaos/Public/Chaos/PBDRigidParticles.h
+++ b/Engine/Source/Runtime/Experimental/Chaos/Public/Chaos/PBDRigidParticles.h
@@ -66,10 +66,10 @@ class TPBDRigidParticles : public TRigidParticles<T, d>
 
     // Must be reinterpret cast instead of static_cast as it's a forward declare
 	typedef TPBDRigidParticleHandle<T, d> THandleType;
-	CHAOS_API const THandleType* Handle(int32 Index) const { return reinterpret_cast<const THandleType*>(TGeometryParticles<T,d>::Handle(Index)); }
+	CHAOS_API const THandleType* Handle(int32 Index) const { return static_cast<const THandleType*>(TGeometryParticles<T,d>::Handle(Index)); }
 
 	//cannot be reference because double pointer would allow for badness, but still useful to have non const access to handle
-	CHAOS_API THandleType* Handle(int32 Index) { return reinterpret_cast<THandleType*>(TGeometryParticles<T, d>::Handle(Index)); }
+	CHAOS_API THandleType* Handle(int32 Index) { return static_cast<THandleType*>(TGeometryParticles<T, d>::Handle(Index)); }
 
 	CHAOS_API void SetSleeping(int32 Index, bool bSleeping)
 	{
@@ -82,7 +82,7 @@ class TPBDRigidParticles : public TRigidParticles<T, d>
 		bool CurrentlySleeping = this->ObjectState(Index) == EObjectStateType::Sleeping;
 		if (CurrentlySleeping != bSleeping)
 		{
-			TGeometryParticleHandle<T, d>* Particle = reinterpret_cast<TGeometryParticleHandle<T, d>*>(this->Handle(Index));
+			TGeometryParticleHandle<T, d>* Particle = static_cast<TGeometryParticleHandle<T, d>*>(this->Handle(Index));
 			this->AddSleepData(Particle, bSleeping);
 		}
 
@@ -147,7 +147,7 @@ class TPBDRigidParticles : public TRigidParticles<T, d>
 		const bool bNewSleeping = InObjectState == EObjectStateType::Sleeping;
 		if(bCurrentSleeping != bNewSleeping)
 		{
-			TGeometryParticleHandle<T, d>* Particle = reinterpret_cast<TGeometryParticleHandle<T, d>*>(this->Handle(Index));
+			TGeometryParticleHandle<T, d>* Particle = static_cast<TGeometryParticleHandle<T, d>*>(this->Handle(Index));
  			this->AddSleepData(Particle, bNewSleeping);
 		}
 
diff --git a/Engine/Source/Runtime/OpenGLDrv/Private/Linux/OpenGLLinux.cpp b/Engine/Source/Runtime/OpenGLDrv/Private/Linux/OpenGLLinux.cpp
--- a/Engine/Source/Runtime/OpenGLDrv/Private/Linux/OpenGLLinux.cpp
+++ b/Engine/Source/Runtime/OpenGLDrv/Private/Linux/OpenGLLinux.cpp
@@ -11,6 +11,7 @@
 #include "OpenGLDrvPrivate.h"
 #include "ComponentReregisterContext.h"
 #include "Linux/LinuxPlatformApplicationMisc.h"
+#include "GenericPlatform/GenericPlatformFramePacer.h"
 
 /*------------------------------------------------------------------------------
 	OpenGL function pointers.
@@ -429,6 +430,7 @@ bool PlatformBlitToViewport(FPlatformOpenGLDevice* Device,
 							bool bPresent,
 							bool bLockToVsync)
 {
+	int32 SyncInterval = FGenericPlatformRHIFramePacer::GetFramePace();
 	FPlatformOpenGLContext* const Context = Viewport.GetGLContext();
 
 	check( Context && Context->hWnd );
