Change to using a TWeakObjectPtr to avoid WITH_EDITOR/UPROPERTY issues
#rnx

diff --git a/Engine/Source/Editor/Kismet/Private/SSCSEditor.cpp b/Engine/Source/Editor/Kismet/Private/SSCSEditor.cpp
--- a/Engine/Source/Editor/Kismet/Private/SSCSEditor.cpp
+++ b/Engine/Source/Editor/Kismet/Private/SSCSEditor.cpp
@@ -1318,7 +1318,7 @@ UActorComponent* FSCSEditorTreeNode::FindComponentInstanceInActor(const AActor*
 				else if (World != nullptr && World->WorldType == EWorldType::EditorPreview)
 				{
 					// If this is the preview actor, return the cached component instance that's being used for the preview actor prior to recompiling the Blueprint
-					ComponentInstance = SCS_Node->EditorComponentInstance;
+					ComponentInstance = SCS_Node->EditorComponentInstance.Get();
 				}
 			}
 		}
diff --git a/Engine/Source/Runtime/Engine/Classes/Engine/SCS_Node.h b/Engine/Source/Runtime/Engine/Classes/Engine/SCS_Node.h
--- a/Engine/Source/Runtime/Engine/Classes/Engine/SCS_Node.h
+++ b/Engine/Source/Runtime/Engine/Classes/Engine/SCS_Node.h
@@ -89,8 +89,7 @@ class USCS_Node : public UObject
 #if WITH_EDITOR
 
 	/** The scene component constructed for component editing in the SCS editor */
-	UPROPERTY(transient)
-	USceneComponent* EditorComponentInstance;
+	TWeakObjectPtr<USceneComponent> EditorComponentInstance;
 
 	/** Make sure, the guid is proper - backward compatibility */
 	ENGINE_API void ValidateGuid();
