Fix issues with creating a new attachment if its default mobility is not compatible with the component it is going to be attached to.

diff --git a/Engine/Source/Editor/Kismet/Private/SSCSEditor.cpp b/Engine/Source/Editor/Kismet/Private/SSCSEditor.cpp
--- a/Engine/Source/Editor/Kismet/Private/SSCSEditor.cpp
+++ b/Engine/Source/Editor/Kismet/Private/SSCSEditor.cpp
@@ -5301,6 +5301,17 @@ UActorComponent* SSCSEditor::AddNewComponent( UClass* NewComponentClass, UObject
 						AttachTo = ActorInstance->GetRootComponent();
 					}
 					check(AttachTo != nullptr);
+
+					// Make sure that the mobility of the new scene component is such that we can attach it
+					if (AttachTo->Mobility == EComponentMobility::Movable)
+					{
+						NewSceneComponent->Mobility = EComponentMobility::Movable;
+					}
+					else if (AttachTo->Mobility == EComponentMobility::Stationary && NewSceneComponent->Mobility == EComponentMobility::Static)
+					{
+						NewSceneComponent->Mobility = EComponentMobility::Stationary;
+					}
+
 					NewSceneComponent->AttachToComponent(AttachTo, FAttachmentTransformRules::KeepRelativeTransform);
 				}
 			}
