Fix an editor crash when choosing to compile a Blueprint with an intermediate graph open in an inactive document tab.

#rnx

diff --git a/Engine/Source/Editor/Kismet/Private/WorkflowOrientedApp/WorkflowTabManager.cpp b/Engine/Source/Editor/Kismet/Private/WorkflowOrientedApp/WorkflowTabManager.cpp
--- a/Engine/Source/Editor/Kismet/Private/WorkflowOrientedApp/WorkflowTabManager.cpp
+++ b/Engine/Source/Editor/Kismet/Private/WorkflowOrientedApp/WorkflowTabManager.cpp
@@ -487,8 +487,13 @@ void FDocumentTracker::OnTabForegrounded(TSharedPtr<SDockTab> ForegroundedTab, T
 	}
 
 	TSharedPtr<FGenericTabHistory> CurrentTabHistory = GetCurrentTabHistory();
+	TSharedPtr<FGenericTabHistory> NewTabHistory;
+	if (NewTabInfo.IsValid())
+	{
+		NewTabHistory = NewTabInfo->GetCurrentHistory();
+	}
 
-	if ( ForegroundFactory.IsValid() && NewTabInfo.IsValid() && NewTabInfo->GetCurrentHistory() != CurrentTabHistory )
+	if ( ForegroundFactory.IsValid() && NewTabHistory.IsValid() && NewTabHistory->IsHistoryValid() && NewTabHistory != CurrentTabHistory )
 	{
 		// If a tab was manually foregrounded, need to add tab history
 		NewTabInfo->AddTabHistory(ForegroundFactory->CreateTabHistoryNode(NewTabInfo->GetPayload()), true);
diff --git a/Engine/Source/Editor/Kismet/Public/WorkflowOrientedApp/WorkflowUObjectDocuments.h b/Engine/Source/Editor/Kismet/Public/WorkflowOrientedApp/WorkflowUObjectDocuments.h
--- a/Engine/Source/Editor/Kismet/Public/WorkflowOrientedApp/WorkflowUObjectDocuments.h
+++ b/Engine/Source/Editor/Kismet/Public/WorkflowOrientedApp/WorkflowUObjectDocuments.h
@@ -40,7 +40,7 @@ public:
 	{
 		if (OtherPayload->PayloadType == PayloadType)
 		{
-			return this->DocumentID == FTabPayload_UObject::CastChecked<UObject>(OtherPayload);
+			return this->DocumentID.HasSameIndexAndSerialNumber(StaticCastSharedRef<FTabPayload_UObject>(OtherPayload)->DocumentID);
 		}
 
 		return false;
