Blueprint Editor: While loading, show which blueprint library is being loaded so that developers can know which libraries take a long time to load.

diff --git a/Engine/Source/Editor/Kismet/Private/BlueprintEditor.cpp b/Engine/Source/Editor/Kismet/Private/BlueprintEditor.cpp
--- a/Engine/Source/Editor/Kismet/Private/BlueprintEditor.cpp
+++ b/Engine/Source/Editor/Kismet/Private/BlueprintEditor.cpp
@@ -164,6 +164,8 @@
 #include "Widgets/Input/SNumericEntryBox.h"
 #include "Subsystems/AssetEditorSubsystem.h"
 
+DEFINE_LOG_CATEGORY_STATIC(LogBlueprintEditor, Log, All);
+
 #define LOCTEXT_NAMESPACE "BlueprintEditor"
 
 static int32 EnableAutomaticLibraryAssetLoading = 1;
@@ -1753,6 +1755,20 @@ void FBlueprintEditor::LoadLibrariesFromAssetRegistry()
 			const FString BPMacroTypeStr( TEXT( "BPTYPE_MacroLibrary" ) );
 			const FString BPFunctionTypeStr( TEXT( "BPTYPE_FunctionLibrary" ) );
 
+			struct FExpensiveObjectRecord
+			{
+				FExpensiveObjectRecord() : Seconds(0.0) {}
+				FExpensiveObjectRecord(double InSeconds, const FName InPath) : Seconds(InSeconds), Path(InPath) {}
+				double Seconds;
+				FName Path;
+			};
+
+			TArray<FExpensiveObjectRecord> ExpensiveObjects;
+			const double MinSecondsToReportExpensiveObject = 1.0;
+			const int32 MaxExpensiveObjectsToList = 10;
+			int32 NumLibariesLoaded = 0;
+
+			const double StartTimeAll = FPlatformTime::Seconds();
 			for (int32 AssetIndex = 0; AssetIndex < AssetData.Num(); ++AssetIndex)
 			{
 				FString TagValue = AssetData[ AssetIndex ].GetTagValueRef<FString>(BPTypeName);
@@ -1774,6 +1790,11 @@ void FBlueprintEditor::LoadLibrariesFromAssetRegistry()
 
 					if( AllowLoadBP )
 					{
+						GWarn->StatusUpdate(AssetIndex, AssetData.Num(), FText::FromName(AssetData[AssetIndex].AssetName));
+
+						++NumLibariesLoaded;
+						const double StartTime = FPlatformTime::Seconds();
+
 						// Load the blueprint
 						UBlueprint* BlueprintLibPtr = LoadObject<UBlueprint>(NULL, *BlueprintPath, NULL, 0, NULL);
 						if (BlueprintLibPtr )
@@ -1781,10 +1802,31 @@ void FBlueprintEditor::LoadLibrariesFromAssetRegistry()
 							StandardLibraries.AddUnique(BlueprintLibPtr);
 							WatchViewer::UpdateWatchListFromBlueprint(BlueprintLibPtr);
 						}
+
+						const double ElapsedTime = FPlatformTime::Seconds() - StartTime;
+						if (ElapsedTime > MinSecondsToReportExpensiveObject)
+						{
+							ExpensiveObjects.Add(FExpensiveObjectRecord(ElapsedTime, AssetData[AssetIndex].PackageName));
+						}
 					}
 				}
 
 			}
+
+			if (ExpensiveObjects.Num() > 0)
+			{
+				const double TotalSeconds = FPlatformTime::Seconds() - StartTimeAll;
+				UE_LOG(LogBlueprintEditor, Log, TEXT("Perf: %.1f total seconds to load all %d blueprint libraries in project. Avoid references to content in blueprint libraries to shorten this time."), TotalSeconds, NumLibariesLoaded);
+
+				// Log the most expensive objects to load
+				Algo::Sort(ExpensiveObjects, [](FExpensiveObjectRecord& A, FExpensiveObjectRecord& B) { return A.Seconds > B.Seconds; });
+				for (int32 i=0; i < ExpensiveObjects.Num() && i < MaxExpensiveObjectsToList; ++i)
+				{
+					FExpensiveObjectRecord& ExpensiveObjectRecord = ExpensiveObjects[i];
+					UE_LOG(LogBlueprintEditor, Log, TEXT("Perf: %.1f seconds loading: %s"), ExpensiveObjectRecord.Seconds, *ExpensiveObjectRecord.Path.ToString());
+				}
+			}
+
 			GWarn->EndSlowTask();
 		}
 	}
