Fix it so it will correctly reuse existing open tabs for BP function nodes again, fixes a regression when hitting breakpoints or double clicking
Also re-fixes UE-64808, which got broken by a later perf fix

diff --git a/Engine/Source/Editor/Kismet/Private/BlueprintEditor.cpp b/Engine/Source/Editor/Kismet/Private/BlueprintEditor.cpp
--- a/Engine/Source/Editor/Kismet/Private/BlueprintEditor.cpp
+++ b/Engine/Source/Editor/Kismet/Private/BlueprintEditor.cpp
@@ -8776,17 +8776,8 @@ TSharedPtr<SGraphEditor> FBlueprintEditor::OpenGraphAndBringToFront(UEdGraph* Gr
 	// First, switch back to standard mode
 	SetCurrentMode(FBlueprintEditorApplicationModes::StandardBlueprintEditorMode);
 
-		
-	TSharedPtr<SDockTab> TabWithGraph;
-	TSharedPtr<SGraphEditor> FocusedGraphEd = FocusedGraphEdPtr.Pin();
-	if (FocusedGraphEd.IsValid() && FocusedGraphEd->GetCurrentGraph() == Graph)
-	{
-		TabWithGraph = OpenDocument(Graph, FDocumentTracker::CreateHistoryEvent);
-	}
-	else
-	{
-		 TabWithGraph = OpenDocument(Graph, FDocumentTracker::NavigatingCurrentDocument);
-	}
+	// This will either reuse an existing tab or spawn a new one
+	TSharedPtr<SDockTab> TabWithGraph = OpenDocument(Graph, FDocumentTracker::OpenNewDocument);
 
 	// We know that the contents of the opened tabs will be a graph editor.
 	TSharedRef<SGraphEditor> NewGraphEditor = StaticCastSharedRef<SGraphEditor>(TabWithGraph->GetContent());
diff --git a/Engine/Source/Editor/Kismet/Private/WorkflowOrientedApp/WorkflowTabManager.cpp b/Engine/Source/Editor/Kismet/Private/WorkflowOrientedApp/WorkflowTabManager.cpp
--- a/Engine/Source/Editor/Kismet/Private/WorkflowOrientedApp/WorkflowTabManager.cpp
+++ b/Engine/Source/Editor/Kismet/Private/WorkflowOrientedApp/WorkflowTabManager.cpp
@@ -572,8 +572,14 @@ TSharedPtr<SDockTab> FDocumentTracker::OpenDocument(TSharedPtr<FTabPayload> InPa
 		TSharedPtr<FTabInfo> LastEditedTabInfo = GetLastEditedTabInfo();
 		if(LastEditedTabInfo.IsValid() && LastEditedTabInfo->PayloadMatches(InPayload))
 		{
-			LastEditedTabInfo->AddTabHistory(Factory->CreateTabHistoryNode(InPayload), true);
-			return LastEditedTabInfo->GetTab().Pin();
+			TSharedPtr<SDockTab> Tab = LastEditedTabInfo->GetTab().Pin();
+			if (Tab.IsValid())
+			{
+				LastEditedTabInfo->AddTabHistory(Factory->CreateTabHistoryNode(InPayload), true);
+				// Ensure that the tab appears if the tab isn't currently in the foreground.
+				Tab->ActivateInParent(ETabActivationCause::SetDirectly);
+			}
+			return Tab;
 		}
 		else
 		{
