Disable drag/drop replacing of variable nodes that have split pins

diff --git a/Engine/Source/Editor/Kismet/Private/BPVariableDragDropAction.cpp b/Engine/Source/Editor/Kismet/Private/BPVariableDragDropAction.cpp
--- a/Engine/Source/Editor/Kismet/Private/BPVariableDragDropAction.cpp
+++ b/Engine/Source/Editor/Kismet/Private/BPVariableDragDropAction.cpp
@@ -63,6 +63,30 @@ void FKismetVariableDragDropAction::GetLinksThatWillBreak(	UEdGraphNode* Node, F
 	}
 }
 
+namespace
+{
+	/**
+	* Helper function to determine if a node has any split pins on it.
+	*
+	* @return	True if there is a split pin on the node
+	*/
+	static bool NodeHasSplitPins(UEdGraphNode* InNode)
+	{
+		if (InNode)
+		{
+			for (UEdGraphPin* Pin : InNode->Pins)
+			{
+				// If a pin has no parent node but has SubPins then it is a split pin
+				if (Pin && !Pin->ParentPin && Pin->SubPins.Num() > 0)
+				{
+					return true;
+				}
+			}
+		}
+		return false;
+	}
+}
+
 void FKismetVariableDragDropAction::HoverTargetChanged()
 {
 	FProperty* VariableProperty = GetVariableProperty();
@@ -185,8 +209,14 @@ void FKismetVariableDragDropAction::HoverTargetChanged()
 			const UBlueprint* Blueprint = FBlueprintEditorUtils::FindBlueprintForNode(VarNodeUnderCursor);
 			const bool bWritableProperty = (FBlueprintEditorUtils::IsPropertyWritableInBlueprint(Blueprint, VariableProperty) == FBlueprintEditorUtils::EPropertyWritableState::Writable);
 			const bool bCanWriteIfNeeded = bIsRead || bWritableProperty;
+			// If this node has split pins then the reconstruction cannot handle it, so don't allow it
+			const bool bHasSplitPins = NodeHasSplitPins(VarNodeUnderCursor);
 
-			if (bCanWriteIfNeeded)
+			if (bHasSplitPins)
+			{
+				SetFeedbackMessageError(FText::Format(LOCTEXT("SplitPinVar_Error", "Cannot change '{VariableName}' because it has split pins"), Args));
+			}
+			else if (bCanWriteIfNeeded)
 			{
 				Args.Add(TEXT("ReadOrWrite"), bIsRead ? LOCTEXT("Read", "read") : LOCTEXT("Write", "write"));
 				SetFeedbackMessageOK(WillBreakLinks(VarNodeUnderCursor, VariableProperty) ?
@@ -235,7 +265,7 @@ FReply FKismetVariableDragDropAction::DroppedOnPin(FVector2D ScreenPosition, FVe
 
 				FProperty* VariableProperty = GetVariableProperty();
 
-				if (CanVariableBeDropped(VariableProperty, *TargetPin->GetOwningNode()->GetGraph()))
+				if (CanVariableBeDropped(VariableProperty, *TargetPin->GetOwningNode()->GetGraph()) && !NodeHasSplitPins(TargetPin->GetOwningNode()))
 				{
 					const bool bIsRead = (TargetPin->Direction == EGPD_Input) && !bIsExecPin;
 					const UBlueprint* Blueprint = FBlueprintEditorUtils::FindBlueprintForNode(TargetPin->GetOwningNode());
@@ -272,12 +302,12 @@ FReply FKismetVariableDragDropAction::DroppedOnNode(FVector2D ScreenPosition, FV
 
 	if (TargetNode && (VariableName != TargetNode->GetVarName()))
 	{
-		const FScopedTransaction Transaction( LOCTEXT("ReplacePinVariable", "Replace Pin Variable") );
-
 		FProperty* VariableProperty = GetVariableProperty();
 
-		if(CanVariableBeDropped(VariableProperty, *TargetNode->GetGraph()))
+		if(CanVariableBeDropped(VariableProperty, *TargetNode->GetGraph()) && !NodeHasSplitPins(TargetNode))
 		{
+			const FScopedTransaction Transaction(LOCTEXT("ReplacePinVariable", "Replace Pin Variable"));
+
 			const FName OldVarName = TargetNode->GetVarName();
 			const UEdGraphSchema_K2* Schema = Cast<const UEdGraphSchema_K2>(TargetNode->GetSchema());
 
@@ -290,15 +320,14 @@ FReply FKismetVariableDragDropAction::DroppedOnNode(FVector2D ScreenPosition, FV
 			DropOnBlueprint->Modify();
 			TargetNode->Modify();
 
-			if (Pin != NULL)
+			if (Pin != nullptr)
 			{
 				Pin->Modify();
 			}
 
 			UEdGraphSchema_K2::ConfigureVarNode(TargetNode, VariableName, VariableSource.Get(), DropOnBlueprint);
 
-
-			if ((Pin == NULL) || (Pin->LinkedTo.Num() == BadLinks.Num()) || (Schema == NULL))
+			if ((Pin == nullptr) || (Pin->LinkedTo.Num() == BadLinks.Num()) || (Schema == nullptr))
 			{
 				TargetNode->GetSchema()->ReconstructNode(*TargetNode);
 			}
