Blueprints: Added an initial prototype of namespacing and importing namespaces in Blueprints
- Currently only affects which BP macro and function libraries are loaded when opening the Blueprint Editor:
- anything in the 'global' namespace is loaded, plus anything imported on a project-wide (in project settings) or per-user basis (in editor preferences)
- This can be used to dramatically reduce the first time to open the BP editor in projects with a lot of macro/function libs that reference other content
- Note: Currently individual BPs cannot import specific namespaces (however, using a macro/function from a library will cause it to automatically load when that BP is loaded), and the namespacing does not affect the functionality of the BP context menu

diff --git a/Engine/Source/Editor/BlueprintGraph/Public/BlueprintEditorSettings.h b/Engine/Source/Editor/BlueprintGraph/Public/BlueprintEditorSettings.h
--- a/Engine/Source/Editor/BlueprintGraph/Public/BlueprintEditorSettings.h
+++ b/Engine/Source/Editor/BlueprintGraph/Public/BlueprintEditorSettings.h
@@ -6,6 +6,7 @@
 #include "UObject/ObjectMacros.h"
 #include "UObject/Object.h"
 #include "Engine/Blueprint.h"
+#include "Engine/DeveloperSettings.h"
 #include "BlueprintEditorSettings.generated.h"
 
 UENUM()
@@ -150,6 +151,11 @@ public:
 	UPROPERTY(EditAnywhere, config, Category = DeveloperTools, meta = (DisplayName = "Display Unique Names for Blueprint Nodes"))
 	bool bBlueprintNodeUniqueNames;
 
+public:
+	// The list of namespaces to always expose in any Blueprint (local per-user)
+	UPROPERTY(EditAnywhere, config, Category = Experimental)
+	TArray<FString> NamespacesToAlwaysInclude;
+
 // Perf Settings
 public:
 	/** If enabled, additional details will be displayed in the Compiler Results tab after compiling a blueprint. */
@@ -169,3 +175,14 @@ protected:
 	virtual void PostEditChangeProperty(struct FPropertyChangedEvent& PropertyChangedEvent) override;
 	//~ End UObject Interface
 };
+
+UCLASS(config=Editor, defaultconfig)
+class BLUEPRINTGRAPH_API UBlueprintProjectSettings : public UDeveloperSettings
+{
+public:
+	GENERATED_BODY()
+
+	// The list of namespaces to always expose in any Blueprint (for all users of the game/project)
+	UPROPERTY(EditAnywhere, config, Category=Experimental)
+	TArray<FString> NamespacesToAlwaysInclude;
+};
diff --git a/Engine/Source/Editor/Kismet/Private/BlueprintEditor.cpp b/Engine/Source/Editor/Kismet/Private/BlueprintEditor.cpp
--- a/Engine/Source/Editor/Kismet/Private/BlueprintEditor.cpp
+++ b/Engine/Source/Editor/Kismet/Private/BlueprintEditor.cpp
@@ -1727,20 +1727,80 @@ void FBlueprintEditor::CommonInitialization(const TArray<UBlueprint*>& InitBluep
 	FKismetEditorUtilities::OnBlueprintUnloaded.AddSP(this, &FBlueprintEditor::OnBlueprintUnloaded);
 }
 
+struct FBlueprintNamespaceHelper
+{
+	TSet<FString> FullyQualifiedListOfNamespaces;
+
+	FBlueprintNamespaceHelper()
+	{
+		AddNamespaces(GetDefault<UBlueprintEditorSettings>()->NamespacesToAlwaysInclude);
+		AddNamespaces(GetDefault<UBlueprintProjectSettings>()->NamespacesToAlwaysInclude);
+	}
+
+	
+	void AddNamespaces(const TArray<FString>& List)
+	{
+		for (const FString& Entry : List)
+		{
+			FullyQualifiedListOfNamespaces.Add(Entry);
+		}
+	}
+
+	void AddNamespace(const FString& Namespace)
+	{
+		if (!Namespace.IsEmpty())
+		{
+			FullyQualifiedListOfNamespaces.Add(Namespace);
+		}
+	}
+
+	bool IsIncludedInNamespaceList(const FString& TestNamespace) const
+	{
+		// Empty namespace == global namespace
+		if (TestNamespace.IsEmpty())
+		{
+			return true;
+		}
+
+		// Check recursively to see if X.Y.Z is present, and if not X.Y (which contains X.Y.Z), and so on until we run out of path segments
+		if (FullyQualifiedListOfNamespaces.Contains(TestNamespace))
+		{
+			return true;
+		}
+		else
+		{
+			int32 RightmostDotIndex;
+			if (TestNamespace.FindLastChar(TEXT('.'), /*out*/ RightmostDotIndex))
+			{
+				if (RightmostDotIndex > 0)
+				{
+					return IsIncludedInNamespaceList(TestNamespace.Left(RightmostDotIndex));
+				}
+			}
+		}
+
+		return false;
+	}
+};
+
 void FBlueprintEditor::LoadLibrariesFromAssetRegistry()
 {
+	FBlueprintNamespaceHelper NamespaceHelper;
+
 	if (EnableAutomaticLibraryAssetLoading == 0)
 	{
 		return;
 	}
 
-	if( GetBlueprintObj() )
+	if (UBlueprint* BP = GetBlueprintObj())
 	{
-		FString UserDeveloperPath = FPackageName::FilenameToLongPackageName( FPaths::GameUserDeveloperDir());
-		FString DeveloperPath = FPackageName::FilenameToLongPackageName( FPaths::GameDevelopersDir() );
+		const FString UserDeveloperPath = FPackageName::FilenameToLongPackageName( FPaths::GameUserDeveloperDir());
+		const FString DeveloperPath = FPackageName::FilenameToLongPackageName( FPaths::GameDevelopersDir() );
 
-		//Don't allow loading blueprint macros & functions for interface & data-only blueprints
-		if( GetBlueprintObj()->BlueprintType != BPTYPE_Interface && GetBlueprintObj()->BlueprintType != BPTYPE_Const)
+		NamespaceHelper.AddNamespace(BP->BlueprintNamespace);
+
+		// Don't allow loading blueprint macros & functions for interface & data-only blueprints
+		if ((BP->BlueprintType != BPTYPE_Interface) && (BP->BlueprintType != BPTYPE_Const))
 		{
 			// Load the asset registry module
 			FAssetRegistryModule& AssetRegistryModule = FModuleManager::LoadModuleChecked<FAssetRegistryModule>(TEXT("AssetRegistry"));
@@ -1749,11 +1809,12 @@ void FBlueprintEditor::LoadLibrariesFromAssetRegistry()
 			TArray<FAssetData> AssetData;
 			AssetRegistryModule.Get().GetAssetsByClass(UBlueprint::StaticClass()->GetFName(), AssetData);
 
-			GWarn->BeginSlowTask(LOCTEXT("LoadingBlueprintAssetData", "Loading Blueprint Asset Data"), true );
+			GWarn->BeginSlowTask(LOCTEXT("LoadingBlueprintAssetData", "Loading Blueprint Asset Data"), true);
 
-			const FName BPTypeName( TEXT( "BlueprintType" ) );
-			const FString BPMacroTypeStr( TEXT( "BPTYPE_MacroLibrary" ) );
-			const FString BPFunctionTypeStr( TEXT( "BPTYPE_FunctionLibrary" ) );
+			const FName BPTypeName(GET_MEMBER_NAME_STRING_CHECKED(UBlueprint, BlueprintType));
+			const FName BPNamespaceName(GET_MEMBER_NAME_STRING_CHECKED(UBlueprint, BlueprintNamespace));
+			const FString BPMacroTypeStr(TEXT("BPTYPE_MacroLibrary"));
+			const FString BPFunctionTypeStr(TEXT("BPTYPE_FunctionLibrary"));
 
 			struct FExpensiveObjectRecord
 			{
@@ -1764,40 +1825,49 @@ void FBlueprintEditor::LoadLibrariesFromAssetRegistry()
 			};
 
 			TArray<FExpensiveObjectRecord> ExpensiveObjects;
-			const double MinSecondsToReportExpensiveObject = 1.0;
-			const int32 MaxExpensiveObjectsToList = 10;
+			const double MinSecondsToReportExpensiveObject = 0.2;
+			const int32 MaxExpensiveObjectsToList = 100;
 			int32 NumLibariesLoaded = 0;
 
 			const double StartTimeAll = FPlatformTime::Seconds();
-			for (int32 AssetIndex = 0; AssetIndex < AssetData.Num(); ++AssetIndex)
+			int32 AssetIndexBeingProcessed = 0;
+			for (const FAssetData& AssetEntry : AssetData)
 			{
-				FString TagValue = AssetData[ AssetIndex ].GetTagValueRef<FString>(BPTypeName);
+				const FString AssetBPType = AssetEntry.GetTagValueRef<FString>(BPTypeName);
 
-				//Only check for Blueprint Macros & Functions in the asset data for loading
-				if ( TagValue == BPMacroTypeStr || TagValue == BPFunctionTypeStr )
+				// Only check for Blueprint Macros & Functions in the asset data for loading
+				if ((AssetBPType == BPMacroTypeStr) || (AssetBPType == BPFunctionTypeStr))
 				{
-					FString BlueprintPath = AssetData[AssetIndex].ObjectPath.ToString();
+					const FString BlueprintPath = AssetEntry.ObjectPath.ToString();
+
+					bool bAllowLoadBP = true;
+
+					// See if this passes the namespace check
+					const FString AssetBPNamespace = AssetEntry.GetTagValueRef<FString>(BPNamespaceName);
+					bAllowLoadBP = bAllowLoadBP && NamespaceHelper.IsIncludedInNamespaceList(AssetBPNamespace);
 
-					//For blueprints inside developers folder, only allow the ones inside current user's developers folder.
-					bool AllowLoadBP = true;
-					if( BlueprintPath.StartsWith(DeveloperPath) )
+					// For blueprints inside developers folder, only allow the ones inside current user's developers folder.
+					if (bAllowLoadBP)
 					{
-						if(  !BlueprintPath.StartsWith(UserDeveloperPath) )
+						if (BlueprintPath.StartsWith(DeveloperPath))
 						{
-							AllowLoadBP = false;
+							if (!BlueprintPath.StartsWith(UserDeveloperPath))
+							{
+								bAllowLoadBP = false;
+							}
 						}
 					}
 
-					if( AllowLoadBP )
+					if (bAllowLoadBP)
 					{
-						GWarn->StatusUpdate(AssetIndex, AssetData.Num(), FText::FromName(AssetData[AssetIndex].AssetName));
+						GWarn->StatusUpdate(AssetIndexBeingProcessed, AssetData.Num(), FText::FromName(AssetEntry.AssetName));
 
 						++NumLibariesLoaded;
 						const double StartTime = FPlatformTime::Seconds();
 
 						// Load the blueprint
-						UBlueprint* BlueprintLibPtr = LoadObject<UBlueprint>(NULL, *BlueprintPath, NULL, 0, NULL);
-						if (BlueprintLibPtr )
+						UBlueprint* BlueprintLibPtr = LoadObject<UBlueprint>(nullptr, *BlueprintPath, nullptr, 0, nullptr);
+						if (BlueprintLibPtr)
 						{
 							StandardLibraries.AddUnique(BlueprintLibPtr);
 							WatchViewer::UpdateWatchListFromBlueprint(BlueprintLibPtr);
@@ -1806,11 +1876,12 @@ void FBlueprintEditor::LoadLibrariesFromAssetRegistry()
 						const double ElapsedTime = FPlatformTime::Seconds() - StartTime;
 						if (ElapsedTime > MinSecondsToReportExpensiveObject)
 						{
-							ExpensiveObjects.Add(FExpensiveObjectRecord(ElapsedTime, AssetData[AssetIndex].PackageName));
+							ExpensiveObjects.Add(FExpensiveObjectRecord(ElapsedTime, AssetEntry.PackageName));
 						}
 					}
 				}
 
+				++AssetIndexBeingProcessed;
 			}
 
 			if (ExpensiveObjects.Num() > 0)
diff --git a/Engine/Source/Runtime/Engine/Classes/Engine/Blueprint.h b/Engine/Source/Runtime/Engine/Classes/Engine/Blueprint.h
--- a/Engine/Source/Runtime/Engine/Classes/Engine/Blueprint.h
+++ b/Engine/Source/Runtime/Engine/Classes/Engine/Blueprint.h
@@ -532,6 +532,10 @@ public:
 	UPROPERTY(EditAnywhere, Category=BlueprintOptions, meta=(MultiLine=true), DuplicateTransient)
 	FString BlueprintDescription;
 
+	/** The namespace of this blueprint (if set, the Blueprint will be treated differently for the context menu) */
+	UPROPERTY(EditAnywhere, Category = BlueprintOptions, AssetRegistrySearchable)
+	FString BlueprintNamespace;
+
 	/** The category of the Blueprint, used to organize this Blueprint class when displayed in palette windows */
 	UPROPERTY(EditAnywhere, Category=BlueprintOptions)
 	FString BlueprintCategory;
