Move Find-in-Blueprint indexing to lower-priority thread pool.

diff --git a/Engine/Source/Editor/Kismet/Private/FindInBlueprintManager.cpp b/Engine/Source/Editor/Kismet/Private/FindInBlueprintManager.cpp
--- a/Engine/Source/Editor/Kismet/Private/FindInBlueprintManager.cpp
+++ b/Engine/Source/Editor/Kismet/Private/FindInBlueprintManager.cpp
@@ -1331,10 +1331,15 @@ public:
 			// Determine whether this is a full or partial indexing operation.
 			const bool bEnableFullIndexingPass = Controller->ShouldFullyIndexAssets();
 
-			// Determine whether to disable multiprocessing for this batch. If true, this batch will be processed only on this thread.
-			const bool bDisableMultiprocessing = !Controller->ShouldEnableMultiprocessing();
+			// Determine whether to disable multiprocessing for this batch. By default, use multiprocessing with low-priority worker threads.
+			EParallelForFlags ParallelForFlags = EParallelForFlags::Unbalanced | EParallelForFlags::BackgroundPriority;
+			if (!Controller->ShouldEnableMultiprocessing())
+			{
+				// This batch will be processed entirely on this thread.
+				ParallelForFlags |= EParallelForFlags::ForceSingleThread;
+			}
 
-			ParallelFor(AssetPathsToIndex.Num(), [&AssetPathsToIndex, Controller = this->Controller, bEnableFullIndexingPass, bDisableMultiprocessing](int32 ArrayIdx)
+			ParallelFor(AssetPathsToIndex.Num(), [&AssetPathsToIndex, Controller = this->Controller, bEnableFullIndexingPass, ParallelForFlags](int32 ArrayIdx)
 			{
 				CSV_CUSTOM_STAT(FindInBlueprint, IndexedAssetCountThisFrame, 1, ECsvCustomStatOp::Accumulate);
 
@@ -1366,7 +1371,7 @@ public:
 									ParallelFor(IndexNodes.Num(), [&IndexNodes](int32 NodeIdx)
 									{
 										IndexNodes[NodeIdx]->ParseAllChildData();
-									}, bDisableMultiprocessing);
+									}, ParallelForFlags);
 
 									TArray<FImaginaryFiBDataSharedPtr> ChildNodes;
 									for (FImaginaryFiBDataSharedPtr& NodePtr : IndexNodes)
@@ -1394,7 +1399,7 @@ public:
 					// Signal that indexing has been completed for this asset path, even though its cache entry is invalid.
 					Controller->IndexCompletedForAssetPath(AssetPath);
 				}
-			}, bDisableMultiprocessing);
+			}, ParallelForFlags);
 		}
 
 		return 0;
