Optimize BlueprintManager after map loading
  - Fix useless strlen in OnAssetAdded
  - Remove up to 4s of memcpy doing FString cloning during garbage collect after a map is loaded

  - 4.1s -> 16ms for aggregated CleanCache in SaveDirtyPackages after map loading

#rnx

diff --git a/Engine/Source/Editor/Kismet/Private/FindInBlueprintManager.cpp b/Engine/Source/Editor/Kismet/Private/FindInBlueprintManager.cpp
--- a/Engine/Source/Editor/Kismet/Private/FindInBlueprintManager.cpp
+++ b/Engine/Source/Editor/Kismet/Private/FindInBlueprintManager.cpp
@@ -2202,7 +2202,7 @@ void FFindInBlueprintSearchManager::ExtractUnloadedFiBData(const FAssetData& InA
 		}
 	}
 
-	NewSearchData.Value = *InFiBData;
+	NewSearchData.Value = InFiBData;
 
 	// This will be set to 'None' if the data is versioned. Deserialization of the actual version from the tag value is deferred until later.
 	NewSearchData.VersionInfo.FiBDataVersion = InFiBDataVersion;
@@ -2844,7 +2844,7 @@ void FFindInBlueprintSearchManager::CleanCache()
 
 		// If the database item is invalid, marked for deletion, stale or pending kill (if loaded), remove it from the database
 		const bool bEvenIfPendingKill = true;
-		const FSearchData& SearchData = SearchArray[SearchValuePair.Value];
+		FSearchData& SearchData = SearchArray[SearchValuePair.Value];
 		if (!SearchData.IsValid()
 			|| SearchData.IsMarkedForDeletion()
 			|| SearchData.Blueprint.IsStale()
@@ -2857,7 +2857,7 @@ void FFindInBlueprintSearchManager::CleanCache()
 		else
 		{
 			// Build the new map/array
-			NewSearchMap.Add(SearchValuePair.Key, NewSearchArray.Add(SearchData));
+			NewSearchMap.Add(SearchValuePair.Key, NewSearchArray.Add(MoveTemp(SearchData)));
 		}
 	}
 
