Blueprints: Skip preloading macro and function libraries when loading a data-only blueprint in the DOBP editor

diff --git a/Engine/Plugins/Experimental/ControlRig/Source/ControlRigEditor/Private/Editor/ControlRigEditor.cpp b/Engine/Plugins/Experimental/ControlRig/Source/ControlRigEditor/Private/Editor/ControlRigEditor.cpp
--- a/Engine/Plugins/Experimental/ControlRig/Source/ControlRigEditor/Private/Editor/ControlRigEditor.cpp
+++ b/Engine/Plugins/Experimental/ControlRig/Source/ControlRigEditor/Private/Editor/ControlRigEditor.cpp
@@ -195,7 +195,7 @@ void FControlRigEditor::InitControlRigEditor(const EToolkitMode::Type Mode, cons
 
 	InControlRigBlueprint->InitializeModelIfRequired();
 
-	CommonInitialization(ControlRigBlueprints);
+	CommonInitialization(ControlRigBlueprints, false);
 
 	for (UBlueprint* Blueprint : ControlRigBlueprints)
 	{
diff --git a/Engine/Source/Editor/AnimationBlueprintEditor/Private/AnimationBlueprintEditor.cpp b/Engine/Source/Editor/AnimationBlueprintEditor/Private/AnimationBlueprintEditor.cpp
--- a/Engine/Source/Editor/AnimationBlueprintEditor/Private/AnimationBlueprintEditor.cpp
+++ b/Engine/Source/Editor/AnimationBlueprintEditor/Private/AnimationBlueprintEditor.cpp
@@ -265,7 +265,7 @@ void FAnimationBlueprintEditor::InitAnimationBlueprintEditor(const EToolkitMode:
 	TArray<UBlueprint*> AnimBlueprints;
 	AnimBlueprints.Add(InAnimBlueprint);
 
-	CommonInitialization(AnimBlueprints);
+	CommonInitialization(AnimBlueprints, /*bShouldOpenInDefaultsMode=*/ false);
 
 	if(InAnimBlueprint->BlueprintType == BPTYPE_Interface)
 	{
diff --git a/Engine/Source/Editor/Kismet/Private/BlueprintEditor.cpp b/Engine/Source/Editor/Kismet/Private/BlueprintEditor.cpp
--- a/Engine/Source/Editor/Kismet/Private/BlueprintEditor.cpp
+++ b/Engine/Source/Editor/Kismet/Private/BlueprintEditor.cpp
@@ -1676,7 +1676,7 @@ struct FLoadObjectsFromAssetRegistryHelper
 	}
 };
 
-void FBlueprintEditor::CommonInitialization(const TArray<UBlueprint*>& InitBlueprints)
+void FBlueprintEditor::CommonInitialization(const TArray<UBlueprint*>& InitBlueprints, bool bShouldOpenInDefaultsMode)
 {
 	TSharedPtr<FBlueprintEditor> ThisPtr(SharedThis(this));
 
@@ -1702,7 +1702,10 @@ void FBlueprintEditor::CommonInitialization(const TArray<UBlueprint*>& InitBluep
 	if (InitBlueprints.Num() == 1)
 	{
 		// Load blueprint libraries
-		LoadLibrariesFromAssetRegistry();
+		if (!bShouldOpenInDefaultsMode)
+		{
+			LoadLibrariesFromAssetRegistry();
+		}
 
 		FLoadObjectsFromAssetRegistryHelper::Load<UUserDefinedEnum>(UserDefinedEnumerators);
 
@@ -1799,8 +1802,8 @@ void FBlueprintEditor::LoadLibrariesFromAssetRegistry()
 
 		NamespaceHelper.AddNamespace(BP->BlueprintNamespace);
 
-		// Don't allow loading blueprint macros & functions for interface & data-only blueprints
-		if ((BP->BlueprintType != BPTYPE_Interface) && (BP->BlueprintType != BPTYPE_Const))
+		// Interface blueprints don't show a node context menu anywhere so we can skip library loading
+		if (BP->BlueprintType != BPTYPE_Interface)
 		{
 			// Load the asset registry module
 			FAssetRegistryModule& AssetRegistryModule = FModuleManager::LoadModuleChecked<FAssetRegistryModule>(TEXT("AssetRegistry"));
@@ -1963,7 +1966,7 @@ void FBlueprintEditor::InitBlueprintEditor(
 	const FName BlueprintEditorAppName = FName(TEXT("BlueprintEditorApp"));
 	InitAssetEditor(Mode, InitToolkitHost, BlueprintEditorAppName, DummyLayout, bCreateDefaultStandaloneMenu, bCreateDefaultToolbar, Objects);
 	
-	CommonInitialization(InBlueprints);
+	CommonInitialization(InBlueprints, bShouldOpenInDefaultsMode);
 
 	InitalizeExtenders();
 
diff --git a/Engine/Source/Editor/Kismet/Public/BlueprintEditor.h b/Engine/Source/Editor/Kismet/Public/BlueprintEditor.h
--- a/Engine/Source/Editor/Kismet/Public/BlueprintEditor.h
+++ b/Engine/Source/Editor/Kismet/Public/BlueprintEditor.h
@@ -441,7 +441,7 @@ public:
 	void EnsureBlueprintIsUpToDate(UBlueprint* BlueprintObj);
 
 	// Should be called when initializing any editor built off this foundation
-	void CommonInitialization(const TArray<UBlueprint*>& InitBlueprints);
+	void CommonInitialization(const TArray<UBlueprint*>& InitBlueprints, bool bShouldOpenInDefaultsMode);
 
 	// Should be called when initializing an editor that has a blueprint, after layout (tab spawning) is done
 	void PostLayoutBlueprintEditorInitialization();
