Support optional CAC node expansion within the SCS editor's tree view as an alternative UX for editing child actor templates/instances (experimental).

diff --git a/Engine/Source/Editor/Kismet/Private/BlueprintDetailsCustomization.cpp b/Engine/Source/Editor/Kismet/Private/BlueprintDetailsCustomization.cpp
--- a/Engine/Source/Editor/Kismet/Private/BlueprintDetailsCustomization.cpp
+++ b/Engine/Source/Editor/Kismet/Private/BlueprintDetailsCustomization.cpp
@@ -26,6 +26,7 @@
 #include "Components/ChildActorComponent.h"
 #include "Components/TimelineComponent.h"
 #include "Kismet2/ComponentEditorUtils.h"
+#include "Kismet2/ChildActorComponentEditorUtils.h"
 #include "Kismet2/KismetEditorUtilities.h"
 #include "EdGraphSchema_K2.h"
 #include "K2Node_Event.h"
@@ -5985,7 +5986,7 @@ bool FBlueprintComponentDetails::OnVariableCategoryChangeEnabled() const
 {
 	check(CachedNodePtr.IsValid());
 	
-	return !CachedNodePtr->IsInherited();
+	return !CachedNodePtr->IsInheritedComponent();
 }
 
 FText FBlueprintComponentDetails::OnGetVariableCategoryText() const
@@ -6191,7 +6192,7 @@ bool FBlueprintComponentDetails::CanChangeSocket() const
 
 	if (CachedNodePtr->GetSCSNode() != NULL)
 	{
-		return !CachedNodePtr->IsInherited();
+		return !CachedNodePtr->IsInheritedComponent();
 	}
 	return true;
 }
@@ -6485,7 +6486,7 @@ void FChildActorComponentDetails::CustomizeDetails(IDetailLayoutBuilder& DetailB
 			{
 				if (UChildActorComponent* CAC = Cast<UChildActorComponent>(ObjectBeingCustomized.Get()))
 				{
-					if (CAC->ChildActorTemplate == nullptr)
+					if (CAC->ChildActorTemplate == nullptr || FChildActorComponentEditorUtils::ShouldShowChildActorNodeInTreeView(CAC))
 					{
 						return EVisibility::Hidden;
 					}
diff --git a/Engine/Source/Editor/Kismet/Private/BlueprintEditor.cpp b/Engine/Source/Editor/Kismet/Private/BlueprintEditor.cpp
--- a/Engine/Source/Editor/Kismet/Private/BlueprintEditor.cpp
+++ b/Engine/Source/Editor/Kismet/Private/BlueprintEditor.cpp
@@ -966,9 +966,9 @@ void FBlueprintEditor::OnSelectionUpdated(const TArray<FSCSEditorTreeNodePtrType
 		{
 			if (NodePtr.IsValid())
 			{
-				if (NodePtr->GetNodeType() == FSCSEditorTreeNode::RootActorNode)
+				if (NodePtr->IsActorNode())
 				{
-					AActor* DefaultActor = GetSCSEditorActorContext();
+					AActor* DefaultActor = NodePtr->GetEditableObjectForBlueprint<AActor>(GetBlueprintObj());
 					InspectorObjects.Add(DefaultActor);
 					
 					FString Title; 
diff --git a/Engine/Source/Editor/Kismet/Private/SCSDiff.cpp b/Engine/Source/Editor/Kismet/Private/SCSDiff.cpp
--- a/Engine/Source/Editor/Kismet/Private/SCSDiff.cpp
+++ b/Engine/Source/Editor/Kismet/Private/SCSDiff.cpp
@@ -98,10 +98,10 @@ void FSCSDiff::OnSCSEditorUpdateSelectionFromNodes(const TArray<FSCSEditorTreeNo
 	for (auto NodeIt = SelectedNodes.CreateConstIterator(); NodeIt; ++NodeIt)
 	{
 		auto NodePtr = *NodeIt;
-		if(NodePtr.IsValid() && NodePtr->CanEditDefaults())
+		if(NodePtr.IsValid() && NodePtr->CanEdit())
 		{
 			InspectorTitle = FText::FromString(NodePtr->GetDisplayString());
-			InspectorObjects.Add(NodePtr->GetOrCreateEditableComponentTemplate(Blueprint));
+			InspectorObjects.Add(NodePtr->GetEditableObjectForBlueprint<UObject>(Blueprint));
 		}
 	}
 
diff --git a/Engine/Source/Editor/Kismet/Private/SCSEditorViewportClient.cpp b/Engine/Source/Editor/Kismet/Private/SCSEditorViewportClient.cpp
--- a/Engine/Source/Editor/Kismet/Private/SCSEditorViewportClient.cpp
+++ b/Engine/Source/Editor/Kismet/Private/SCSEditorViewportClient.cpp
@@ -607,9 +607,9 @@ FWidget::EWidgetMode FSCSEditorViewportClient::GetWidgetMode() const
 						FSCSEditorTreeNodePtrType CurrentNodePtr = SelectedNodes[CurrentNodeIndex];
 						if ((CurrentNodePtr.IsValid() &&
 							((!RootNodes.Contains(CurrentNodePtr) && !CurrentNodePtr->IsRootComponent()) ||
-							(Cast<UInstancedStaticMeshComponent>(CurrentNodePtr->GetComponentTemplate()) && // show widget if we are editing individual instances even if it is the root component
+							(CurrentNodePtr->GetObject<UInstancedStaticMeshComponent>() && // show widget if we are editing individual instances even if it is the root component
 								CastChecked<UInstancedStaticMeshComponent>(CurrentNodePtr->FindComponentInstanceInActor(GetPreviewActor()))->SelectedInstances.Contains(true))) &&
-							CurrentNodePtr->CanEditDefaults() &&
+							CurrentNodePtr->CanEdit() &&
 							CurrentNodePtr->FindComponentInstanceInActor(PreviewActor)))
 						{
 							// a non-NULL, non-root item is selected, draw the widget
diff --git a/Engine/Source/Editor/Kismet/Private/SSCSEditor.cpp b/Engine/Source/Editor/Kismet/Private/SSCSEditor.cpp
--- a/Engine/Source/Editor/Kismet/Private/SSCSEditor.cpp
+++ b/Engine/Source/Editor/Kismet/Private/SSCSEditor.cpp
@@ -22,6 +22,7 @@
 #include "ThumbnailRendering/ThumbnailManager.h"
 #include "Components/ChildActorComponent.h"
 #include "Kismet2/ComponentEditorUtils.h"
+#include "Kismet2/ChildActorComponentEditorUtils.h"
 #include "Engine/Selection.h"
 #include "UnrealEdGlobals.h"
 #include "Kismet2/KismetEditorUtilities.h"
@@ -84,6 +85,8 @@ static const FName SCS_ColumnName_ComponentClass( "ComponentClass" );
 static const FName SCS_ColumnName_Asset( "Asset" );
 static const FName SCS_ColumnName_Mobility( "Mobility" );
 
+static const FName SCS_ContextMenuName( "Kismet.SCSEditorContextMenu" );
+
 //////////////////////////////////////////////////////////////////////////
 // SSCSEditorDragDropTree
 void SSCSEditorDragDropTree::Construct( const FArguments& InArgs )
@@ -126,7 +129,7 @@ FReply SSCSEditorDragDropTree::OnDragOver( const FGeometry& MyGeometry, const FD
 			{
 				if (Operation->IsOfType<FAssetDragDropOp>())
 				{
-					const auto& AssetDragDropOp = StaticCastSharedPtr<FAssetDragDropOp>(Operation);
+					const TSharedPtr<FAssetDragDropOp> AssetDragDropOp = StaticCastSharedPtr<FAssetDragDropOp>(Operation);
 
 					for (const FAssetData& AssetData : AssetDragDropOp->GetAssets())
 					{
@@ -328,7 +331,13 @@ void FSCSRowDragDropOp::HoverTargetChanged()
 
 	if (!bHoverHandled)
 	{
-		if (FProperty* VariableProperty = GetVariableProperty())
+		if (FChildActorComponentEditorUtils::ContainsChildActorSubtreeNode(SourceNodes))
+		{
+			// @todo - Add support for drag/drop of child actor template components to create variable get nodes in a local Blueprint event/function graph
+			FText Message = LOCTEXT("ChildActorDragDropAddVariableNode_Unsupported", "This operation is not currently supported for one or more of the selected components.");
+			SetSimpleFeedbackMessage(ErrorSymbol, IconTint, Message);
+		}
+		else if (FProperty* VariableProperty = GetVariableProperty())
 		{
 			const FSlateBrush* PrimarySymbol;
 			const FSlateBrush* SecondarySymbol;
@@ -400,8 +409,7 @@ FReply FSCSRowDragDropOp::DroppedOnPanel(const TSharedRef< class SWidget >& Pane
 // FSCSEditorTreeNode
 
 FSCSEditorTreeNode::FSCSEditorTreeNode(FSCSEditorTreeNode::ENodeType InNodeType)
-	: ComponentTemplatePtr(nullptr)
-	, NodeType(InNodeType)
+	: NodeType(InNodeType)
 	, FilterFlags((uint8)EFilteredState::Unknown)
 {
 }
@@ -440,53 +448,11 @@ class USCS_Node* FSCSEditorTreeNode::GetSCSNode() const
 	return nullptr;
 }
 
-UActorComponent* FSCSEditorTreeNode::GetOrCreateEditableComponentTemplate(UBlueprint* ActualEditedBlueprint) const
-{
-	return nullptr;
-}
-
-UBlueprint* FSCSEditorTreeNode::GetBlueprint() const
-{
-	USCS_Node* SCS_Node = GetSCSNode();
-	UActorComponent* ComponentTemplate = GetComponentTemplate();
-
-	if(SCS_Node)
-	{
-		USimpleConstructionScript* SCS = SCS_Node->GetSCS();
-		if(SCS)
-		{
-			return SCS->GetBlueprint();
-		}
-	}
-	else if(ComponentTemplate)
-	{
-		AActor* CDO = ComponentTemplate->GetOwner();
-		if(CDO)
-		{
-			check(CDO->GetClass());
-
-			return Cast<UBlueprint>(CDO->GetClass()->ClassGeneratedBy);
-		}
-	}
-
-	return NULL;
-}
-
 FSCSEditorTreeNode::ENodeType FSCSEditorTreeNode::GetNodeType() const
 {
 	return NodeType;
 }
 
-UActorComponent* FSCSEditorTreeNode::GetComponentTemplate(bool bEvenIfPendingKill) const
-{
-	return ComponentTemplatePtr.Get(bEvenIfPendingKill);
-}
-
-void FSCSEditorTreeNode::SetComponentTemplate(UActorComponent* Component)
-{
-	ComponentTemplatePtr = Component;
-}
-
 bool FSCSEditorTreeNode::IsAttachedTo(FSCSEditorTreeNodePtrType InNodePtr) const
 { 
 	FSCSEditorTreeNodePtrType TestParentPtr = ParentNodePtr;
@@ -629,15 +595,21 @@ FSCSEditorTreeNodePtrType FSCSEditorTreeNode::FindClosestParent(TArray<FSCSEdito
 
 void FSCSEditorTreeNode::AddChild(FSCSEditorTreeNodePtrType InChildNodePtr)
 {
-	USCS_Node* SCS_Node = GetSCSNode();
-	UActorComponent* ComponentTemplate = GetComponentTemplate();
-
 	// Ensure the node is not already parented elsewhere
 	if(InChildNodePtr->GetParent().IsValid())
 	{
 		InChildNodePtr->GetParent()->RemoveChild(InChildNodePtr);
 	}
 
+	// If this is an actor node, start a new subtree
+	if (IsActorNode())
+	{
+		ActorRootNodePtr = StaticCastSharedRef<FSCSEditorTreeNodeActorBase>(AsShared());
+	}
+
+	// Link the child node to the actor root node for this subtree
+	InChildNodePtr->ActorRootNodePtr = ActorRootNodePtr;
+
 	// Add the given node as a child and link its parent
 	Children.AddUnique(InChildNodePtr);
 	InChildNodePtr->ParentNodePtr = AsShared();
@@ -652,67 +624,73 @@ void FSCSEditorTreeNode::AddChild(FSCSEditorTreeNodePtrType InChildNodePtr)
 		}
 	}
 
-	// Add a child node to the SCS tree node if not already present
-	USCS_Node* SCS_ChildNode = InChildNodePtr->GetSCSNode();
-	if(SCS_ChildNode != NULL)
+	if (InChildNodePtr->IsComponentNode())
 	{
-		// Get the SCS instance that owns the child node
-		USimpleConstructionScript* SCS = SCS_ChildNode->GetSCS();
-		if(SCS != NULL)
+		USCS_Node* SCS_Node = GetSCSNode();
+		const UActorComponent* ComponentTemplate = GetObject<UActorComponent>();
+
+		// Add a child node to the SCS tree node if not already present
+		USCS_Node* SCS_ChildNode = InChildNodePtr->GetSCSNode();
+		if (SCS_ChildNode != NULL)
 		{
-			// If the parent is also a valid SCS node
-			if(SCS_Node != NULL)
+			// Get the SCS instance that owns the child node
+			USimpleConstructionScript* SCS = SCS_ChildNode->GetSCS();
+			if (SCS != NULL)
 			{
-				// If the parent and child are both owned by the same SCS instance
-				if(SCS_Node->GetSCS() == SCS)
+				// If the parent is also a valid SCS node
+				if (SCS_Node != NULL)
 				{
-					// Add the child into the parent's list of children
-					if(!SCS_Node->GetChildNodes().Contains(SCS_ChildNode))
+					// If the parent and child are both owned by the same SCS instance
+					if (SCS_Node->GetSCS() == SCS)
 					{
-						SCS_Node->AddChildNode(SCS_ChildNode);
+						// Add the child into the parent's list of children
+						if (!SCS_Node->GetChildNodes().Contains(SCS_ChildNode))
+						{
+							SCS_Node->AddChildNode(SCS_ChildNode);
+						}
+					}
+					else
+					{
+						// Adds the child to the SCS root set if not already present
+						SCS->AddNode(SCS_ChildNode);
+
+						// Set parameters to parent this node to the "inherited" SCS node
+						SCS_ChildNode->SetParent(SCS_Node);
 					}
 				}
-				else
+				else if (ComponentTemplate != NULL)
 				{
 					// Adds the child to the SCS root set if not already present
 					SCS->AddNode(SCS_ChildNode);
 
-					// Set parameters to parent this node to the "inherited" SCS node
-					SCS_ChildNode->SetParent(SCS_Node);
+					// Set parameters to parent this node to the native component template
+					SCS_ChildNode->SetParent(Cast<const USceneComponent>(ComponentTemplate));
+				}
+				else
+				{
+					// Adds the child to the SCS root set if not already present
+					SCS->AddNode(SCS_ChildNode);
 				}
-			}
-			else if(ComponentTemplate != NULL)
-			{
-				// Adds the child to the SCS root set if not already present
-				SCS->AddNode(SCS_ChildNode);
-
-				// Set parameters to parent this node to the native component template
-				SCS_ChildNode->SetParent(Cast<USceneComponent>(ComponentTemplate));
-			}
-			else
-			{
-				// Adds the child to the SCS root set if not already present
-				SCS->AddNode(SCS_ChildNode);
 			}
 		}
-	}
-	else if (IsInstanced())
-	{
-		USceneComponent* ChildInstance = Cast<USceneComponent>(InChildNodePtr->GetComponentTemplate());
-		if (ensure(ChildInstance != nullptr))
+		else if (IsInstancedComponent())
 		{
-			USceneComponent* ParentInstance = Cast<USceneComponent>(GetComponentTemplate());
-			if (ensure(ParentInstance != nullptr))
+			USceneComponent* ChildInstance = Cast<USceneComponent>(InChildNodePtr->GetComponentTemplate());
+			if (ensure(ChildInstance != nullptr))
 			{
-				// Handle attachment at the instance level
-				if (ChildInstance->GetAttachParent() != ParentInstance)
+				USceneComponent* ParentInstance = Cast<USceneComponent>(GetComponentTemplate());
+				if (ensure(ParentInstance != nullptr))
 				{
-					AActor* Owner = ParentInstance->GetOwner();
-					if (Owner->GetRootComponent() == ChildInstance)
+					// Handle attachment at the instance level
+					if (ChildInstance->GetAttachParent() != ParentInstance)
 					{
-						Owner->SetRootComponent(ParentInstance);
+						AActor* Owner = ParentInstance->GetOwner();
+						if (Owner->GetRootComponent() == ChildInstance)
+						{
+							Owner->SetRootComponent(ParentInstance);
+						}
+						ChildInstance->AttachToComponent(ParentInstance, FAttachmentTransformRules::KeepWorldTransform);
 					}
-					ChildInstance->AttachToComponent(ParentInstance, FAttachmentTransformRules::KeepWorldTransform);
 				}
 			}
 		}
@@ -754,7 +732,7 @@ FSCSEditorTreeNodePtrType FSCSEditorTreeNode::AddChildFromComponent(UActorCompon
 }
 
 // Tries to find a SCS node that was likely responsible for creating the specified instance component.  Note: This is not always possible to do!
-USCS_Node* FSCSEditorTreeNode::FindSCSNodeForInstance(UActorComponent* InstanceComponent, UClass* ClassToSearch)
+USCS_Node* FSCSEditorTreeNode::FindSCSNodeForInstance(const UActorComponent* InstanceComponent, UClass* ClassToSearch)
 { 
 	if ((ClassToSearch != nullptr) && InstanceComponent->IsCreatedByConstructionScript())
 	{
@@ -924,6 +902,16 @@ void FSCSEditorTreeNode::OnCompleteRename(const FText& InNewName)
 	OngoingCreateTransaction.Reset();
 }
 
+UObject* FSCSEditorTreeNode::GetOrCreateEditableObjectForBlueprint(UBlueprint* InBlueprint) const
+{
+	if (CanEdit())
+	{
+		return WeakObjectPtr.Get();
+	}
+
+	return nullptr;
+}
+
 //////////////////////////////////////////////////////////////////////////
 // FSCSEditorTreeNodeComponentBase
 
@@ -934,7 +922,7 @@ FName FSCSEditorTreeNodeComponentBase::GetVariableName() const
 	USCS_Node* SCS_Node = GetSCSNode();
 	UActorComponent* ComponentTemplate = GetComponentTemplate();
 
-	if (IsInstanced() && (SCS_Node == nullptr) && (ComponentTemplate != nullptr))
+	if (IsInstancedComponent() && (SCS_Node == nullptr) && (ComponentTemplate != nullptr))
 	{
 		if (ComponentTemplate->GetOwner())
 		{
@@ -980,7 +968,7 @@ FString FSCSEditorTreeNodeComponentBase::GetDisplayString() const
 	else
 	{
 		FString UnnamedString = LOCTEXT("UnnamedToolTip", "Unnamed").ToString();
-		FString NativeString = IsNative() ? LOCTEXT("NativeToolTip", "Native ").ToString() : TEXT("");
+		FString NativeString = IsNativeComponent() ? LOCTEXT("NativeToolTip", "Native ").ToString() : TEXT("");
 
 		if (ComponentTemplate != NULL)
 		{
@@ -993,6 +981,52 @@ FString FSCSEditorTreeNodeComponentBase::GetDisplayString() const
 	}
 }
 
+bool FSCSEditorTreeNodeComponentBase::CanReparent() const
+{
+	FSCSEditorActorNodePtrType ActorTreeRootNode = GetActorRootNode();
+	if (FChildActorComponentEditorUtils::IsChildActorSubtreeNode(AsShared()))
+	{
+		// Cannot reparent nodes within a child actor node subtree.
+		return false;
+	}
+
+	return !IsInheritedComponent() && !IsDefaultSceneRoot() && IsSceneComponent();
+}
+
+UBlueprint* FSCSEditorTreeNodeComponentBase::GetBlueprint() const
+{
+	if (const USCS_Node* SCS_Node = GetSCSNode())
+	{
+		if (const USimpleConstructionScript* SCS = SCS_Node->GetSCS())
+		{
+			return SCS->GetBlueprint();
+		}
+	}
+	else if (const UActorComponent* ActorComponent = GetObject<UActorComponent>())
+	{
+		if (const AActor* Actor = ActorComponent->GetOwner())
+		{
+			return UBlueprint::GetBlueprintFromClass(Actor->GetClass());
+		}
+	}
+	
+	return nullptr;
+}
+
+FSCSEditorChildActorNodePtrType FSCSEditorTreeNodeComponentBase::GetChildActorNode()
+{
+	if (!ChildActorNodePtr.IsValid() || ChildActorNodePtr->GetChildActorComponent() != GetObject<UActorComponent>())
+	{
+		if (const UChildActorComponent* ChildActorComponent = GetObject<UChildActorComponent>())
+		{
+			AActor* ChildActor = IsInstanced() ? ChildActorComponent->GetChildActor() : ChildActorComponent->GetChildActorTemplate();
+			ChildActorNodePtr = MakeShareable(new FSCSEditorTreeNodeChildActor(this->AsShared(), ChildActor));
+		}
+	}
+
+	return ChildActorNodePtr;
+}
+
 //////////////////////////////////////////////////////////////////////////
 // FSCSEditorTreeNodeInstancedInheritedComponent
 
@@ -1002,10 +1036,10 @@ FSCSEditorTreeNodeInstancedInheritedComponent::FSCSEditorTreeNodeInstancedInheri
 
 	InstancedComponentOwnerPtr = Owner;
 
-	SetComponentTemplate(ComponentInstance);
+	SetObject(ComponentInstance);
 }
 
-bool FSCSEditorTreeNodeInstancedInheritedComponent::IsNative() const
+bool FSCSEditorTreeNodeInstancedInheritedComponent::IsNativeComponent() const
 {
 	if (UActorComponent* Template = GetComponentTemplate())
 	{
@@ -1032,7 +1066,7 @@ bool FSCSEditorTreeNodeInstancedInheritedComponent::IsRootComponent() const
 	return false;
 }
 
-bool FSCSEditorTreeNodeInstancedInheritedComponent::IsInheritedSCS() const
+bool FSCSEditorTreeNodeInstancedInheritedComponent::IsInheritedSCSNode() const
 {
 	return false;
 }
@@ -1042,7 +1076,7 @@ bool FSCSEditorTreeNodeInstancedInheritedComponent::IsDefaultSceneRoot() const
 	return false;
 }
 
-bool FSCSEditorTreeNodeInstancedInheritedComponent::CanEditDefaults() const
+bool FSCSEditorTreeNodeInstancedInheritedComponent::CanEdit() const
 {
 	UActorComponent* Template = GetComponentTemplate();
 	return (Template ? Template->IsEditableWhenInherited() : false);
@@ -1059,16 +1093,6 @@ FText FSCSEditorTreeNodeInstancedInheritedComponent::GetDisplayName() const
 	return FText::GetEmpty();
 }
 
-UActorComponent* FSCSEditorTreeNodeInstancedInheritedComponent::GetOrCreateEditableComponentTemplate(UBlueprint* ActualEditedBlueprint) const
-{
-	if (CanEditDefaults())
-	{
-		return GetComponentTemplate();
-	}
-
-	return nullptr;
-}
-
 //////////////////////////////////////////////////////////////////////////
 // FSCSEditorTreeNodeInstanceAddedComponent
 
@@ -1077,7 +1101,7 @@ FSCSEditorTreeNodeInstanceAddedComponent::FSCSEditorTreeNodeInstanceAddedCompone
 	check(InComponentTemplate);
 	InstancedComponentName = InComponentTemplate->GetFName();
 	InstancedComponentOwnerPtr = Owner;
-	SetComponentTemplate(InComponentTemplate);
+	SetObject(InComponentTemplate);
 }
 
 bool FSCSEditorTreeNodeInstanceAddedComponent::IsRootComponent() const
@@ -1118,11 +1142,6 @@ FText FSCSEditorTreeNodeInstanceAddedComponent::GetDisplayName() const
 	return FText::FromName(InstancedComponentName);
 }
 
-UActorComponent* FSCSEditorTreeNodeInstanceAddedComponent::GetOrCreateEditableComponentTemplate(UBlueprint* ActualEditedBlueprint) const
-{
-	return GetComponentTemplate();
-}
-
 void FSCSEditorTreeNodeInstanceAddedComponent::RemoveMeAsChild()
 {
 	USceneComponent* ChildInstance = Cast<USceneComponent>(GetComponentTemplate());
@@ -1174,7 +1193,7 @@ FSCSEditorTreeNodeComponent::FSCSEditorTreeNodeComponent(USCS_Node* InSCSNode, b
 	: bIsInheritedSCS(bInIsInheritedSCS)
 	, SCSNodePtr(InSCSNode)
 {
-	SetComponentTemplate(( InSCSNode != nullptr ) ? InSCSNode->ComponentTemplate : nullptr);
+	SetObject(( InSCSNode != nullptr ) ? InSCSNode->ComponentTemplate : nullptr);
 }
 
 FSCSEditorTreeNodeComponent::FSCSEditorTreeNodeComponent(UActorComponent* InComponentTemplate)
@@ -1183,7 +1202,7 @@ FSCSEditorTreeNodeComponent::FSCSEditorTreeNodeComponent(UActorComponent* InComp
 {
 	check(InComponentTemplate != nullptr);
 
-	SetComponentTemplate(InComponentTemplate);
+	SetObject(InComponentTemplate);
 	AActor* Owner = InComponentTemplate->GetOwner();
 	if (Owner != nullptr)
 	{
@@ -1191,7 +1210,7 @@ FSCSEditorTreeNodeComponent::FSCSEditorTreeNodeComponent(UActorComponent* InComp
 	}
 }
 
-bool FSCSEditorTreeNodeComponent::IsNative() const
+bool FSCSEditorTreeNodeComponent::IsNativeComponent() const
 {
 	return GetSCSNode() == NULL && GetComponentTemplate() != NULL;
 }
@@ -1224,7 +1243,7 @@ bool FSCSEditorTreeNodeComponent::IsRootComponent() const
 	return bIsRoot;
 }
 
-bool FSCSEditorTreeNodeComponent::IsInheritedSCS() const
+bool FSCSEditorTreeNodeComponent::IsInheritedSCSNode() const
 {
 	return bIsInheritedSCS;
 }
@@ -1243,11 +1262,11 @@ bool FSCSEditorTreeNodeComponent::IsDefaultSceneRoot() const
 	return false;
 }
 
-bool FSCSEditorTreeNodeComponent::CanEditDefaults() const
+bool FSCSEditorTreeNodeComponent::CanEdit() const
 {
 	bool bCanEdit = false;
 
-	if (!IsNative())
+	if (!IsNativeComponent())
 	{
 		USCS_Node* SCS_Node = GetSCSNode();
 		bCanEdit = (SCS_Node != nullptr);
@@ -1275,26 +1294,14 @@ class USCS_Node* FSCSEditorTreeNodeComponent::GetSCSNode() const
 	return SCSNodePtr.Get();
 }
 
-UActorComponent* FSCSEditorTreeNodeComponent::GetOrCreateEditableComponentTemplate(UBlueprint* ActualEditedBlueprint) const
+UObject* FSCSEditorTreeNodeComponent::GetOrCreateEditableObjectForBlueprint(UBlueprint* InBlueprint) const
 {
-	if (CanEditDefaults())
+	if (CanEdit() && !IsNativeComponent() && IsInheritedSCSNode())
 	{
-		if (!IsNative() && IsInheritedSCS())
-		{
-			if (ActualEditedBlueprint != nullptr)
-			{
-				return INTERNAL_GetOverridenComponentTemplate(ActualEditedBlueprint, true);
-			}
-			else
-			{
-				return nullptr;
-			}
-		}
-
-		return GetComponentTemplate();
+		return INTERNAL_GetOverridenComponentTemplate(InBlueprint);
 	}
 
-	return nullptr;
+	return FSCSEditorTreeNode::GetOrCreateEditableObjectForBlueprint(InBlueprint);
 }
 
 UActorComponent* FSCSEditorTreeNode::FindComponentInstanceInActor(const AActor* InActor) const
@@ -1359,19 +1366,20 @@ void FSCSEditorTreeNodeComponent::RemoveMeAsChild()
 	}
 }
 
-UActorComponent* FSCSEditorTreeNodeComponent::INTERNAL_GetOverridenComponentTemplate(UBlueprint* Blueprint, bool bCreateIfNecessary) const
+UActorComponent* FSCSEditorTreeNodeComponent::INTERNAL_GetOverridenComponentTemplate(UBlueprint* Blueprint) const
 {
-	UActorComponent* OverriddenComponent = NULL;
+	UActorComponent* OverriddenComponent = nullptr;
 
 	FComponentKey Key(GetSCSNode());
 
-	const bool BlueprintCanOverrideComponentFromKey = Key.IsValid()
+	const bool bBlueprintCanOverrideComponentFromKey = Key.IsValid()
 		&& Blueprint
 		&& Blueprint->ParentClass
 		&& Blueprint->ParentClass->IsChildOf(Key.GetComponentOwner());
 
-	if (BlueprintCanOverrideComponentFromKey)
+	if (bBlueprintCanOverrideComponentFromKey)
 	{
+		const bool bCreateIfNecessary = true;
 		UInheritableComponentHandler* InheritableComponentHandler = Blueprint->GetInheritableComponentHandler(bCreateIfNecessary);
 		if (InheritableComponentHandler)
 		{
@@ -1386,13 +1394,14 @@ UActorComponent* FSCSEditorTreeNodeComponent::INTERNAL_GetOverridenComponentTemp
 }
 
 //////////////////////////////////////////////////////////////////////////
-// FSCSEditorTreeNodeRootActor
-FSCSEditorTreeNodePtrType FSCSEditorTreeNodeRootActor::GetSceneRootNode() const
+// FSCSEditorTreeNodeActorBase
+
+FSCSEditorTreeNodePtrType FSCSEditorTreeNodeActorBase::GetSceneRootNode() const
 {
 	return SceneRootNodePtr;
 }
 
-void FSCSEditorTreeNodeRootActor::SetSceneRootNode(FSCSEditorTreeNodePtrType NewSceneRootNode)
+void FSCSEditorTreeNodeActorBase::SetSceneRootNode(FSCSEditorTreeNodePtrType NewSceneRootNode)
 {
 	if (SceneRootNodePtr.IsValid())
 	{
@@ -1407,84 +1416,178 @@ void FSCSEditorTreeNodeRootActor::SetSceneRootNode(FSCSEditorTreeNodePtrType New
 	}
 }
 
-const TArray<FSCSEditorTreeNodePtrType>& FSCSEditorTreeNodeRootActor::GetComponentNodes() const
+const TArray<FSCSEditorTreeNodePtrType>& FSCSEditorTreeNodeActorBase::GetComponentNodes() const
 {
 	return ComponentNodes;
 }
 
-void FSCSEditorTreeNodeRootActor::AddChild(FSCSEditorTreeNodePtrType InChildNodePtr)
+FName FSCSEditorTreeNodeActorBase::GetNodeID() const
+{
+	if (const AActor* Actor = GetObject<AActor>())
+	{
+		return Actor->GetFName();
+	}
+	return NAME_None;
+}
+
+bool FSCSEditorTreeNodeActorBase::IsInstanced() const
+{
+	if (const AActor* Actor = GetObject<AActor>())
+	{
+		return !Actor->IsTemplate();
+	}
+
+	return false;
+}
+
+void FSCSEditorTreeNodeActorBase::AddChild(FSCSEditorTreeNodePtrType InChildNodePtr)
 {
-	if (InChildNodePtr->GetNodeType() == FSCSEditorTreeNode::ComponentNode)
+	if (InChildNodePtr.IsValid() && InChildNodePtr->IsComponentNode())
 	{
 		ComponentNodes.Add(InChildNodePtr);
-		USceneComponent* SceneComponent = Cast<USceneComponent>(InChildNodePtr->GetComponentTemplate());
-		if (!SceneRootNodePtr.IsValid() && SceneComponent != nullptr)
+		if (!SceneRootNodePtr.IsValid() && InChildNodePtr->IsSceneComponent())
 		{
 			SetSceneRootNode(InChildNodePtr);
 		}
+	}
+
+	Super::AddChild(InChildNodePtr);
+}
+
+void FSCSEditorTreeNodeActorBase::RemoveChild(FSCSEditorTreeNodePtrType InChildNodePtr)
+{
+	if (InChildNodePtr.IsValid() && InChildNodePtr->IsComponentNode())
+	{
+		ComponentNodes.Remove(InChildNodePtr);
+		if (InChildNodePtr == SceneRootNodePtr)
+		{
+			SceneRootNodePtr.Reset();
+		}
+	}
+
+	Super::RemoveChild(InChildNodePtr);
+}
+
+UBlueprint* FSCSEditorTreeNodeActorBase::GetBlueprint() const
+{
+	if (const AActor* Actor = GetObject<AActor>())
+	{
+		return UBlueprint::GetBlueprintFromClass(Actor->GetClass());
+	}
+
+	return nullptr;
+}
+
+//////////////////////////////////////////////////////////////////////////
+// FSCSEditorTreeNodeRootActor
+
+void FSCSEditorTreeNodeRootActor::OnCompleteRename(const FText& InNewName)
+{
+	if (AActor* Actor = GetMutableObject<AActor>())
+	{
+		if (Actor->IsActorLabelEditable() && !InNewName.ToString().Equals(Actor->GetActorLabel(), ESearchCase::CaseSensitive))
+		{
+			const FScopedTransaction Transaction(LOCTEXT("SCSEditorRenameActorTransaction", "Rename Actor"));
+			FActorLabelUtilities::RenameExistingActor(Actor, InNewName.ToString());
+		}
+	}
+
+	// Not expected to reach here with an ongoing create transaction, but if it does, end it.
+	OngoingCreateTransaction.Reset();
+}
+
+void FSCSEditorTreeNodeRootActor::AddChild(FSCSEditorTreeNodePtrType InChildNodePtr)
+{
+	if (InChildNodePtr.IsValid() && InChildNodePtr->IsComponentNode())
+	{
+		const bool bIsSceneComponentNode = InChildNodePtr->IsSceneComponent();
 
 		// Make sure separators are shown
-		if (SceneComponent != nullptr && !SceneComponentSeparatorNodePtr.IsValid())
+		if (bIsSceneComponentNode && !SceneComponentSeparatorNodePtr.IsValid())
 		{
 			SceneComponentSeparatorNodePtr = MakeShareable(new FSCSEditorTreeNodeSeparator());
-			Super::AddChild(SceneComponentSeparatorNodePtr);
+			FSCSEditorTreeNodeActorBase::AddChild(SceneComponentSeparatorNodePtr);
 		}
-		else if (SceneComponent == nullptr && !NonSceneComponentSeparatorNodePtr.IsValid())
+		else if (!bIsSceneComponentNode && !NonSceneComponentSeparatorNodePtr.IsValid())
 		{
 			NonSceneComponentSeparatorNodePtr = MakeShareable(new FSCSEditorTreeNodeSeparator());
-			Super::AddChild(NonSceneComponentSeparatorNodePtr);
+			FSCSEditorTreeNodeActorBase::AddChild(NonSceneComponentSeparatorNodePtr);
 		}
 	}
-	
-	Super::AddChild(InChildNodePtr);
+
+	FSCSEditorTreeNodeActorBase::AddChild(InChildNodePtr);
 }
 
 void FSCSEditorTreeNodeRootActor::RemoveChild(FSCSEditorTreeNodePtrType InChildNodePtr)
 {
-	Super::RemoveChild(InChildNodePtr);
-
-	int32 indexOfFirstSceneComponent = ComponentNodes.IndexOfByPredicate([](const FSCSEditorTreeNodePtrType& NodePtr)
+	const TArray<FSCSEditorTreeNodePtrType>& ComponentNodePtrs = GetComponentNodes();
+	const int32 IndexOfFirstSceneComponent = ComponentNodePtrs.IndexOfByPredicate([](const FSCSEditorTreeNodePtrType& NodePtr)
 	{
-		return NodePtr->GetNodeType() == FSCSEditorTreeNode::ComponentNode && Cast<USceneComponent>(NodePtr->GetComponentTemplate());
+		return NodePtr->IsSceneComponent();
 	});
 
-	if (indexOfFirstSceneComponent == -1 && SceneComponentSeparatorNodePtr.IsValid())
+	if (IndexOfFirstSceneComponent == -1 && SceneComponentSeparatorNodePtr.IsValid())
 	{
-		Super::RemoveChild(SceneComponentSeparatorNodePtr);
-		SceneComponentSeparatorNodePtr = nullptr;
+		FSCSEditorTreeNodeActorBase::RemoveChild(SceneComponentSeparatorNodePtr);
+		SceneComponentSeparatorNodePtr.Reset();
 	}
-	
-	int32 indexOffFirstNonSceneComponent = ComponentNodes.IndexOfByPredicate([](const FSCSEditorTreeNodePtrType& NodePtr)
+
+	const int32 IndexOfFirstNonSceneComponent = ComponentNodePtrs.IndexOfByPredicate([](const FSCSEditorTreeNodePtrType& NodePtr)
 	{
-		return NodePtr->GetNodeType() == FSCSEditorTreeNode::ComponentNode && !Cast<USceneComponent>(NodePtr->GetComponentTemplate());
+		return !NodePtr->IsSceneComponent();
 	});
 
-	if (indexOffFirstNonSceneComponent == -1 && NonSceneComponentSeparatorNodePtr.IsValid())
+	if (IndexOfFirstNonSceneComponent == -1 && NonSceneComponentSeparatorNodePtr.IsValid())
 	{
-		Super::RemoveChild(NonSceneComponentSeparatorNodePtr);
-		NonSceneComponentSeparatorNodePtr = nullptr;
+		FSCSEditorTreeNodeActorBase::RemoveChild(NonSceneComponentSeparatorNodePtr);
+		NonSceneComponentSeparatorNodePtr.Reset();
 	}
+
+	FSCSEditorTreeNodeActorBase::RemoveChild(InChildNodePtr);
 }
 
-FName FSCSEditorTreeNodeRootActor::GetNodeID() const
+//////////////////////////////////////////////////////////////////////////
+// FSCSEditorTreeNodeChildActor
+
+bool FSCSEditorTreeNodeChildActor::IsFlaggedForFiltration() const
 {
-	if (Actor)
+	// Not filtered out if any children are flagged as such
+	for (const FSCSEditorTreeNodePtrType& Child : GetChildren())
 	{
-		return Actor->GetFName();
+		if (!Child->IsFlaggedForFiltration())
+		{
+			return false;
+		}
 	}
-	return NAME_None;
+
+	if (OwnerNode.IsValid())
+	{
+		// Mirror the owning node's filtered state
+		return OwnerNode->IsFlaggedForFiltration();
+	}
+
+	return false;
 }
 
-void FSCSEditorTreeNodeRootActor::OnCompleteRename(const FText& InNewName)
+bool FSCSEditorTreeNodeChildActor::ShouldShowInTreeView() const
 {
-	if (Actor && Actor->IsActorLabelEditable() && !InNewName.ToString().Equals(Actor->GetActorLabel(), ESearchCase::CaseSensitive))
+	if (UChildActorComponent* ChildActorComponent = GetChildActorComponent())
 	{
-		const FScopedTransaction Transaction(LOCTEXT("SCSEditorRenameActorTransaction", "Rename Actor"));
-		FActorLabelUtilities::RenameExistingActor(Actor, InNewName.ToString());
+		const bool bHasValidActorContext = GetObject<AActor>() != nullptr;
+		return bHasValidActorContext && FChildActorComponentEditorUtils::ShouldShowChildActorNodeInTreeView(ChildActorComponent);
 	}
 
-	// Not expected to reach here with an ongoing create transaction, but if it does, end it.
-	OngoingCreateTransaction.Reset();
+	return false;
+}
+
+UChildActorComponent* FSCSEditorTreeNodeChildActor::GetChildActorComponent() const
+{
+	if (OwnerNode.IsValid())
+	{
+		return Cast<UChildActorComponent>(OwnerNode->GetComponentTemplate());
+	}
+
+	return nullptr;
 }
 
 //////////////////////////////////////////////////////////////////////////
@@ -1499,7 +1602,7 @@ void SSCS_RowWidget::Construct( const FArguments& InArgs, TSharedPtr<SSCSEditor>
 
 	bool bIsSeparator = InNodePtr->GetNodeType() == FSCSEditorTreeNode::SeparatorNode;
 	
-	auto Args = FSuperRowType::FArguments()
+	FSuperRowType::FArguments Args = FSuperRowType::FArguments()
 		.Style(bIsSeparator ?
 			&FEditorStyle::Get().GetWidgetStyle<FTableRowStyle>("TableView.NoHoverTableRow") :
 			&FEditorStyle::Get().GetWidgetStyle<FTableRowStyle>("SceneOutliner.TableViewRow")) //@todo create editor style for the SCS tree
@@ -1533,23 +1636,15 @@ BEGIN_SLATE_FUNCTION_BUILD_OPTIMIZATION
 TSharedRef<SWidget> SSCS_RowWidget::GenerateWidgetForColumn( const FName& ColumnName )
 {
 	FSCSEditorTreeNodePtrType NodePtr = GetNode();
-	
 	if(ColumnName == SCS_ColumnName_ComponentClass)
 	{
-		// Setup a default icon brush.
-		const FSlateBrush* ComponentIcon = FEditorStyle::GetBrush("SCS.NativeComponent");
-		if(NodePtr->GetComponentTemplate() != NULL)
-		{
-			ComponentIcon = FSlateIconFinder::FindIconBrushForClass( NodePtr->GetComponentTemplate()->GetClass(), TEXT("SCS.Component") );
-		}
-
 		InlineWidget =
 			SNew(SInlineEditableTextBlock)
 				.Text(this, &SSCS_RowWidget::GetNameLabel)
 				.OnVerifyTextChanged( this, &SSCS_RowWidget::OnNameTextVerifyChanged )
 				.OnTextCommitted( this, &SSCS_RowWidget::OnNameTextCommit )
 				.IsSelected( this, &SSCS_RowWidget::IsSelectedExclusively )
-				.IsReadOnly(!NodePtr->CanRename() || (SCSEditor.IsValid() && !SCSEditor.Pin()->IsEditingAllowed()));
+				.IsReadOnly(!NodePtr->CanRename() || (SCSEditor.IsValid() && !SCSEditor.Pin()->IsEditingAllowed()) || FChildActorComponentEditorUtils::IsChildActorSubtreeNode(NodePtr));
 
 		NodePtr->SetRenameRequestedDelegate(FSCSEditorTreeNode::FOnRenameRequested::CreateSP(InlineWidget.Get(), &SInlineEditableTextBlock::EnterEditingMode));
 		
@@ -1568,7 +1663,7 @@ TSharedRef<SWidget> SSCS_RowWidget::GenerateWidgetForColumn( const FName& Column
 					.VAlign(VAlign_Center)
 					[
 						SNew(SImage)
-						.Image(ComponentIcon)
+						.Image(GetIconBrush())
 						.ColorAndOpacity(this, &SSCS_RowWidget::GetColorTintForIcon)
 					]
 				+SHorizontalBox::Slot()
@@ -1667,7 +1762,7 @@ TSharedRef<SToolTip> SSCS_RowWidget::CreateToolTipWidget() const
 	// 
 	if (FSCSEditorTreeNode* TreeNode = GetNode().Get())
 	{
-		if (TreeNode->GetNodeType() == FSCSEditorTreeNode::ComponentNode)
+		if (TreeNode->IsComponentNode())
 		{
 			// Add the tooltip
 			if (UActorComponent* Template = TreeNode->GetComponentTemplate())
@@ -1689,7 +1784,7 @@ TSharedRef<SToolTip> SSCS_RowWidget::CreateToolTipWidget() const
 
 			// Add introduction point
 			AddToToolTipInfoBox(InfoBox, LOCTEXT("TooltipAddType", "Source"), SNullWidget::NullWidget, TAttribute<FText>::Create(TAttribute<FText>::FGetter::CreateSP(this, &SSCS_RowWidget::GetComponentAddSourceToolTipText)), false);
-			if (TreeNode->IsInherited())
+			if (TreeNode->IsInheritedComponent())
 			{
 				AddToToolTipInfoBox(InfoBox, LOCTEXT("TooltipIntroducedIn", "Introduced in"), SNullWidget::NullWidget, TAttribute<FText>::Create(TAttribute<FText>::FGetter::CreateSP(this, &SSCS_RowWidget::GetIntroducedInToolTipText)), false);
 			}
@@ -1817,9 +1912,9 @@ FText SSCS_RowWidget::GetComponentAddSourceToolTipText() const
 	
 	if (FSCSEditorTreeNode* TreeNode = TreeNodePtr.Get())
 	{
-		if (TreeNode->IsInherited())
+		if (TreeNode->IsInheritedComponent())
 		{
-			if (TreeNode->IsNative())
+			if (TreeNode->IsNativeComponent())
 			{
 				NodeType = LOCTEXT("InheritedNativeComponent", "Inherited (C++)");
 			}
@@ -1830,7 +1925,7 @@ FText SSCS_RowWidget::GetComponentAddSourceToolTipText() const
 		}
 		else
 		{
-			if (TreeNode->IsInstanced())
+			if (TreeNode->IsInstancedComponent())
 			{
 				NodeType = LOCTEXT("ThisInstanceAddedComponent", "This actor instance");
 			}
@@ -1850,14 +1945,14 @@ FText SSCS_RowWidget::GetIntroducedInToolTipText() const
 
 	if (FSCSEditorTreeNode* TreeNode = TreeNodePtr.Get())
 	{
-		if (TreeNode->IsInherited())
+		if (TreeNode->IsInheritedComponent())
 		{
 			if (UActorComponent* ComponentTemplate = TreeNode->GetComponentTemplate())
 			{
 				UClass* BestClass = nullptr;
 				AActor* OwningActor = ComponentTemplate->GetOwner();
 
-				if (TreeNode->IsNative() && (OwningActor != nullptr))
+				if (TreeNode->IsNativeComponent() && (OwningActor != nullptr))
 				{
 					for (UClass* TestClass = OwningActor->GetClass(); TestClass != AActor::StaticClass(); TestClass = TestClass->GetSuperClass())
 					{
@@ -1871,7 +1966,7 @@ FText SSCS_RowWidget::GetIntroducedInToolTipText() const
 						}
 					}
 				}
-				else if (!TreeNode->IsNative())
+				else if (!TreeNode->IsNativeComponent())
 				{
 					USCS_Node* SCSNode = TreeNode->GetSCSNode();
 
@@ -1907,11 +2002,11 @@ FText SSCS_RowWidget::GetIntroducedInToolTipText() const
 						IntroducedInTooltip = LOCTEXT("IntroducedInNativeError", "Unknown native source (via C++ code)");
 					}
 				}
-				else if (TreeNode->IsInstanced() && ComponentTemplate->CreationMethod == EComponentCreationMethod::Native && !ComponentTemplate->HasAnyFlags(RF_DefaultSubObject))
+				else if (TreeNode->IsInstancedComponent() && ComponentTemplate->CreationMethod == EComponentCreationMethod::Native && !ComponentTemplate->HasAnyFlags(RF_DefaultSubObject))
 				{
 					IntroducedInTooltip = FText::Format(LOCTEXT("IntroducedInCPPErrorFmt", "{0} (via C++ code)"), FBlueprintEditorUtils::GetFriendlyClassDisplayName(BestClass));
 				}
-				else if (TreeNode->IsInstanced() && ComponentTemplate->CreationMethod == EComponentCreationMethod::UserConstructionScript)
+				else if (TreeNode->IsInstancedComponent() && ComponentTemplate->CreationMethod == EComponentCreationMethod::UserConstructionScript)
 				{
 					IntroducedInTooltip = FText::Format(LOCTEXT("IntroducedInUCSErrorFmt", "{0} (via an Add Component call)"), FBlueprintEditorUtils::GetFriendlyClassDisplayName(BestClass));
 				}
@@ -1925,7 +2020,7 @@ FText SSCS_RowWidget::GetIntroducedInToolTipText() const
 				IntroducedInTooltip = LOCTEXT("IntroducedInNoTemplateError", "[no component template found]");
 			}
 		}
-		else if (TreeNode->IsInstanced())
+		else if (TreeNode->IsInstancedComponent())
 		{
 			IntroducedInTooltip = LOCTEXT("IntroducedInThisActorInstanceTooltip", "this actor instance");
 		}
@@ -1983,6 +2078,21 @@ EVisibility SSCS_RowWidget::GetAssetVisibility() const
 	}
 }
 
+const FSlateBrush* SSCS_RowWidget::GetIconBrush() const
+{
+	const FSlateBrush* ComponentIcon = FEditorStyle::GetBrush("SCS.NativeComponent");
+	FSCSEditorTreeNodePtrType NodePtr = GetNode();
+	if (NodePtr.IsValid())
+	{
+		if (UActorComponent* ComponentTemplate = NodePtr->GetComponentTemplate())
+		{
+			ComponentIcon = FSlateIconFinder::FindIconBrushForClass(ComponentTemplate->GetClass(), TEXT("SCS.Component"));
+		}
+	}
+
+	return ComponentIcon;
+}
+
 FSlateColor SSCS_RowWidget::GetColorTintForIcon() const
 {
 	return GetColorTintForIcon(GetNode());
@@ -1995,13 +2105,13 @@ FSlateColor SSCS_RowWidget::GetColorTintForIcon(FSCSEditorTreeNodePtrType InNode
 	const FLinearColor InheritedNativeComponentColor(0.7f, 0.9f, 0.7f);
 	const FLinearColor IntroducedHereColor(FLinearColor::White);
 	
-	if (InNode->IsInherited())
+	if (InNode->IsInheritedComponent())
 	{
-		if (InNode->IsNative())
+		if (InNode->IsNativeComponent())
 		{
 			return InheritedNativeComponentColor;
 		}
-		else if (InNode->IsInstanced())
+		else if (InNode->IsInstancedComponent())
 		{
 			return InstancedInheritedBlueprintComponentColor;
 		}
@@ -2081,7 +2191,7 @@ FReply SSCS_RowWidget::OnMouseButtonDown(const FGeometry& MyGeometry, const FPoi
 
 FReply SSCS_RowWidget::HandleOnDragDetected( const FGeometry& MyGeometry, const FPointerEvent& MouseEvent )
 {
-	auto SCSEditorPtr = SCSEditor.Pin();
+	TSharedPtr<SSCSEditor> SCSEditorPtr = SCSEditor.Pin();
 	if (MouseEvent.IsMouseButtonDown(EKeys::LeftMouseButton)
 		&& SCSEditorPtr.IsValid()
 		&& SCSEditorPtr->IsEditingAllowed()) //can only drag when editing
@@ -2093,7 +2203,7 @@ FReply SSCS_RowWidget::HandleOnDragDetected( const FGeometry& MyGeometry, const
 		}
 
 		TSharedPtr<FSCSEditorTreeNode> FirstNode = SelectedNodePtrs[0];
-		if (FirstNode->GetNodeType() == FSCSEditorTreeNode::ComponentNode)
+		if (FirstNode->IsComponentNode())
 		{
 			// Do not use the Blueprint from FirstNode, it may still be referencing the parent.
 			UBlueprint* Blueprint = GetBlueprint();
@@ -2128,18 +2238,22 @@ void SSCS_RowWidget::HandleOnDragEnter( const FDragDropEvent& DragDropEvent )
 		FText Message;
 		FSlateColor IconColor = FLinearColor::White;
 		
-		for (const auto& SelectedNodePtr : DragRowOp->SourceNodes)
+		for (const FSCSEditorTreeNodePtrType& SelectedNodePtr : DragRowOp->SourceNodes)
 		{
 			if (!SelectedNodePtr->CanReparent())
 			{
 				// We set the tooltip text here because it won't change across entry/leave events
 				if (DragRowOp->SourceNodes.Num() == 1)
 				{
-					if (!SelectedNodePtr->IsSceneComponent())
+					if (FChildActorComponentEditorUtils::IsChildActorSubtreeNode(SelectedNodePtr))
+					{
+						Message = LOCTEXT("DropActionToolTip_Error_CannotReparent_ChildActorSubTreeNodes", "The selected component is part of a child actor template and cannot be reordered here.");
+					}
+					else if (!SelectedNodePtr->IsSceneComponent())
 					{
 						Message = LOCTEXT("DropActionToolTip_Error_CannotReparent_NotSceneComponent", "The selected component is not a scene component and cannot be attached to other components.");
 					}
-					else if (SelectedNodePtr->IsInherited())
+					else if (SelectedNodePtr->IsInheritedComponent())
 					{
 						Message = LOCTEXT("DropActionToolTip_Error_CannotReparent_Inherited", "The selected component is inherited and cannot be reordered here.");
 					}
@@ -2162,11 +2276,16 @@ void SSCS_RowWidget::HandleOnDragEnter( const FDragDropEvent& DragDropEvent )
 			check(SceneRootNodePtr.IsValid());
 
 			FSCSEditorTreeNodePtrType NodePtr = GetNode();
-			if ((NodePtr->GetNodeType() == FSCSEditorTreeNode::SeparatorNode) || (NodePtr->GetNodeType() == FSCSEditorTreeNode::RootActorNode))
+			if (!NodePtr->IsComponentNode())
 			{
 				// Don't show a feedback message if over a node that makes no sense, such as a separator or the instance node
 				Message = LOCTEXT("DropActionToolTip_FriendlyError_DragToAComponent", "Drag to another component in order to attach to that component or become the root component.\nDrag to a Blueprint graph in order to drop a reference.");
 			}
+			else if (FChildActorComponentEditorUtils::IsChildActorSubtreeNode(NodePtr))
+			{
+				// Can't drag onto components within a child actor node's subtree
+				Message = LOCTEXT("DropActionToolTip_Error_ChildActorSubTree", "Cannot attach to this component as it is part of a child actor template.");
+			}
 
 			// Validate each selected node being dragged against the node that belongs to this row. Exit the loop if we have a valid tooltip OR a valid pending drop action once all nodes in the selection have been validated.
 			for (auto SourceNodeIter = DragRowOp->SourceNodes.CreateConstIterator(); SourceNodeIter && (Message.IsEmpty() || DragRowOp->PendingDropAction != FSCSRowDragDropOp::DropAction_None); ++SourceNodeIter)
@@ -2209,7 +2328,7 @@ void SSCS_RowWidget::HandleOnDragEnter( const FDragDropEvent& DragDropEvent )
 				{
 					if (HoveredTemplate == nullptr)
 					{
-						// Can't attach non-USceneComponent types
+						// Can't attach to non-USceneComponent types
 						Message = LOCTEXT("DropActionToolTip_Error_NotAttachable_NotSceneComponent", "Cannot attach to this component as it is not a scene component.");
 					}
 					else
@@ -2226,7 +2345,7 @@ void SSCS_RowWidget::HandleOnDragEnter( const FDragDropEvent& DragDropEvent )
 						&& DraggedTemplate->Mobility >= HoveredTemplate->Mobility
 						&& (!HoveredTemplate->IsEditorOnly() || DraggedTemplate->IsEditorOnly());
 
-					if (!NodePtr->CanReparent() && (!NodePtr->IsDefaultSceneRoot() || NodePtr->IsInherited()))
+					if (!NodePtr->CanReparent() && (!NodePtr->IsDefaultSceneRoot() || NodePtr->IsInheritedComponent()))
 					{
 						// Cannot make the dropped node the new root if we cannot reparent the current root
 						Message = LOCTEXT("DropActionToolTip_Error_CannotReparentRootNode", "The root component in this Blueprint is inherited and cannot be replaced.");
@@ -2347,12 +2466,12 @@ void SSCS_RowWidget::HandleOnDragEnter( const FDragDropEvent& DragDropEvent )
 					// Can't attach Static components to mobile ones
 					Message = LOCTEXT("DropActionToolTip_Error_CannotAttachStationary", "Cannot attach Stationary components to movable ones.");
 				}
-				else if ((NodePtr->IsInstanced() && HoveredTemplate->CreationMethod == EComponentCreationMethod::Native && !HoveredTemplate->HasAnyFlags(RF_DefaultSubObject)))
+				else if ((NodePtr->IsInstancedComponent() && HoveredTemplate->CreationMethod == EComponentCreationMethod::Native && !HoveredTemplate->HasAnyFlags(RF_DefaultSubObject)))
 				{
 					// Can't attach to post-construction C++-added components as they exist outside of the CDO and are not known at SCS execution time
 					Message = LOCTEXT("DropActionToolTip_Error_CannotAttachCPPAdded", "Cannot attach to components added in post-construction C++ code.");
 				}
-				else if (NodePtr->IsInstanced() && HoveredTemplate->CreationMethod == EComponentCreationMethod::UserConstructionScript)
+				else if (NodePtr->IsInstancedComponent() && HoveredTemplate->CreationMethod == EComponentCreationMethod::UserConstructionScript)
 				{
 					// Can't attach to UCS-added components as they exist outside of the CDO and are not known at SCS execution time
 					Message = LOCTEXT("DropActionToolTip_Error_CannotAttachUCSAdded", "Cannot attach to components added in the Construction Script.");
@@ -2549,7 +2668,7 @@ void SSCS_RowWidget::OnAttachToDropAction(const TArray<FSCSEditorTreeNodePtrType
 		AActor* PreviewActor = SCSEditorPtr->PreviewActor.Get();
 		check(PreviewActor);
 
-		for(const auto& DroppedNodePtr : DroppedNodePtrs)
+		for(const FSCSEditorTreeNodePtrType& DroppedNodePtr : DroppedNodePtrs)
 		{
 			// Clone the component if it's being dropped into a different SCS
 			if(DroppedNodePtr->GetBlueprint() != Blueprint)
@@ -2678,7 +2797,7 @@ void SSCS_RowWidget::OnAttachToDropAction(const TArray<FSCSEditorTreeNodePtrType
 	}
 	else    // EComponentEditorMode::ActorInstance
 	{
-		for(const auto& DroppedNodePtr : DroppedNodePtrs)
+		for(const FSCSEditorTreeNodePtrType& DroppedNodePtr : DroppedNodePtrs)
 		{
 			// Check for a valid parent node
 			FSCSEditorTreeNodePtrType ParentNodePtr = DroppedNodePtr->GetParent();
@@ -3049,7 +3168,7 @@ FText SSCS_RowWidget::GetNameLabel() const
 	if( InlineWidget.IsValid() && !InlineWidget->IsInEditMode() )
 	{
 		FSCSEditorTreeNodePtrType NodePtr = GetNode();
-		if(NodePtr->IsInherited())
+		if(NodePtr->IsInheritedComponent())
 		{
 			return FText::Format(LOCTEXT("NativeComponentFormatString","{0} (Inherited)"), FText::FromString(GetNode()->GetDisplayString()));
 		}
@@ -3065,7 +3184,7 @@ FText SSCS_RowWidget::GetTooltipText() const
 
 	if (NodePtr->IsDefaultSceneRoot())
 	{
-		if (NodePtr->IsInherited())
+		if (NodePtr->IsInheritedComponent())
 		{
 			return LOCTEXT("InheritedDefaultSceneRootToolTip", "This is the default scene root component. It cannot be copied, renamed or deleted.\nIt has been inherited from the parent class, so its properties cannot be edited here.\nNew scene components will automatically be attached to it.");
 		}
@@ -3094,7 +3213,7 @@ FString SSCS_RowWidget::GetDocumentationLink() const
 	check(SCSEditor.IsValid());
 
 	FSCSEditorTreeNodePtrType NodePtr = GetNode();
-	if ((NodePtr == SCSEditor.Pin()->GetSceneRootNode()) || NodePtr->IsInherited())
+	if ((NodePtr == SCSEditor.Pin()->GetSceneRootNode()) || NodePtr->IsInheritedComponent())
 	{
 		return TEXT("Shared/Editors/BlueprintEditor/ComponentsMode");
 	}
@@ -3111,11 +3230,11 @@ FString SSCS_RowWidget::GetDocumentationExcerptName() const
 	{
 		return TEXT("RootComponent");
 	}
-	else if (NodePtr->IsNative())
+	else if (NodePtr->IsNativeComponent())
 	{
 		return TEXT("NativeComponents");
 	}
-	else if (NodePtr->IsInherited())
+	else if (NodePtr->IsInheritedComponent())
 	{
 		return TEXT("InheritedComponents");
 	}
@@ -3237,6 +3356,11 @@ void SSCS_RowWidget::OnNameTextCommit(const FText& InNewName, ETextCommit::Type
 //////////////////////////////////////////////////////////////////////////
 // SSCS_RowWidget_ActorRoot
 
+FSCSEditorActorNodePtrType SSCS_RowWidget_ActorRoot::GetActorNode() const
+{
+	return StaticCastSharedPtr<FSCSEditorTreeNodeActorBase>(GetNode());
+}
+
 TSharedRef<SWidget> SSCS_RowWidget_ActorRoot::GenerateWidgetForColumn(const FName& ColumnName)
 {
 	FSCSEditorTreeNodePtrType NodePtr = GetNode();
@@ -3251,20 +3375,28 @@ TSharedRef<SWidget> SSCS_RowWidget_ActorRoot::GenerateWidgetForColumn(const FNam
 		.OnVerifyTextChanged(this, &SSCS_RowWidget_ActorRoot::OnVerifyActorLabelChanged)
 		.OnTextCommitted(this, &SSCS_RowWidget_ActorRoot::OnNameTextCommit)
 		.IsSelected(this, &SSCS_RowWidget_ActorRoot::IsSelectedExclusively)
-		.IsReadOnly(!NodePtr->CanRename() || (SCSEditor.IsValid() && !SCSEditor.Pin()->IsEditingAllowed()));
+		.IsReadOnly(!NodePtr->CanRename() || (SCSEditor.IsValid() && !SCSEditor.Pin()->IsEditingAllowed()) || FChildActorComponentEditorUtils::IsChildActorNode(NodePtr));
 
 	NodePtr->SetRenameRequestedDelegate(FSCSEditorTreeNode::FOnRenameRequested::CreateSP(InlineEditableWidget.Get(), &SInlineEditableTextBlock::EnterEditingMode));
 
 	return SNew(SHorizontalBox)
 		.ToolTip(CreateToolTipWidget())
 
+	+ SHorizontalBox::Slot()
+		.AutoWidth()
+		.VAlign(VAlign_Center)
+		[
+			SNew(SExpanderArrow, SharedThis(this))
+			.Visibility(NodePtr->GetNodeType() == FSCSEditorTreeNode::ENodeType::RootActorNode ? EVisibility::Collapsed : EVisibility::Visible)
+		]
+
 	+ SHorizontalBox::Slot()
 		.AutoWidth()
 		.VAlign(VAlign_Center)
 		.Padding(FMargin(0.f, 0.f, 6.f, 0.f))
 		[
 			SNew(SImage)
-			.Image(this, &SSCS_RowWidget_ActorRoot::GetActorIcon)
+			.Image(GetIconBrush())
 		]
 
 	+ SHorizontalBox::Slot()
@@ -3350,32 +3482,39 @@ bool SSCS_RowWidget_ActorRoot::OnVerifyActorLabelChanged(const FText& InLabel, F
 	return FActorEditorUtils::ValidateActorName(InLabel, OutErrorMessage);
 }
 
-const FSlateBrush* SSCS_RowWidget_ActorRoot::GetActorIcon() const
+const FSlateBrush* SSCS_RowWidget_ActorRoot::GetIconBrush() const
 {
-	if (SCSEditor.IsValid())
+	FSCSEditorActorNodePtrType NodePtr = GetActorNode();
+	if (NodePtr.IsValid())
 	{
-		TSharedPtr<SSCSEditor> SCSEditorPtr = SCSEditor.Pin();
-		if (SCSEditorPtr->ActorContext.IsSet())
+		if (const AActor* Actor = NodePtr->GetObject<AActor>())
 		{
-			return FClassIconFinder::FindIconForActor(SCSEditorPtr->GetActorContext());
+			return FClassIconFinder::FindIconForActor(Actor);
 		}
 	}
+
 	return nullptr;
 }
 
 FText SSCS_RowWidget_ActorRoot::GetActorDisplayText() const
 {
-	if (SCSEditor.IsValid())
+	FSCSEditorActorNodePtrType NodePtr = GetActorNode();
+	if (NodePtr.IsValid())
 	{
-		TSharedPtr<SSCSEditor> SCSEditorPtr = SCSEditor.Pin();
-		if (SCSEditorPtr->ActorContext.IsSet())
+		if (FChildActorComponentEditorUtils::IsChildActorNode(NodePtr))
+		{
+			if (const AActor* ChildActor = NodePtr->GetObject<AActor>())
+			{
+				return ChildActor->GetClass()->GetDisplayNameText();
+			}
+		}
+		else
 		{
-			AActor* DefaultActor = SCSEditorPtr->ActorContext.Get();
-			if( DefaultActor )
+			if (const AActor* DefaultActor = NodePtr->GetObject<AActor>())
 			{
 				FString Name;
 				UBlueprint* Blueprint = UBlueprint::GetBlueprintFromClass(DefaultActor->GetClass());
-				if(Blueprint != nullptr && SCSEditorPtr->GetEditorMode() != EComponentEditorMode::ActorInstance)
+				if (Blueprint != nullptr && !NodePtr->IsInstanced())
 				{
 					Blueprint->GetName(Name);
 				}
@@ -3387,23 +3526,31 @@ FText SSCS_RowWidget_ActorRoot::GetActorDisplayText() const
 			}
 		}
 	}
+
 	return FText::GetEmpty();
 }
 
 FText SSCS_RowWidget_ActorRoot::GetActorContextText() const
 {
-	if (SCSEditor.IsValid())
+	FSCSEditorActorNodePtrType NodePtr = GetActorNode();
+	if (NodePtr.IsValid())
 	{
-		TSharedPtr<SSCSEditor> SCSEditorPtr = SCSEditor.Pin();
-		if (AActor* DefaultActor = SCSEditorPtr->GetActorContext())
+		if (FChildActorComponentEditorUtils::IsChildActorNode(NodePtr))
 		{
-			if (UBlueprint* Blueprint = UBlueprint::GetBlueprintFromClass(DefaultActor->GetClass()))
-			{
-				return LOCTEXT("ActorContext_self", " (self)");
-			}
-			else
+			return LOCTEXT("ActorContext_ChildActor", " (Child Actor)");
+		}
+		else
+		{
+			if (const AActor* DefaultActor = NodePtr->GetObject<AActor>())
 			{
-				return LOCTEXT("ActorContext_Instance", " (Instance)");
+				if (UBlueprint* Blueprint = UBlueprint::GetBlueprintFromClass(DefaultActor->GetClass()))
+				{
+					return LOCTEXT("ActorContext_self", " (self)");
+				}
+				else
+				{
+					return LOCTEXT("ActorContext_Instance", " (Instance)");
+				}
 			}
 		}
 	}
@@ -3412,45 +3559,42 @@ FText SSCS_RowWidget_ActorRoot::GetActorContextText() const
 
 FText SSCS_RowWidget_ActorRoot::GetActorClassNameText() const
 {
-	FText Text;
-	if (SCSEditor.IsValid())
+	FSCSEditorActorNodePtrType NodePtr = GetActorNode();
+	if (NodePtr.IsValid())
 	{
-		TSharedPtr<SSCSEditor> SCSEditorPtr = SCSEditor.Pin();
-		if (AActor* DefaultActor = SCSEditorPtr->GetActorContext())
+		if (const AActor* DefaultActor = NodePtr->GetObject<AActor>())
 		{
-			Text = FText::FromString(DefaultActor->GetClass()->GetName());
+			return FText::FromString(DefaultActor->GetClass()->GetName());
 		}
 	}
 
-	return Text;
+	return FText::GetEmpty();
 }
 
 FText SSCS_RowWidget_ActorRoot::GetActorSuperClassNameText() const
 {
-	FText Text;
-	if (SCSEditor.IsValid())
+	FSCSEditorActorNodePtrType NodePtr = GetActorNode();
+	if (NodePtr.IsValid())
 	{
-		TSharedPtr<SSCSEditor> SCSEditorPtr = SCSEditor.Pin();
-		if (AActor* DefaultActor = SCSEditorPtr->GetActorContext())
+		if (const AActor* DefaultActor = NodePtr->GetObject<AActor>())
 		{
-			Text = FText::FromString(DefaultActor->GetClass()->GetSuperClass()->GetName());
+			return FText::FromString(DefaultActor->GetClass()->GetSuperClass()->GetName());
 		}
 	}
 
-	return Text;
+	return FText::GetEmpty();
 }
 
 FText SSCS_RowWidget_ActorRoot::GetActorMobilityText() const
 {
-	FText Text;
-	if (SCSEditor.IsValid())
+	FSCSEditorActorNodePtrType NodePtr = GetActorNode();
+	if (NodePtr.IsValid())
 	{
-		TSharedPtr<SSCSEditor> SCSEditorPtr = SCSEditor.Pin();
-		if (AActor* DefaultActor = SCSEditorPtr->GetActorContext())
+		if (const AActor* DefaultActor = NodePtr->GetObject<AActor>())
 		{
 			USceneComponent* RootComponent = DefaultActor->GetRootComponent();
 
-			FSCSEditorTreeNodePtrType SceneRootNodePtr = SCSEditorPtr->GetSceneRootNode();
+			FSCSEditorTreeNodePtrType SceneRootNodePtr = NodePtr->GetSceneRootNode();
 			if ((RootComponent == nullptr) && SceneRootNodePtr.IsValid())
 			{
 				RootComponent = Cast<USceneComponent>(SceneRootNodePtr->GetComponentTemplate());
@@ -3460,25 +3604,25 @@ FText SSCS_RowWidget_ActorRoot::GetActorMobilityText() const
 			{
 				if (RootComponent->Mobility == EComponentMobility::Static)
 				{
-					Text = LOCTEXT("ComponentMobility_Static", "Static");
+					return LOCTEXT("ComponentMobility_Static", "Static");
 				}
 				else if (RootComponent->Mobility == EComponentMobility::Stationary)
 				{
-					Text = LOCTEXT("ComponentMobility_Stationary", "Stationary");
+					return LOCTEXT("ComponentMobility_Stationary", "Stationary");
 				}
 				else if (RootComponent->Mobility == EComponentMobility::Movable)
 				{
-					Text = LOCTEXT("ComponentMobility_Movable", "Movable");
+					return LOCTEXT("ComponentMobility_Movable", "Movable");
 				}
 			}
 			else
 			{
-				Text = LOCTEXT("ComponentMobility_NoRoot", "No root component, unknown mobility");
+				return LOCTEXT("ComponentMobility_NoRoot", "No root component, unknown mobility");
 			}
 		}
 	}
 
-	return Text;
+	return FText::GetEmpty();
 }
 
 //////////////////////////////////////////////////////////////////////////
@@ -3834,9 +3978,14 @@ void SSCSEditor::OnLevelComponentRequestRename(const UActorComponent* InComponen
 
 void SSCSEditor::OnObjectsReplaced(const TMap<UObject*, UObject*>& OldToNewInstanceMap)
 {
-	if (GetActorNode().IsValid())
+	ReplaceComponentReferencesInTree(GetActorNode(), OldToNewInstanceMap);
+}
+
+void SSCSEditor::ReplaceComponentReferencesInTree(FSCSEditorActorNodePtrType InActorNode, const TMap<UObject*, UObject*>& OldToNewInstanceMap)
+{
+	if (InActorNode.IsValid())
 	{
-		ReplaceComponentReferencesInTree(GetActorNode()->GetComponentNodes(), OldToNewInstanceMap);
+		ReplaceComponentReferencesInTree(InActorNode->GetComponentNodes(), OldToNewInstanceMap);
 	}
 }
 
@@ -3846,18 +3995,15 @@ void SSCSEditor::ReplaceComponentReferencesInTree(const TArray<FSCSEditorTreeNod
 	{
 		if (Node.IsValid())
 		{
-			// We need to get the actual pointer to the old component which will be marked for pending kill, as these are the references which need updating
+			// We need to get the actual pointer to the old object which will be marked for pending kill, as these are the references which need updating
 			const bool bEvenIfPendingKill = true;
-			UActorComponent* ComponentTemplate = Node->GetComponentTemplate(bEvenIfPendingKill);
-			if (ComponentTemplate)
+			const UObject* OldObject = Node->GetObject<UObject>(bEvenIfPendingKill);
+			if (OldObject)
 			{
-				UObject* const* NewComponentTemplatePtr = OldToNewInstanceMap.Find(ComponentTemplate);
-				if (NewComponentTemplatePtr)
+				UObject* const* NewObject = OldToNewInstanceMap.Find(OldObject);
+				if (NewObject)
 				{
-					if (UActorComponent* NewComponentTemplate = Cast<UActorComponent>(*NewComponentTemplatePtr))
-					{
-						Node->SetComponentTemplate(NewComponentTemplate);
-					}
+					Node->SetObject(*NewObject);
 				}
 			}
 
@@ -3870,10 +4016,10 @@ UBlueprint* SSCSEditor::GetBlueprint() const
 {
 	if (AActor* Actor = GetActorContext())
 	{
-		UClass* ActorClass = Actor->GetClass();
+		const UClass* ActorClass = Actor->GetClass();
 		check(ActorClass != nullptr);
 
-		return Cast<UBlueprint>(ActorClass->ClassGeneratedBy);
+		return UBlueprint::GetBlueprintFromClass(ActorClass);
 	}
 
 	return nullptr;
@@ -3898,7 +4044,7 @@ TSharedRef<ITableRow> SSCSEditor::MakeTableRowWidget( FSCSEditorTreeNodePtrType
 	}
 
 	// Create the node of the appropriate type
-	if (InNodePtr->GetNodeType() == FSCSEditorTreeNode::RootActorNode)
+	if (InNodePtr->IsActorNode())
 	{
 		return SNew(SSCS_RowWidget_ActorRoot, SharedThis(this), InNodePtr, OwnerTable);
 	}
@@ -3917,7 +4063,7 @@ void SSCSEditor::GetSelectedItemsForContextMenu(TArray<FComponentEventConstructi
 	for ( auto NodeIter = SelectedTreeItems.CreateConstIterator(); NodeIter; ++NodeIter )
 	{
 		FComponentEventConstructionData NewItem;
-		auto TreeNode = *NodeIter;
+		const FSCSEditorTreeNodePtrType& TreeNode = *NodeIter;
 		NewItem.VariableName = TreeNode->GetVariableName();
 		NewItem.Component = TreeNode->GetComponentTemplate();
 		OutSelectedItems.Add(NewItem);
@@ -3934,15 +4080,23 @@ void SSCSEditor::PopulateContextMenu(UToolMenu* Menu)
 
 		if (SelectedItems.Num() > 0)
 		{
-			if (SelectedItems.Num() == 1 && SelectedItems[0]->GetNodeType() == FSCSEditorTreeNode::RootActorNode)
+			if (SelectedItems.Num() == 1 && SelectedItems[0]->IsActorNode())
 			{
-				bOnlyShowPasteOption = true;
+				if (FChildActorComponentEditorUtils::IsChildActorNode(SelectedItems[0]))
+				{
+					// Include specific context menu options for a single child actor node selection
+					FChildActorComponentEditorUtils::FillChildActorContextMenuOptions(Menu, SelectedItems[0]);
+				}
+				else
+				{
+					bOnlyShowPasteOption = true;
+				}
 			}
 			else
 			{
 				for (FSCSEditorTreeNodePtrType SelectedNode : SelectedItems)
 				{
-					if (SelectedNode->GetNodeType() != FSCSEditorTreeNode::ComponentNode)
+					if (!SelectedNode->IsComponentNode())
 					{
 						bOnlyShowPasteOption = true;
 						break;
@@ -3950,6 +4104,8 @@ void SSCSEditor::PopulateContextMenu(UToolMenu* Menu)
 				}
 				if (!bOnlyShowPasteOption)
 				{
+					bool bIsChildActorSubtreeNodeSelected = false;
+
 					TArray<UActorComponent*> SelectedComponents;
 					TArray<FSCSEditorTreeNodePtrType> SelectedNodes = GetSelectedNodes();
 					for (int32 i = 0; i < SelectedNodes.Num(); ++i)
@@ -3964,9 +4120,16 @@ void SSCSEditor::PopulateContextMenu(UToolMenu* Menu)
 						{
 							SelectedComponents.Add(ComponentTemplate);
 						}
+
+						// Determine if any selected node belongs to a child actor template
+						if (!bIsChildActorSubtreeNodeSelected)
+						{
+							bIsChildActorSubtreeNodeSelected = FChildActorComponentEditorUtils::IsChildActorSubtreeNode(SelectedNodePtr);
+						}
 					}
 
-					if (EditorMode == EComponentEditorMode::BlueprintSCS)
+					// Don't include these commands if any component was found above to belong to a child actor template (not supported at this time)
+					if (EditorMode == EComponentEditorMode::BlueprintSCS && !bIsChildActorSubtreeNodeSelected)
 					{
 						FToolMenuSection& BlueprintSCSSection = Menu->AddSection("BlueprintSCS");
 						if (SelectedItems.Num() == 1)
@@ -3976,34 +4139,43 @@ void SSCSEditor::PopulateContextMenu(UToolMenu* Menu)
 
 						// Collect the classes of all selected objects
 						TArray<UClass*> SelectionClasses;
-						for( auto NodeIter = SelectedNodes.CreateConstIterator(); NodeIter; ++NodeIter )
+						for (auto NodeIter = SelectedNodes.CreateConstIterator(); NodeIter; ++NodeIter)
 						{
-							auto TreeNode = *NodeIter;
-							if( UActorComponent* ComponentTemplate = TreeNode->GetComponentTemplate() )
+							const FSCSEditorTreeNodePtrType& TreeNode = *NodeIter;
+							if (UActorComponent* ComponentTemplate = TreeNode->GetComponentTemplate())
 							{
 								SelectionClasses.Add(ComponentTemplate->GetClass());
 							}
 						}
 
-						if ( SelectionClasses.Num() )
+						if (SelectionClasses.Num())
 						{
 							// Find the common base class of all selected classes
-							UClass* SelectedClass = UClass::FindCommonBase( SelectionClasses );
+							UClass* SelectedClass = UClass::FindCommonBase(SelectionClasses);
 							// Build an event submenu if we can generate events
-							if( FBlueprintEditorUtils::CanClassGenerateEvents( SelectedClass ))
+							if (FBlueprintEditorUtils::CanClassGenerateEvents(SelectedClass))
 							{
 								BlueprintSCSSection.AddSubMenu(
 									"AddEventSubMenu",
-									LOCTEXT("AddEventSubMenu", "Add Event"), 
-									LOCTEXT("ActtionsSubMenu_ToolTip", "Add Event"), 
-									FNewMenuDelegate::CreateStatic( &SSCSEditor::BuildMenuEventsSection,
-									GetBlueprint(), SelectedClass, FCanExecuteAction::CreateSP(this, &SSCSEditor::IsEditingAllowed),
-									FGetSelectedObjectsDelegate::CreateSP(this, &SSCSEditor::GetSelectedItemsForContextMenu)));
+									LOCTEXT("AddEventSubMenu", "Add Event"),
+									LOCTEXT("ActtionsSubMenu_ToolTip", "Add Event"),
+									FNewMenuDelegate::CreateStatic(&SSCSEditor::BuildMenuEventsSection,
+										GetBlueprint(), SelectedClass, FCanExecuteAction::CreateSP(this, &SSCSEditor::IsEditingAllowed),
+										FGetSelectedObjectsDelegate::CreateSP(this, &SSCSEditor::GetSelectedItemsForContextMenu)));
 							}
 						}
-					}					
+					}
 
 					FComponentEditorUtils::FillComponentContextMenuOptions(Menu, SelectedComponents);
+
+					// For a single selected child actor component, we may choose to include additional context menu options
+					if (SelectedComponents.Num() == 1)
+					{
+						if (UChildActorComponent* SelectedChildActorComponent = Cast<UChildActorComponent>(SelectedComponents[0]))
+						{
+							FChildActorComponentEditorUtils::FillComponentContextMenuOptions(Menu, SelectedChildActorComponent);
+						}
+					}
 				}
 			}
 		}
@@ -4014,9 +4186,9 @@ void SSCSEditor::PopulateContextMenu(UToolMenu* Menu)
 
 		if (bOnlyShowPasteOption)
 		{
-			FToolMenuSection& Section = Menu->AddSection("PasteComponent", LOCTEXT("EditComponentHeading", "Edit") );
+			FToolMenuSection& Section = Menu->AddSection("PasteComponent", LOCTEXT("EditComponentHeading", "Edit"));
 			{
-				Section.AddMenuEntry( FGenericCommands::Get().Paste );
+				Section.AddMenuEntry(FGenericCommands::Get().Paste);
 			}
 		}
 	}
@@ -4025,9 +4197,9 @@ void SSCSEditor::PopulateContextMenu(UToolMenu* Menu)
 void SSCSEditor::RegisterContextMenu()
 {
 	UToolMenus* ToolMenus = UToolMenus::Get();
-	if (!ToolMenus->IsMenuRegistered("Kismet.SCSEditorContextMenu"))
+	if (!ToolMenus->IsMenuRegistered(SCS_ContextMenuName))
 	{
-		UToolMenu* Menu = ToolMenus->RegisterMenu("Kismet.SCSEditorContextMenu");
+		UToolMenu* Menu = ToolMenus->RegisterMenu(SCS_ContextMenuName);
 		Menu->AddDynamicSection("SCSEditorDynamic", FNewToolMenuDelegate::CreateLambda([](UToolMenu* InMenu)
 		{
 			USSCSEditorMenuContext* ContextObject = InMenu->FindContext<USSCSEditorMenuContext>();
@@ -4049,7 +4221,7 @@ TSharedPtr< SWidget > SSCSEditor::CreateContextMenu()
 		USSCSEditorMenuContext* ContextObject = NewObject<USSCSEditorMenuContext>();
 		ContextObject->SCSEditor = SharedThis(this);
 		FToolMenuContext ToolMenuContext(CommandList, TSharedPtr<FExtender>(), ContextObject);
-		return UToolMenus::Get()->GenerateWidget("Kismet.SCSEditorContextMenu", ToolMenuContext);
+		return UToolMenus::Get()->GenerateWidget(SCS_ContextMenuName, ToolMenuContext);
 	}
 	return TSharedPtr<SWidget>();
 }
@@ -4194,7 +4366,8 @@ bool SSCSEditor::CanDuplicateComponent() const
 		return false;
 	}
 
-	return CanCopyNodes();
+	// @todo - Allow duplication of components that belong to a child actor template? For now, we don't support this.
+	return CanCopyNodes() && !FChildActorComponentEditorUtils::ContainsChildActorSubtreeNode(SCSTreeWidget->GetSelectedItems());
 }
 
 void SSCSEditor::OnDuplicateComponent()
@@ -4508,9 +4681,9 @@ static FSCSEditorTreeNode* FindRecursive( FSCSEditorTreeNode* Node, FName Name )
 	}
 	else
 	{
-		for (const auto& Child : Node->GetChildren())
+		for (const FSCSEditorTreeNodePtrType& Child : Node->GetChildren())
 		{
-			if (auto Result = FindRecursive(Child.Get(), Name))
+			if (FSCSEditorTreeNode* Result = FindRecursive(Child.Get(), Name))
 			{
 				return Result;
 			}
@@ -4522,9 +4695,9 @@ static FSCSEditorTreeNode* FindRecursive( FSCSEditorTreeNode* Node, FName Name )
 
 void SSCSEditor::HighlightTreeNode(FName TreeNodeName, const class FPropertyPath& Property)
 {
-	for( const auto& Node : GetRootNodes() )
+	for( const FSCSEditorTreeNodePtrType& Node : GetRootNodes() )
 	{
-		if( auto FoundNode = FindRecursive( Node.Get(), TreeNodeName ) )
+		if( FSCSEditorTreeNode* FoundNode = FindRecursive( Node.Get(), TreeNodeName ) )
 		{
 			SelectNode(FoundNode->AsShared(), false);
 
@@ -4544,7 +4717,7 @@ void SSCSEditor::HighlightTreeNode(FName TreeNodeName, const class FPropertyPath
 void SSCSEditor::HighlightTreeNode(const USCS_Node* Node, FName Property)
 {
 	check(Node);
-	auto TreeNode = FindTreeNode( Node );
+	FSCSEditorTreeNodePtrType TreeNode = FindTreeNode( Node );
 	check( TreeNode.IsValid() );
 	SelectNode( TreeNode, false );
 	if( Property != FName() )
@@ -4563,211 +4736,211 @@ void SSCSEditor::HighlightTreeNode(const USCS_Node* Node, FName Property)
 	}
 }
 
-void SSCSEditor::UpdateTree(bool bRegenerateTreeNodes)
+void SSCSEditor::BuildSubTreeForActorNode(FSCSEditorActorNodePtrType InActorNode)
 {
-	check(SCSTreeWidget.IsValid());
-
-	// Early exit if we're deferring tree updates
-	if(!bAllowTreeUpdates)
+	if (!InActorNode.IsValid())
 	{
 		return;
 	}
 
-	if(bRegenerateTreeNodes)
+	// Get the actor instance that we're editing
+	const AActor* Actor = InActorNode->GetObject<AActor>();
+	if (!Actor)
 	{
-		// Obtain the set of expandable tree nodes that are currently collapsed
-		TSet<FSCSEditorTreeNodePtrType> CollapsedTreeNodes;
-		GetCollapsedNodes(GetSceneRootNode(), CollapsedTreeNodes);
+		return;
+	}
 
-		// Obtain the list of selected items
-		TArray<FSCSEditorTreeNodePtrType> SelectedTreeNodes = SCSTreeWidget->GetSelectedItems();
+	// Build the tree data source according to what mode we're in
+	if (!InActorNode->IsInstanced())
+	{
+		TInlineComponentArray<UActorComponent*> Components;
+		Actor->GetComponents(Components);
 
-		// Clear the current tree
-		if (SelectedTreeNodes.Num() != 0)
+		// Add the native root component
+		USceneComponent* RootComponent = Actor->GetRootComponent();
+		if (RootComponent != nullptr)
 		{
-			SCSTreeWidget->ClearSelection();
+			Components.Remove(RootComponent);
+			AddTreeNodeFromComponent(RootComponent, FindOrCreateParentForExistingComponent(RootComponent, InActorNode));
 		}
-		RootNodes.Empty();
-
-		TSharedPtr<FSCSEditorTreeNode> ActorTreeNode = MakeShareable(new FSCSEditorTreeNodeRootActor(GetActorContext(),EditorMode == EComponentEditorMode::ActorInstance));
-		RefreshFilteredState(ActorTreeNode, false);
-		SCSTreeWidget->SetItemExpansion(ActorTreeNode, true);
-		RootNodes.Add(ActorTreeNode);
 
-		// Build the tree data source according to what mode we're in
-		if (EditorMode == EComponentEditorMode::BlueprintSCS)
+		// Add the rest of the native base class SceneComponent hierarchy
+		for (UActorComponent* Component : Components)
 		{
-			// Get the class default object
-			AActor* CDO = nullptr;
-			TArray<UBlueprint*> ParentBPStack;
+			AddTreeNodeFromComponent(Component, FindOrCreateParentForExistingComponent(Component, InActorNode));
+		}
 
-			if(AActor* Actor = GetActorContext())
-			{
-				UClass* ActorClass = Actor->GetClass();
-				if(ActorClass != nullptr)
-				{
-					CDO = ActorClass->GetDefaultObject<AActor>();
+		// If it's a Blueprint-generated class, also get the inheritance stack
+		TArray<UBlueprint*> ParentBPStack;
+		UBlueprint::GetBlueprintHierarchyFromClass(Actor->GetClass(), ParentBPStack);
 
-					// If it's a Blueprint-generated class, also get the inheritance stack
-					UBlueprint::GetBlueprintHierarchyFromClass(ActorClass, ParentBPStack);
-				}
-			}
-
-			if(CDO != nullptr)
+		// Add the full SCS tree node hierarchy (including SCS nodes inherited from parent blueprints)
+		for (int32 StackIndex = ParentBPStack.Num() - 1; StackIndex >= 0; --StackIndex)
+		{
+			if (ParentBPStack[StackIndex]->SimpleConstructionScript != nullptr)
 			{
-				
-				TInlineComponentArray<UActorComponent*> Components;
-				CDO->GetComponents(Components);
-
-				// Add the native root component
-				USceneComponent* RootComponent = CDO->GetRootComponent();
-				if(RootComponent != nullptr)
-				{
-					Components.Remove(RootComponent);
-					AddTreeNodeFromComponent(RootComponent, FindOrCreateParentForExistingComponent(RootComponent, GetActorNode()));
-				}
-				
-				for (UActorComponent* Component : Components)
+				const TArray<USCS_Node*>& SCS_RootNodes = ParentBPStack[StackIndex]->SimpleConstructionScript->GetRootNodes();
+				for (int32 NodeIndex = 0; NodeIndex < SCS_RootNodes.Num(); ++NodeIndex)
 				{
-					// Add the rest of the native base class SceneComponent hierarchy
-					AddTreeNodeFromComponent(Component, FindOrCreateParentForExistingComponent(Component, GetActorNode()));
-				}
-			}
+					USCS_Node* SCS_Node = SCS_RootNodes[NodeIndex];
+					check(SCS_Node != nullptr);
 
-			// Add the full SCS tree node hierarchy (including SCS nodes inherited from parent blueprints)
-			for(int32 StackIndex = ParentBPStack.Num() - 1; StackIndex >= 0; --StackIndex)
-			{
-				if(ParentBPStack[StackIndex]->SimpleConstructionScript != nullptr)
-				{
-					const TArray<USCS_Node*>& SCS_RootNodes = ParentBPStack[StackIndex]->SimpleConstructionScript->GetRootNodes();
-					for(int32 NodeIndex = 0; NodeIndex < SCS_RootNodes.Num(); ++NodeIndex)
+					FSCSEditorTreeNodePtrType NewNodePtr;
+					if (SCS_Node->ParentComponentOrVariableName != NAME_None)
 					{
-						USCS_Node* SCS_Node = SCS_RootNodes[NodeIndex];
-						check(SCS_Node != nullptr);
-
-						FSCSEditorTreeNodePtrType NewNodePtr;
-						if(SCS_Node->ParentComponentOrVariableName != NAME_None)
+						USceneComponent* ParentComponent = SCS_Node->GetParentComponentTemplate(ParentBPStack[0]);
+						if (ParentComponent != nullptr)
 						{
-							USceneComponent* ParentComponent = SCS_Node->GetParentComponentTemplate(ParentBPStack[0]);
-							if(ParentComponent != nullptr)
+							FSCSEditorTreeNodePtrType ParentNodePtr = FindTreeNode(ParentComponent, InActorNode);
+							if (ParentNodePtr.IsValid())
 							{
-								FSCSEditorTreeNodePtrType ParentNodePtr = FindTreeNode(ParentComponent);
-								if(ParentNodePtr.IsValid())
-								{
-									NewNodePtr = AddTreeNode(SCS_Node, ParentNodePtr, StackIndex > 0);
-								}
+								NewNodePtr = AddTreeNode(SCS_Node, ParentNodePtr, StackIndex > 0);
 							}
 						}
-						else
-						{
-							NewNodePtr = AddTreeNode(SCS_Node, ActorTreeNode, StackIndex > 0);
-						}
+					}
+					else
+					{
+						NewNodePtr = AddTreeNode(SCS_Node, InActorNode, StackIndex > 0);
+					}
 
-						// Only necessary to do the following for inherited nodes (StackIndex > 0).
-						if (NewNodePtr.IsValid() && StackIndex > 0)
+					// Only necessary to do the following for inherited nodes (StackIndex > 0).
+					if (NewNodePtr.IsValid() && StackIndex > 0)
+					{
+						// This call creates ICH override templates for the current Blueprint. Without this, the parent node
+						// search above can fail when attempting to match an inherited node in the tree via component template.
+						NewNodePtr->GetOrCreateEditableComponentTemplate(ParentBPStack[0]);
+						for (FSCSEditorTreeNodePtrType ChildNodePtr : NewNodePtr->GetChildren())
 						{
-							// This call creates ICH override templates for the current Blueprint. Without this, the parent node
-							// search above can fail when attempting to match an inherited node in the tree via component template.
-							NewNodePtr->GetOrCreateEditableComponentTemplate(ParentBPStack[0]);
-							for (FSCSEditorTreeNodePtrType ChildNodePtr : NewNodePtr->GetChildren())
+							if (ensure(ChildNodePtr.IsValid()))
 							{
-								if (ensure(ChildNodePtr.IsValid()))
-								{
-									ChildNodePtr->GetOrCreateEditableComponentTemplate(ParentBPStack[0]);
-								}
+								ChildNodePtr->GetOrCreateEditableComponentTemplate(ParentBPStack[0]);
 							}
 						}
 					}
 				}
 			}
+		}
+	}
+	else    // InActorNode->IsInstanced()
+	{
+		// Get the full set of instanced components
+		TSet<UActorComponent*> ComponentsToAdd(Actor->GetComponents());
+
+		const bool bHideConstructionScriptComponentsInDetailsView = GetDefault<UBlueprintEditorSettings>()->bHideConstructionScriptComponentsInDetailsView;
+		auto ShouldAddInstancedActorComponent = [bHideConstructionScriptComponentsInDetailsView](UActorComponent* ActorComp, USceneComponent* ParentSceneComp)
+		{
+			// Exclude nested DSOs attached to BP-constructed instances, which are not mutable.
+			return (ActorComp != nullptr
+				&& (!ActorComp->IsVisualizationComponent())
+				&& (ActorComp->CreationMethod != EComponentCreationMethod::UserConstructionScript || !bHideConstructionScriptComponentsInDetailsView)
+				&& (ParentSceneComp == nullptr || !ParentSceneComp->IsCreatedByConstructionScript() || !ActorComp->HasAnyFlags(RF_DefaultSubObject)))
+				&& (ActorComp->CreationMethod != EComponentCreationMethod::Native || FComponentEditorUtils::GetPropertyForEditableNativeComponent(ActorComp));
+		};
 
-			AActor* PreviewActorInstance = PreviewActor.Get();
-			if(PreviewActorInstance != nullptr && !GetDefault<UBlueprintEditorSettings>()->bHideConstructionScriptComponentsInDetailsView)
+		for (TSet<UActorComponent*>::TIterator It(ComponentsToAdd.CreateIterator()); It; ++It)
+		{
+			UActorComponent* ActorComp = *It;
+			USceneComponent* SceneComp = Cast<USceneComponent>(ActorComp);
+			USceneComponent* ParentSceneComp = SceneComp != nullptr ? SceneComp->GetAttachParent() : nullptr;
+			if (!ShouldAddInstancedActorComponent(ActorComp, ParentSceneComp))
 			{
-				TInlineComponentArray<UActorComponent*> Components;
-				PreviewActorInstance->GetComponents(Components);
+				It.RemoveCurrent();
+			}
+		}
 
-				for (UActorComponent* Component : Components)
+		TFunction<void(USceneComponent*, FSCSEditorTreeNodePtrType)> AddInstancedTreeNodesRecursive = [&](USceneComponent* Component, FSCSEditorTreeNodePtrType TreeNode)
+		{
+			if (Component != nullptr)
+			{
+				TArray<USceneComponent*> Components = Component->GetAttachChildren();
+				for (USceneComponent* ChildComponent : Components)
 				{
-					if(Component->CreationMethod == EComponentCreationMethod::UserConstructionScript)
+					if (ComponentsToAdd.Contains(ChildComponent)
+						&& ChildComponent->GetOwner() == Component->GetOwner())
 					{
-						AddTreeNodeFromComponent(Component, FindOrCreateParentForExistingComponent(Component, GetActorNode()));
+						ComponentsToAdd.Remove(ChildComponent);
+
+						FSCSEditorTreeNodePtrType NewParentNode = AddTreeNodeFromComponent(ChildComponent, TreeNode);
+						AddInstancedTreeNodesRecursive(ChildComponent, NewParentNode);
 					}
 				}
 			}
+		};
+
+		// Add the root component first (it may not be the first one)
+		USceneComponent* RootComponent = Actor->GetRootComponent();
+		if (RootComponent != nullptr)
+		{
+			ComponentsToAdd.Remove(RootComponent);
+
+			// Recursively add any instanced children that are already attached through the root, and keep track of added
+			// instances. This will be a faster path than the loop below, because we create new parent tree nodes as we go.
+			FSCSEditorTreeNodePtrType NewParentNode = AddTreeNodeFromComponent(RootComponent, FindOrCreateParentForExistingComponent(RootComponent, InActorNode));
+			AddInstancedTreeNodesRecursive(RootComponent, NewParentNode);
 		}
-		else    // EComponentEditorMode::ActorInstance
+
+		// Sort components by type (always put scene components first in the tree)
+		ComponentsToAdd.Sort([](const UActorComponent& A, const UActorComponent& /* B */)
 		{
-			// Get the actor instance that we're editing
-			if (AActor* ActorInstance = GetActorContext())
-			{
-				// Get the full set of instanced components
-				TSet<UActorComponent*> ComponentsToAdd(ActorInstance->GetComponents());
+			return A.IsA<USceneComponent>();
+		});
 
-				const bool bHideConstructionScriptComponentsInDetailsView = GetDefault<UBlueprintEditorSettings>()->bHideConstructionScriptComponentsInDetailsView;
-				auto ShouldAddInstancedActorComponent = [bHideConstructionScriptComponentsInDetailsView](UActorComponent* ActorComp, USceneComponent* ParentSceneComp)
-				{
-					// Exclude nested DSOs attached to BP-constructed instances, which are not mutable.
-					return (ActorComp != nullptr
-						&& (!ActorComp->IsVisualizationComponent())
-						&& (ActorComp->CreationMethod != EComponentCreationMethod::UserConstructionScript || !bHideConstructionScriptComponentsInDetailsView)
-						&& (ParentSceneComp == nullptr || !ParentSceneComp->IsCreatedByConstructionScript() || !ActorComp->HasAnyFlags(RF_DefaultSubObject)))
-						&& (ActorComp->CreationMethod != EComponentCreationMethod::Native || FComponentEditorUtils::GetPropertyForEditableNativeComponent(ActorComp));
-				};
+		// Now add any remaining instanced owned components not already added above. This will first add any
+		// unattached scene components followed by any instanced non-scene components owned by the Actor instance.
+		for (UActorComponent* ActorComp : ComponentsToAdd)
+		{
+			AddTreeNodeFromComponent(ActorComp, FindOrCreateParentForExistingComponent(ActorComp, InActorNode));
+		}
+	}
 
-				for (auto It(ComponentsToAdd.CreateIterator()); It; ++It)
-				{
-					UActorComponent* ActorComp = *It;
-					USceneComponent* SceneComp = Cast<USceneComponent>(ActorComp);
-					USceneComponent* ParentSceneComp = SceneComp != nullptr ? SceneComp->GetAttachParent() : nullptr;
-					if (!ShouldAddInstancedActorComponent(ActorComp, ParentSceneComp))
-					{
-						It.RemoveCurrent();
-					}
-				}
+	// Always expand actor nodes to reveal children
+	SCSTreeWidget->SetItemExpansion(InActorNode, true);
+}
 
-				TFunction<void(USceneComponent*,FSCSEditorTreeNodePtrType)> AddInstancedTreeNodesRecursive = [&](USceneComponent* Component, FSCSEditorTreeNodePtrType TreeNode)
-				{
-					if (Component != nullptr)
-					{
-						TArray<USceneComponent*> Components = Component->GetAttachChildren();
-						for (USceneComponent* ChildComponent : Components)
-						{
-							if (ComponentsToAdd.Contains(ChildComponent)
-								&& ChildComponent->GetOwner() == Component->GetOwner())
-							{
-								ComponentsToAdd.Remove(ChildComponent);
+void SSCSEditor::UpdateTree(bool bRegenerateTreeNodes)
+{
+	check(SCSTreeWidget.IsValid());
 
-								FSCSEditorTreeNodePtrType NewParentNode = AddTreeNodeFromComponent(ChildComponent, TreeNode);
-								AddInstancedTreeNodesRecursive(ChildComponent, NewParentNode);
-							}
-						}
-					}
-				};
+	// Early exit if we're deferring tree updates
+	if(!bAllowTreeUpdates)
+	{
+		return;
+	}
 
-				// Add the root component first (it may not be the first one)
-				USceneComponent* RootComponent = ActorInstance->GetRootComponent();
-				if(RootComponent != nullptr)
-				{
-					ComponentsToAdd.Remove(RootComponent);
+	if(bRegenerateTreeNodes)
+	{
+		// Obtain the set of expandable tree nodes that are currently collapsed
+		TSet<FSCSEditorTreeNodePtrType> CollapsedTreeNodes;
+		GetCollapsedNodes(GetSceneRootNode(), CollapsedTreeNodes);
 
-					// Recursively add any instanced children that are already attached through the root, and keep track of added
-					// instances. This will be a faster path than the loop below, because we create new parent tree nodes as we go.
-					FSCSEditorTreeNodePtrType NewParentNode = AddTreeNodeFromComponent(RootComponent, FindOrCreateParentForExistingComponent(RootComponent, GetActorNode()));
-					AddInstancedTreeNodesRecursive(RootComponent, NewParentNode);
-				}
+		// Obtain the list of selected items
+		TArray<FSCSEditorTreeNodePtrType> SelectedTreeNodes = SCSTreeWidget->GetSelectedItems();
 
-				// Sort components by type (always put scene components first in the tree)
-				ComponentsToAdd.Sort([](const UActorComponent& A, const UActorComponent& /* B */)
-				{
-					return A.IsA<USceneComponent>();
-				});
+		// Clear the current tree
+		if (SelectedTreeNodes.Num() != 0)
+		{
+			SCSTreeWidget->ClearSelection();
+		}
+		RootNodes.Empty();
 
-				// Now add any remaining instanced owned components not already added above. This will first add any
-				// unattached scene components followed by any instanced non-scene components owned by the Actor instance.
-				for (UActorComponent* ActorComp : ComponentsToAdd)
+		TSharedPtr<FSCSEditorTreeNodeRootActor> RootActorNode = MakeShareable(new FSCSEditorTreeNodeRootActor(GetActorContext(), EditorMode == EComponentEditorMode::ActorInstance));
+		RefreshFilteredState(RootActorNode, false);
+		SCSTreeWidget->SetItemExpansion(RootActorNode, true);
+		RootNodes.Add(RootActorNode);
+
+		BuildSubTreeForActorNode(RootActorNode);
+
+		AActor* PreviewActorInstance = PreviewActor.Get();
+		if (PreviewActorInstance != nullptr && !GetDefault<UBlueprintEditorSettings>()->bHideConstructionScriptComponentsInDetailsView)
+		{
+			TInlineComponentArray<UActorComponent*> Components;
+			PreviewActorInstance->GetComponents(Components);
+
+			for (UActorComponent* Component : Components)
+			{
+				if (Component->CreationMethod == EComponentCreationMethod::UserConstructionScript)
 				{
-					AddTreeNodeFromComponent(ActorComp, FindOrCreateParentForExistingComponent(ActorComp, GetActorNode()));
+					AddTreeNodeFromComponent(Component, FindOrCreateParentForExistingComponent(Component, RootActorNode));
 				}
 			}
 		}
@@ -4791,7 +4964,7 @@ void SSCSEditor::UpdateTree(bool bRegenerateTreeNodes)
 			{
 				if (SelectedTreeNodes[i]->GetNodeType() == FSCSEditorTreeNode::RootActorNode)
 				{
-					SCSTreeWidget->SetItemSelection(ActorTreeNode, true);
+					SCSTreeWidget->SetItemSelection(RootActorNode, true);
 				}
 				else
 				{
@@ -4943,11 +5116,16 @@ void SSCSEditor::DumpTree()
 
 				case FSCSEditorTreeNode::ENodeType::ComponentNode:
 					NodeLabel = Node->GetDisplayString();
-					if (Node->IsInherited())
+					if (Node->IsInheritedComponent())
 					{
 						NodeLabel += TEXT(" (Inherited)");
 					}
 					break;
+
+				case FSCSEditorTreeNode::ENodeType::ChildActorNode:
+					NodeLabel = Node->GetDisplayString();
+					NodeLabel += TEXT(" [CHILD ACTOR]");
+					break;
 				}
 
 				UE_LOG(LogSCSEditor, Log, TEXT(" %s%s%s"), *Prefix, *NodePrefix, *NodeLabel);
@@ -5363,22 +5541,22 @@ FSCSEditorTreeNodePtrType SSCSEditor::FindOrCreateParentForExistingComponent(UAc
 	if (SceneComponent == nullptr)
 	{
 		check(ActorRootNode.IsValid());
-		check(ActorRootNode->GetNodeType() == FSCSEditorTreeNode::RootActorNode);
+		check(ActorRootNode->IsActorNode());
 		return ActorRootNode;
 	}
 
 	FSCSEditorTreeNodePtrType ParentNodePtr;
 	if (SceneComponent->GetAttachParent() != nullptr
-		&& (EditorMode != EComponentEditorMode::ActorInstance || SceneComponent->GetAttachParent()->GetOwner() == GetActorContext()))
+		&& (EditorMode != EComponentEditorMode::ActorInstance || SceneComponent->GetAttachParent()->GetOwner() == ActorRootNode->GetObject<AActor>()))
 	{
 		// Attempt to find the parent node in the current tree
-		ParentNodePtr = FindTreeNode(SceneComponent->GetAttachParent());
+		ParentNodePtr = FindTreeNode(SceneComponent->GetAttachParent(), ActorRootNode);
 		if (!ParentNodePtr.IsValid())
-{
+		{
 			// If the actual attach parent wasn't found, attempt to find its archetype.
 			// This handles the BP editor case where we might add UCS component nodes taken
 			// from the preview actor instance, which are not themselves template objects.
-			ParentNodePtr = FindTreeNode(Cast<USceneComponent>(SceneComponent->GetAttachParent()->GetArchetype()));
+			ParentNodePtr = FindTreeNode(Cast<USceneComponent>(SceneComponent->GetAttachParent()->GetArchetype()), ActorRootNode);
 			if (!ParentNodePtr.IsValid())
 			{
 				// Recursively add the parent node to the tree if it does not exist yet
@@ -5411,13 +5589,19 @@ FSCSEditorTreeNodePtrType SSCSEditor::FindParentForNewComponent(UActorComponent*
 		TargetParentNode = SelectedTreeNodes[0];
 	}
 
+	// If the current selection belongs to a child actor template, move the target to its outer component node.
+	while (TargetParentNode.IsValid() && FChildActorComponentEditorUtils::IsChildActorSubtreeNode(TargetParentNode))
+	{
+		TargetParentNode = FChildActorComponentEditorUtils::GetOuterChildActorComponentNode(TargetParentNode);
+	}
+
 	if (USceneComponent* NewSceneComponent = Cast<USceneComponent>(NewComponent))
 	{
 		if (TargetParentNode.IsValid())
 		{
-			if (TargetParentNode->GetNodeType() == FSCSEditorTreeNode::RootActorNode)
+			if (TargetParentNode->IsActorNode())
 			{
-				FSCSEditorActorNodePtrType TargetActorNode = StaticCastSharedPtr<FSCSEditorTreeNodeRootActor>(TargetParentNode);
+				FSCSEditorActorNodePtrType TargetActorNode = StaticCastSharedPtr<FSCSEditorTreeNodeActorBase>(TargetParentNode);
 				if (TargetActorNode.IsValid())
 				{
 					FSCSEditorTreeNodePtrType TargetSceneRootNode = TargetActorNode->GetSceneRootNode();
@@ -5432,7 +5616,7 @@ FSCSEditorTreeNodePtrType SSCSEditor::FindParentForNewComponent(UActorComponent*
 					}
 				}
 			}
-			else if(TargetParentNode->GetNodeType() == FSCSEditorTreeNode::ComponentNode)
+			else if(TargetParentNode->IsComponentNode())
 			{
 				USceneComponent* CastTargetToSceneComponent = Cast<USceneComponent>(TargetParentNode->GetComponentTemplate());
 				if (CastTargetToSceneComponent == nullptr || !NewSceneComponent->CanAttachAsChild(CastTargetToSceneComponent, NAME_None))
@@ -5450,7 +5634,7 @@ FSCSEditorTreeNodePtrType SSCSEditor::FindParentForNewComponent(UActorComponent*
 	{
 		if (TargetParentNode.IsValid())
 		{
-			while (TargetParentNode->GetNodeType() != FSCSEditorTreeNode::RootActorNode)
+			while (!TargetParentNode->IsActorNode())
 			{
 				TargetParentNode = TargetParentNode->GetParent();
 			}
@@ -5460,7 +5644,7 @@ FSCSEditorTreeNodePtrType SSCSEditor::FindParentForNewComponent(UActorComponent*
 			TargetParentNode = GetActorNode();
 		}
 
-		check(TargetParentNode.IsValid() && TargetParentNode->GetNodeType() == FSCSEditorTreeNode::RootActorNode);
+		check(TargetParentNode.IsValid() && TargetParentNode->IsActorNode());
 	}
 
 	return TargetParentNode;
@@ -5815,7 +5999,19 @@ bool SSCSEditor::CanDeleteNodes() const
 	TArray<FSCSEditorTreeNodePtrType> SelectedNodes = SCSTreeWidget->GetSelectedItems();
 	for (int32 i = 0; i < SelectedNodes.Num(); ++i)
 	{
-		if (!SelectedNodes[i]->CanDelete()) {return false;}
+		if (!SelectedNodes[i]->CanDelete())
+		{
+			return false;
+		}
+		else
+		{
+			// Don't allow nodes that belong to a child actor template to be deleted
+			const bool bIsChildActorSubtreeNode = FChildActorComponentEditorUtils::IsChildActorSubtreeNode(SelectedNodes[i]);
+			if (bIsChildActorSubtreeNode)
+			{
+				return false;
+			}
+		}
 	}
 	return SelectedNodes.Num() > 0;
 }
@@ -6120,6 +6316,9 @@ FSCSEditorTreeNodePtrType SSCSEditor::AddTreeNode(USCS_Node* InSCSNode, FSCSEdit
 	// Expand parent nodes by default
 	SCSTreeWidget->SetItemExpansion(InParentNodePtr, true);
 
+	// Add this node's child actor node, if present
+	AddTreeNodeFromChildActor(NewNodePtr);
+
 	// Recursively add the given SCS node's child nodes
 	for (USCS_Node* ChildNode : InSCSNode->GetChildNodes())
 	{
@@ -6142,11 +6341,34 @@ FSCSEditorTreeNodePtrType SSCSEditor::AddTreeNodeFromComponent(UActorComponent*
 		RefreshFilteredState(NewNodePtr, false);
 	}
 
+	AddTreeNodeFromChildActor(NewNodePtr);
+
 	SCSTreeWidget->SetItemExpansion(NewNodePtr, true);
 
 	return NewNodePtr;
 }
 
+FSCSEditorTreeNodePtrType SSCSEditor::AddTreeNodeFromChildActor(FSCSEditorTreeNodePtrType InNodePtr)
+{
+	if (!InNodePtr.IsValid())
+	{
+		return nullptr;
+	}
+
+	FSCSEditorChildActorNodePtrType ChildActorNodePtr = InNodePtr->GetChildActorNode();
+	if (ChildActorNodePtr.IsValid() && ChildActorNodePtr->ShouldShowInTreeView() && ChildActorNodePtr->GetParent() != InNodePtr)
+	{
+		InNodePtr->AddChild(ChildActorNodePtr);
+		RefreshFilteredState(ChildActorNodePtr, false);
+		BuildSubTreeForActorNode(ChildActorNodePtr);
+
+		// Expand the given parent to reveal the child actor node
+		SCSTreeWidget->SetItemExpansion(InNodePtr, true);
+	}
+
+	return ChildActorNodePtr;
+}
+
 FSCSEditorTreeNodePtrType SSCSEditor::FindTreeNode(const USCS_Node* InSCSNode, FSCSEditorTreeNodePtrType InStartNodePtr) const
 {
 	FSCSEditorTreeNodePtrType NodePtr;
@@ -6325,7 +6547,14 @@ void SSCSEditor::OnPostTick(float)
 
 bool SSCSEditor::CanRenameComponent() const
 {
-	return IsEditingAllowed() && SCSTreeWidget->GetSelectedItems().Num() == 1 && SCSTreeWidget->GetSelectedItems()[0]->CanRename();
+	if (!IsEditingAllowed())
+	{
+		return false;
+	}
+
+	// In addition to certain node types, don't allow nodes within a child actor template's hierarchy to be renamed
+	TArray<FSCSEditorTreeNodePtrType> SelectedItems = SCSTreeWidget->GetSelectedItems();
+	return SelectedItems.Num() == 1 && SelectedItems[0]->CanRename() && !FChildActorComponentEditorUtils::IsChildActorSubtreeNode(SelectedItems[0]);
 }
 
 void SSCSEditor::DepthFirstTraversal(const FSCSEditorTreeNodePtrType& InNodePtr, TSet<FSCSEditorTreeNodePtrType>& OutVisitedNodes, const TFunctionRef<void(const FSCSEditorTreeNodePtrType&)> InFunction) const
@@ -6497,7 +6726,7 @@ struct FRestoreSelectedInstanceComponent
 
 		check(GEditor);
 		TArray<UActorComponent*> ComponentsToSaveAndDelesect;
-		for (auto Iter = GEditor->GetSelectedComponentIterator(); Iter; ++Iter)
+		for (FSelectionIterator Iter = GEditor->GetSelectedComponentIterator(); Iter; ++Iter)
 		{
 			UActorComponent* Component = CastChecked<UActorComponent>(*Iter, ECastCheckedType::NullAllowed);
 			if (Component && InActor->GetInstanceComponents().Contains(Component))
@@ -6815,9 +7044,11 @@ bool SSCSEditor::RefreshFilteredState(FSCSEditorTreeNodePtrType TreeNode, bool b
 					bIsFilteredOut = true;
 				}
 			}
+
 			// if we're not recursing, then assume this is for a new node and we need to update the parent
 			// otherwise, assume the parent was hit as part of the recursion
-			TreeNodeIn->UpdateCachedFilterState(!bIsFilteredOut, /*bUpdateParent =*/!bRecursiveIn);
+			bool bUpdateParent = !bRecursiveIn;
+			TreeNodeIn->UpdateCachedFilterState(!bIsFilteredOut, bUpdateParent);
 		}
 	};
 
diff --git a/Engine/Source/Editor/Kismet/Public/SSCSEditor.h b/Engine/Source/Editor/Kismet/Public/SSCSEditor.h
--- a/Engine/Source/Editor/Kismet/Public/SSCSEditor.h
+++ b/Engine/Source/Editor/Kismet/Public/SSCSEditor.h
@@ -31,14 +31,15 @@ class SSCSEditor;
 class UPrimitiveComponent;
 struct EventData;
 
-// SCS tree node pointer type
+// SCS editor tree node pointer types
 using FSCSEditorTreeNodePtrType = TSharedPtr<class FSCSEditorTreeNode>;
-using FSCSEditorActorNodePtrType = TSharedPtr<class FSCSEditorTreeNodeRootActor>;
+using FSCSEditorActorNodePtrType = TSharedPtr<class FSCSEditorTreeNodeActorBase>;
+using FSCSEditorChildActorNodePtrType = TSharedPtr<class FSCSEditorTreeNodeChildActor>;
 
 /**
  * FSCSEditorTreeNode
  *
- * Wrapper class for component template nodes displayed in the SCS (Simple Construction Script) editor tree widget.
+ * Wrapper class for nodes displayed in the SCS (Simple Construction Script) editor tree widget.
  */
 class KISMET_API FSCSEditorTreeNode : public TSharedFromThis<FSCSEditorTreeNode>
 {
@@ -51,6 +52,7 @@ public:
 		ComponentNode,
 		RootActorNode,
 		SeparatorNode,
+		ChildActorNode,
 	};
 
 	/**
@@ -76,15 +78,19 @@ public:
 	*/
 	virtual FText GetDisplayName() const;
 	/**
-	 * @return The SCS node that is represented by this object, or NULL if there is no SCS node associated with the component template.
+	 * @return The SCS node that is represented by this object, or NULL if there is no associated SCS node.
 	 */
 	virtual class USCS_Node* GetSCSNode() const;
 	/**
 	 * @param ActualEditedBlueprint currently edited blueprint
-	 *
+	 * @note Derived classes should override GetOrCreateEditableObjectForBlueprint().
 	 * @return The component template that can be editable for actual class.
 	 */
-	virtual UActorComponent* GetOrCreateEditableComponentTemplate(UBlueprint* ActualEditedBlueprint) const;
+	inline UActorComponent* GetOrCreateEditableComponentTemplate(UBlueprint* ActualEditedBlueprint) const
+	{
+		// @TODO - Deprecate this public API in favor of GetEditableObjectForBlueprint().
+		return Cast<UActorComponent>(GetOrCreateEditableObjectForBlueprint(ActualEditedBlueprint));
+	}
 	/**
 	 * Finds the component instance represented by this node contained within a given Actor instance.
 	 *
@@ -100,10 +106,28 @@ public:
 	 * @return The set of nodes which are parented to this node (read-only).
 	 */
 	const TArray<FSCSEditorTreeNodePtrType>& GetChildren() const { return Children; }
+	/**
+	 * @return The root of the actor subtree to which this node belongs.
+	 */
+	FSCSEditorActorNodePtrType GetActorRootNode() const { return ActorRootNodePtr; }
 	/**
 	 * @return Type of node
 	 */
 	ENodeType GetNodeType() const;
+	/**
+	 * @return True if this represents an actor node type
+	 */
+	bool IsActorNode() const
+	{
+		return NodeType == ENodeType::RootActorNode || NodeType == ENodeType::ChildActorNode;
+	}
+	/**
+	 * @return True if this represents a component node type
+	 */
+	bool IsComponentNode() const
+	{
+		return NodeType == ENodeType::ComponentNode;
+	}
 	/**
 	 * @param	bEvenIfPendingKill	If false, nullptr will be returned if the cached component template is pending kill.
 	 *								If true, it will be returned regardless (this is used for recaching the component template if the objects
@@ -113,11 +137,42 @@ public:
 	 * @warning This will not return the right component for components overridden by the inherited component handler, you need to call GetOrCreateEditableComponentTemplate instead
 	 * @return	The component template or instance represented by this node, if it's a component node.
 	 */
-	UActorComponent* GetComponentTemplate(bool bEvenIfPendingKill = false) const;
+	inline UActorComponent* GetComponentTemplate(bool bEvenIfPendingKill = false) const
+	{
+		// @todo - Deprecate this API? For backwards-compatibility, this continues to provide non-const access to the internal object instance specifically as a component reference.
+		return Cast<UActorComponent>(WeakObjectPtr.Get(bEvenIfPendingKill));
+	}
 	/**
-	 * Set the component template represented by this node, if it's a component node.
+	 * @param	bEvenIfPendingKill	If false, nullptr will be returned if the cached object instance is pending kill.
+	 *								If true, it will be returned regardless (this is used for recaching the object if the objects
+	 *								have been reinstanced following construction script execution).
+	 *
+	 * @note	Deliberately non-virtual, for performance reasons.
+	 * @return	A read-only reference to the object represented by this node.
+	 */
+	template<class T>
+	inline const T* GetObject(bool bEvenIfPendingKill = false) const
+	{
+		return Cast<T>(WeakObjectPtr.Get(bEvenIfPendingKill));
+	}
+	/**
+	 * @param	InBlueprint			The Blueprint in which the object will be edited.
+	 *
+	 * @note	May not be the same as the value returned by GetObject().
+	 * @return	A reference to the object represented by this node that can be modified within the given Blueprint.
 	 */
-	void SetComponentTemplate(UActorComponent* Component);
+	template<class T>
+	inline T* GetEditableObjectForBlueprint(UBlueprint* InBlueprint) const
+	{
+		return Cast<T>(GetOrCreateEditableObjectForBlueprint(InBlueprint));
+	}
+	/**
+	 * Sets the internal object instance represented by this node.
+	 */
+	inline void SetObject(UObject* InObject)
+	{
+		WeakObjectPtr = InObject;
+	}
 	/**
 	 * @return Whether or not this node is a direct child of the given node.
 	 */
@@ -201,8 +256,11 @@ public:
 		return Cast<USceneComponent>(GetComponentTemplate()) != nullptr;
 	}
 
+	/** Returns the associated child actor node if applicable to this node type. */
+	virtual FSCSEditorChildActorNodePtrType GetChildActorNode() { return nullptr; }
+
 	// Tries to find a SCS node that was likely responsible for creating the specified instance component.  Note: This is not always possible to do!
-	static USCS_Node* FindSCSNodeForInstance(UActorComponent* InstanceComponent, UClass* ClassToSearch);
+	static USCS_Node* FindSCSNodeForInstance(const UActorComponent* InstanceComponent, UClass* ClassToSearch);
 
 	/**
 	 * Creates the correct type of node based on the component (instanced or not, etc...)
@@ -211,24 +269,31 @@ public:
 
 	// Destructor
 	virtual ~FSCSEditorTreeNode() {}
+
 protected:
 	// Called when this node is being removed via a RemoveChild call
 	virtual void RemoveMeAsChild() {}
 
+	// Provides derived classes with non-const access to the object represented by this node (e.g. for rename operations, etc.). This should not be made public.
+	template<class T>
+	inline T* GetMutableObject() const
+	{
+		return Cast<T>(WeakObjectPtr.Get());
+	}
 
-public:
-
-
+	// Derived classes can override to create and/or return a reference to an alternate editable object.
+	virtual UObject* GetOrCreateEditableObjectForBlueprint(UBlueprint* InBlueprint) const;
 
+public:
 	/**
-	* @return The Blueprint to which this node belongs.
+	* @return The Blueprint to which the object represented by this node belongs (requires implementation in subclass).
 	*/
-	UBlueprint* GetBlueprint() const;
+	virtual UBlueprint* GetBlueprint() const = 0;
 
 	/**
 	 * @return Whether or not this object represents a "native" component template (i.e. one that is not found in the SCS tree).
 	 */
-	virtual bool IsNative() const { return false; }
+	virtual bool IsNativeComponent() const { return false; }
 
 	/**
 	 * @return Whether or not this object represents a root component.
@@ -238,25 +303,49 @@ public:
 	/**
 	 * @return Whether or not this object represents an inherited SCS node (one from a SCS node in a parent Blueprint).
 	 */
-	virtual bool IsInheritedSCS() const { return false; }
+	virtual bool IsInheritedSCSNode() const { return false; }
 
 	/**
 	 * @return Whether or not this object was declared in the current class (or instance).  Anything inherited cannot be reorganized (renamed, reparented, etc...).
 	 */
-	virtual bool IsInherited() const
+	virtual bool IsInheritedComponent() const
 	{
-		return IsNative() || IsInheritedSCS();
+		return IsNativeComponent() || IsInheritedSCSNode();
 	}
 
 	/**
-	 * @return Whether or not this object represents a component instance rather than a template.
+	 * @return Whether or not this node represents an instanced object (i.e. not a template).
 	 */
 	virtual bool IsInstanced() const { return false; }
 
+	/**
+	 * @return Whether or not this node represents an instanced actor object.
+	 */
+	bool IsInstancedActor() const
+	{
+		return IsInstanced() && IsActorNode();
+	}
+
+	/**
+	 * @return Whether or not this node represents an instanced component object.
+	 */
+	bool IsInstancedComponent() const
+	{
+		return IsInstanced() && IsComponentNode();
+	}
+
+	/**
+	 * @return Whether or not this node represents a Blueprint (i.e. non-native) component object.
+	 */
+	bool IsBlueprintComponent() const
+	{
+		return IsComponentNode() && !IsNativeComponent();
+	}
+
 	/**
 	 * @return Whether or not this object represents a component instance that was created by the user and not by a native or Blueprint-generated class.
 	 */
-	virtual bool IsUserInstanced() const { return false; }
+	virtual bool IsUserInstancedComponent() const { return false; }
 
 	/**
 	* @return Whether or not this object represents the default SCS scene root component.
@@ -274,25 +363,25 @@ public:
 	virtual bool CanReparent() const { return false; }
 
 	/**
-	 * @return Whether or not we can edit default properties for the component template represented by this object.
+	 * @return Whether or not we can edit properties for the object represented by this node.
 	 */
-	virtual bool CanEditDefaults() const { return false; }
+	virtual bool CanEdit() const { return false; }
 
 	/**
-	 * @return Whether or not this object represents a node that can be renamed from the components tree.
+	 * @return Whether or not we can rename the object or variable represented by this node.
 	 */
 	virtual bool CanRename() const { return false; }
 
 	/**
-	 * Requests a rename on the component.
+	 * Requests a rename on the node.
 	 * @param OngoingCreateTransaction The transaction scoping the node creation which will end once the node is named by the user or null if the rename is not part of a the creation process.
 	 */
 	void OnRequestRename(TUniquePtr<FScopedTransaction> OngoingCreateTransaction);
 
-	/** Renames the component */
+	/** Renames the object or variable represented by this node */
 	virtual void OnCompleteRename(const FText& InNewName);
 
-	/** Sets up the delegate for renaming a component */
+	/** Sets up the delegate for a rename operation */
 	void SetRenameRequestedDelegate(FOnRenameRequested InRenameRequested) { RenameRequestedDelegate = InRenameRequested; }
 
 	/** Query that determines if this item should be filtered out or not */
@@ -311,18 +400,19 @@ protected:
 	/** Used to update the EFilteredState::ChildMatches flag for parent nodes, when this item's filtration state has changed */
 	void ApplyFilteredStateToParent();
 	
-	// Scope the creation of a component which ends when the initial component 'name' is given/accepted by the user, which can be several frames after the component was actually created.
+	// Scope the creation of a node which ends when the initial 'name' is given/accepted by the user, which can be several frames after the node was actually created.
 	TUniquePtr<FScopedTransaction> OngoingCreateTransaction;
 
-	// Component template represented by this node, if it's a component node, otherwise invalid
-	TWeakObjectPtr<UActorComponent> ComponentTemplatePtr;
-
 private:
-	// The type of component tree node
+	// The type of tree node
 	ENodeType NodeType;
 
+	// Weak ptr to the object instance represented by this node (e.g. component template)
+	TWeakObjectPtr<UObject> WeakObjectPtr;
+
 	// Actual tree structure
 	FSCSEditorTreeNodePtrType ParentNodePtr;
+	FSCSEditorActorNodePtrType ActorRootNodePtr;
 	TArray<FSCSEditorTreeNodePtrType> Children;
 
 	/** Handles rename requests */
@@ -338,6 +428,25 @@ private:
 		Unknown = 0xFC // ~FilteredInMask
 	};
 	uint8 FilterFlags;
+
+public:
+	UE_DEPRECATED(4.26, "Use IsNativeComponent() instead.")
+	bool IsNative() const { return IsNativeComponent(); }
+
+	UE_DEPRECATED(4.26, "Use IsInheritedSCSNode() instead.")
+	bool IsInheritedSCS() const { return IsInheritedSCSNode(); }
+
+	UE_DEPRECATED(4.26, "Use IsInheritedComponent() instead.")
+	bool IsInherited() const { return IsInheritedComponent(); }
+
+	UE_DEPRECATED(4.26, "Use IsUserInstancedComponent() instead.")
+	bool IsUserInstanced() const { return IsUserInstancedComponent(); }
+
+	UE_DEPRECATED(4.26, "Use CanEdit() instead.")
+	bool CanEditDefaults() const { return CanEdit(); }
+
+	UE_DEPRECATED(4.26, "Use SetObject() instead.")
+	void SetComponentTemplate(UActorComponent* Component) { SetObject(Component); }
 };
 
 //////////////////////////////////////////////////////////////////////////
@@ -355,10 +464,16 @@ public:
 	// FSCSEditorTreeNode interface
 	virtual FName GetVariableName() const override;
 	virtual FString GetDisplayString() const override;
-	virtual bool CanRename() const override { return !IsInherited() && !IsDefaultSceneRoot(); }
-	virtual bool CanDelete() const override { return !IsInherited() && !IsDefaultSceneRoot(); }
-	virtual bool CanReparent() const override { return !IsInherited() && !IsDefaultSceneRoot() && IsSceneComponent(); }
+	virtual bool CanRename() const override { return !IsInheritedComponent() && !IsDefaultSceneRoot(); }
+	virtual bool CanDelete() const override { return !IsInheritedComponent() && !IsDefaultSceneRoot(); }
+	virtual bool CanReparent() const override;
+	virtual UBlueprint* GetBlueprint() const override;
+	virtual FSCSEditorChildActorNodePtrType GetChildActorNode();
 	// End of FSCSEditorTreeNode interface
+
+private:
+	/** Child actor node associated with this component node, if applicable. */
+	FSCSEditorChildActorNodePtrType ChildActorNodePtr;
 };
 
 
@@ -376,18 +491,14 @@ public:
 	FSCSEditorTreeNodeInstancedInheritedComponent(AActor* Owner, UActorComponent* InComponentTemplate);
 
 	// FSCSEditorTreeNode public interface
-	virtual bool IsNative() const override;
+	virtual bool IsNativeComponent() const override;
 	virtual bool IsRootComponent() const override;
-	virtual bool IsInheritedSCS() const override;
+	virtual bool IsInheritedSCSNode() const override;
 	virtual bool IsInstanced() const override { return true; }
-	virtual bool IsInherited() const override { return true; }
-	virtual bool IsUserInstanced() const override { return false; }
+	virtual bool IsInheritedComponent() const override { return true; }
 	virtual bool IsDefaultSceneRoot() const override;
-	virtual bool CanEditDefaults() const override;
-	//virtual FName GetVariableName() const override;
-	//virtual FString GetDisplayString() const override;
+	virtual bool CanEdit() const override;
 	virtual FText GetDisplayName() const override;
-	virtual UActorComponent* GetOrCreateEditableComponentTemplate(UBlueprint* ActualEditedBlueprint) const override;
 	// End of FSCSEditorTreeNode public interface
 
 private:
@@ -408,17 +519,14 @@ public:
 	FSCSEditorTreeNodeInstanceAddedComponent(AActor* Owner, UActorComponent* InComponentTemplate);
 
 	// FSCSEditorTreeNode public interface
-	virtual bool IsNative() const override { return false; }
 	virtual bool IsRootComponent() const override;
-	virtual bool IsInheritedSCS() const override { return false; }
 	virtual bool IsInstanced() const override { return true; }
-	virtual bool IsUserInstanced() const override { return true; }
+	virtual bool IsUserInstancedComponent() const override { return true; }
 	virtual bool IsDefaultSceneRoot() const override;
-	virtual bool CanEditDefaults() const override { return true; }
+	virtual bool CanEdit() const override { return true; }
 	virtual FName GetVariableName() const override { return NAME_None; }
 	virtual FString GetDisplayString() const override;
 	virtual FText GetDisplayName() const override;
-	virtual UActorComponent* GetOrCreateEditableComponentTemplate(UBlueprint* ActualEditedBlueprint) const override;
 	virtual void OnCompleteRename(const FText& InNewName) override;
 	// End of FSCSEditorTreeNode public interface
 
@@ -455,28 +563,24 @@ public:
 
 
 	// FSCSEditorTreeNode public interface
-	virtual bool IsNative() const override;
+	virtual bool IsNativeComponent() const override;
 	virtual bool IsRootComponent() const override;
-	virtual bool IsInheritedSCS() const override;
-	virtual bool IsInstanced() const override { return false; }
-	virtual bool IsUserInstanced() const override { return false; }
+	virtual bool IsInheritedSCSNode() const override;
 	virtual bool IsDefaultSceneRoot() const override;
-	virtual bool CanEditDefaults() const override;
-	//virtual FName GetVariableName() const override;
-	//virtual FString GetDisplayString() const override;
+	virtual bool CanEdit() const override;
 	virtual FText GetDisplayName() const override;
 	virtual class USCS_Node* GetSCSNode() const override;
-	virtual UActorComponent* GetOrCreateEditableComponentTemplate(UBlueprint* ActualEditedBlueprint) const override;
 	virtual void OnCompleteRename(const FText& InNewName) override;
 	// End of FSCSEditorTreeNode public interface
 
 protected:
 	// FSCSEditorTreeNode protected interface
 	virtual void RemoveMeAsChild() override;
+	virtual UObject* GetOrCreateEditableObjectForBlueprint(UBlueprint* InBlueprint) const override;
 	// End of FSCSEditorTreeNode protected interface
 
 	/** Get overridden template component, specialized in given blueprint */
-	UActorComponent* INTERNAL_GetOverridenComponentTemplate(UBlueprint* Blueprint, bool bCreateIfNecessary) const;
+	UActorComponent* INTERNAL_GetOverridenComponentTemplate(UBlueprint* Blueprint) const;
 
 private:
 	// Was this component inherited from a parent class or introduced in this class?
@@ -486,14 +590,13 @@ private:
 	TWeakObjectPtr<class USCS_Node> SCSNodePtr;
 };
 
-class KISMET_API FSCSEditorTreeNodeRootActor : public FSCSEditorTreeNode
+class KISMET_API FSCSEditorTreeNodeActorBase : public FSCSEditorTreeNode
 {
 public:
-	FSCSEditorTreeNodeRootActor(AActor* InActor, bool bInAllowRename)
-		: FSCSEditorTreeNode(FSCSEditorTreeNode::RootActorNode)
-		, Actor(InActor)
-		, bAllowRename(bInAllowRename)
+	FSCSEditorTreeNodeActorBase(FSCSEditorTreeNode::ENodeType InNodeType, AActor* InActor)
+		: FSCSEditorTreeNode(InNodeType)
 	{
+		SetObject(InActor);
 	}
 
 	FSCSEditorTreeNodePtrType GetSceneRootNode() const;
@@ -504,26 +607,64 @@ public:
 
 	// FSCSEditorTreeNode public interface
 	virtual FName GetNodeID() const override;
-	virtual bool CanRename() const override { return bAllowRename; }
-	virtual void OnCompleteRename(const FText& InNewName) override;
+	virtual bool IsInstanced() const override;
+	virtual bool CanEdit() const override { return true; }
 	virtual void AddChild(FSCSEditorTreeNodePtrType InChildNodePtr) override;
 	virtual void RemoveChild(FSCSEditorTreeNodePtrType InChildNodePtr) override;
+	virtual UBlueprint* GetBlueprint() const override;
 	// End of FSCSEditorTreeNode public interface
+
 protected:
 	using Super = FSCSEditorTreeNode;
 
 private:
-	
-	AActor* Actor;
-	bool bAllowRename;
-
 	FSCSEditorTreeNodePtrType SceneRootNodePtr;
 	/** Root set of components (contains the root scene component and any non-scene component nodes) */
 	TArray<FSCSEditorTreeNodePtrType> ComponentNodes;
+};
+
+class KISMET_API FSCSEditorTreeNodeRootActor : public FSCSEditorTreeNodeActorBase
+{
+public:
+	FSCSEditorTreeNodeRootActor(AActor* InActor, bool bInAllowRename)
+		: FSCSEditorTreeNodeActorBase(FSCSEditorTreeNode::RootActorNode, InActor)
+		, bAllowRename(bInAllowRename)
+	{
+	}
+
+	// FSCSEditorTreeNode public interface
+	virtual bool CanRename() const override { return bAllowRename; }
+	virtual void OnCompleteRename(const FText& InNewName) override;
+	virtual void AddChild(FSCSEditorTreeNodePtrType InChildNodePtr) override;
+	virtual void RemoveChild(FSCSEditorTreeNodePtrType InChildNodePtr) override;
+	// End of FSCSEditorTreeNode public interface
+
+private:
+	bool bAllowRename;
 	FSCSEditorTreeNodePtrType SceneComponentSeparatorNodePtr;
 	FSCSEditorTreeNodePtrType NonSceneComponentSeparatorNodePtr;
 };
 
+class KISMET_API FSCSEditorTreeNodeChildActor : public FSCSEditorTreeNodeActorBase
+{
+public:
+	FSCSEditorTreeNodeChildActor(FSCSEditorTreeNodePtrType InOwner, AActor* InActor)
+		: FSCSEditorTreeNodeActorBase(FSCSEditorTreeNode::ChildActorNode, InActor)
+		, OwnerNode(InOwner)
+	{
+	}
+
+	// FSCSEditorTreeNode public interface
+	virtual bool IsFlaggedForFiltration() const override;
+	// End of FSCSEditorTreeNode public interface
+
+	bool ShouldShowInTreeView() const;
+	UChildActorComponent* GetChildActorComponent() const;
+
+private:
+	FSCSEditorTreeNodePtrType OwnerNode;
+};
+
 class KISMET_API FSCSEditorTreeNodeSeparator : public FSCSEditorTreeNode
 {
 public:
@@ -532,7 +673,10 @@ public:
 	{
 	}
 
+	// FSCSEditorTreeNode public interface
+	virtual UBlueprint* GetBlueprint() const override { return nullptr; }
 	virtual bool IsFlaggedForFiltration() const override { return false; }
+	// End of FSCSEditorTreeNode public interface
 };
 
 //////////////////////////////////////////////////////////////////////////
@@ -577,6 +721,8 @@ public:
 protected:
 	virtual ESelectionMode::Type GetSelectionMode() const override;
 
+	virtual const FSlateBrush* GetIconBrush() const;
+
 	static void AddToToolTipInfoBox(const TSharedRef<SVerticalBox>& InfoBox, const FText& Key, TSharedRef<SWidget> ValueIcon, const TAttribute<FText>& Value, bool bImportant);
 
 	/** Commits the new name of the component */
@@ -664,6 +810,13 @@ public:
 	virtual TSharedRef<SWidget> GenerateWidgetForColumn(const FName& ColumnName) override;
 	// End of SMultiColumnTableRow<T>
 
+protected:
+	/** Gets the associated actor node */
+	FSCSEditorActorNodePtrType GetActorNode() const;
+
+	/** Data accessors */
+	virtual const FSlateBrush* GetIconBrush() const override;
+
 private:
 	/** Creates a tooltip for this row */
 	TSharedRef<SToolTip> CreateToolTipWidget() const;
@@ -672,7 +825,6 @@ private:
 	bool OnVerifyActorLabelChanged(const FText& InLabel, FText& OutErrorMessage);
 
 	/** Data accessors */
-	const FSlateBrush* GetActorIcon() const;
 	FText GetActorDisplayText() const;
 	FText GetActorContextText() const;
 	FText GetActorClassNameText() const;
@@ -1046,6 +1198,9 @@ protected:
 	/** Called when component objects are replaced following construction script execution */
 	void OnObjectsReplaced(const TMap<UObject*, UObject*>& OldToNewInstanceMap);
 
+	/** Helper method to update component pointers held by the given actor node's subtree */
+	void ReplaceComponentReferencesInTree(FSCSEditorActorNodePtrType InActorNode, const TMap<UObject*, UObject*>& OldToNewInstanceMap);
+
 	/** Update component pointers held by tree nodes if components have been replaced following construction script execution */
 	void ReplaceComponentReferencesInTree(const TArray<FSCSEditorTreeNodePtrType>& Nodes, const TMap<UObject*, UObject*>& OldToNewInstanceMap);
 
@@ -1080,6 +1235,9 @@ protected:
 	/** Helper method to add a tree node for the given actor component */
 	FSCSEditorTreeNodePtrType AddTreeNodeFromComponent(UActorComponent* InSceneComponent, FSCSEditorTreeNodePtrType InParentTreeNode = FSCSEditorTreeNodePtrType());
 	
+	/** Helper method to add a tree node for the given node's child actor node, if present */
+	FSCSEditorTreeNodePtrType AddTreeNodeFromChildActor(FSCSEditorTreeNodePtrType InNodePtr);
+	
 	/** Helper method to recursively find a tree node for the given SCS node starting at the given tree node */
 	FSCSEditorTreeNodePtrType FindTreeNode(const USCS_Node* InSCSNode, FSCSEditorTreeNodePtrType InStartNodePtr = FSCSEditorTreeNodePtrType()) const;
 
@@ -1157,6 +1315,9 @@ protected:
 	 */
 	bool RefreshFilteredState(FSCSEditorTreeNodePtrType TreeNode, bool bRecursive);
 
+	/** Helper method to construct the subtree for the given actor (root) node. */
+	void BuildSubTreeForActorNode(FSCSEditorActorNodePtrType InActorNode);
+
 public:
 	/** Tree widget */
 	TSharedPtr<SSCSTreeType> SCSTreeWidget;
diff --git a/Engine/Source/Editor/Kismet/Public/SSCSEditorMenuContext.h b/Engine/Source/Editor/Kismet/Public/SSCSEditorMenuContext.h
--- a/Engine/Source/Editor/Kismet/Public/SSCSEditorMenuContext.h
+++ b/Engine/Source/Editor/Kismet/Public/SSCSEditorMenuContext.h
@@ -9,7 +9,7 @@
 class SSCSEditor;
 
 UCLASS()
-class USSCSEditorMenuContext : public UObject
+class KISMET_API USSCSEditorMenuContext : public UObject
 {
 	GENERATED_BODY()
 public:
diff --git a/Engine/Source/Editor/LevelEditor/Private/SActorDetails.cpp b/Engine/Source/Editor/LevelEditor/Private/SActorDetails.cpp
--- a/Engine/Source/Editor/LevelEditor/Private/SActorDetails.cpp
+++ b/Engine/Source/Editor/LevelEditor/Private/SActorDetails.cpp
@@ -666,7 +666,7 @@ bool SActorDetails::IsPropertyEditingEnabled() const
 	bool bIsEditable = true;
 	for (const FSCSEditorTreeNodePtrType& Node : SCSEditor->GetSelectedNodes())
 	{
-		bIsEditable = Node->CanEditDefaults() || Node->GetNodeType() == FSCSEditorTreeNode::ENodeType::RootActorNode;
+		bIsEditable = Node->CanEdit();
 		if (!bIsEditable)
 		{
 			break;
@@ -715,7 +715,7 @@ EVisibility SActorDetails::GetUCSComponentWarningVisibility() const
 	// Check to see if any selected components are inherited from blueprint
 	for (const FSCSEditorTreeNodePtrType& Node : SCSEditor->GetSelectedNodes())
 	{
-		if (!Node->IsNative())
+		if (Node->IsBlueprintComponent())
 		{
 			UActorComponent* Component = Node->GetComponentTemplate();
 			bIsUneditableBlueprintComponent = Component ? Component->CreationMethod == EComponentCreationMethod::UserConstructionScript : false;
@@ -760,7 +760,7 @@ EVisibility SActorDetails::GetInheritedBlueprintComponentWarningVisibility() con
 	// Check to see if any selected components are inherited from blueprint
 	for (const FSCSEditorTreeNodePtrType& Node : SCSEditor->GetSelectedNodes())
 	{
-		if (!Node->IsNative())
+		if (Node->IsBlueprintComponent())
 		{
 			if (UActorComponent* Component = Node->GetComponentTemplate())
 			{
@@ -771,7 +771,7 @@ EVisibility SActorDetails::GetInheritedBlueprintComponentWarningVisibility() con
 				}
 			}
 		}
-		else if (!Node->CanEditDefaults() && NotEditableSetByBlueprint(Node->GetComponentTemplate()))
+		else if (!Node->CanEdit() && NotEditableSetByBlueprint(Node->GetComponentTemplate()))
 		{
 			bIsUneditableBlueprintComponent = true;
 			break;
@@ -787,7 +787,7 @@ EVisibility SActorDetails::GetNativeComponentWarningVisibility() const
 	for (const FSCSEditorTreeNodePtrType& Node : SCSEditor->GetSelectedNodes())
 	{
 		// Check to see if the component is native and not editable
-		if (Node->IsNative() && !Node->CanEditDefaults() && !NotEditableSetByBlueprint(Node->GetComponentTemplate()))
+		if (Node->IsNativeComponent() && !Node->CanEdit() && !NotEditableSetByBlueprint(Node->GetComponentTemplate()))
 		{
 			bIsUneditableNative = true;
 			break;
diff --git a/Engine/Source/Editor/UnrealEd/Private/AssetSelection.cpp b/Engine/Source/Editor/UnrealEd/Private/AssetSelection.cpp
--- a/Engine/Source/Editor/UnrealEd/Private/AssetSelection.cpp
+++ b/Engine/Source/Editor/UnrealEd/Private/AssetSelection.cpp
@@ -89,7 +89,7 @@ namespace AssetSelectionUtils
 			// Get the class type of the first actor.
 			AActor* FirstActor = SelectedActors[0];
 
-			if( FirstActor && !FirstActor->HasAnyFlags( RF_ClassDefaultObject ) )
+			if( FirstActor && !FirstActor->IsTemplate() )
 			{
 				UClass* FirstClass = FirstActor->GetClass();
 				UObject* FirstArchetype = FirstActor->GetArchetype();
@@ -102,7 +102,7 @@ namespace AssetSelectionUtils
 				{
 					AActor* CurrentActor = SelectedActors[ ActorIndex ];
 
-					if( CurrentActor->HasAnyFlags( RF_ClassDefaultObject ) )
+					if( CurrentActor->IsTemplate() )
 					{
 						continue;
 					}
diff --git a/Engine/Source/Editor/UnrealEd/Private/Kismet2/ChildActorComponentEditorUtils.cpp b/Engine/Source/Editor/UnrealEd/Private/Kismet2/ChildActorComponentEditorUtils.cpp
new file mode 100644
--- /dev/null
+++ b/Engine/Source/Editor/UnrealEd/Private/Kismet2/ChildActorComponentEditorUtils.cpp
@@ -0,0 +1,245 @@
+// Copyright Epic Games, Inc. All Rights Reserved.
+
+#include "Kismet2/ChildActorComponentEditorUtils.h"
+#include "Settings/EditorProjectSettings.h"
+#include "ToolMenus.h"
+#include "SSCSEditor.h"
+#include "SSCSEditorMenuContext.h"
+
+#define LOCTEXT_NAMESPACE "ChildActorComponentEditorUtils"
+
+bool FChildActorComponentEditorUtils::IsChildActorNode(TSharedPtr<const FSCSEditorTreeNode> InNodePtr)
+{
+	return InNodePtr.IsValid() && InNodePtr->GetNodeType() == FSCSEditorTreeNode::ENodeType::ChildActorNode;
+}
+
+bool FChildActorComponentEditorUtils::IsChildActorSubtreeNode(TSharedPtr<const FSCSEditorTreeNode> InNodePtr)
+{
+	return InNodePtr.IsValid() && IsChildActorNode(InNodePtr->GetActorRootNode());
+}
+
+bool FChildActorComponentEditorUtils::ContainsChildActorSubtreeNode(const TArray<TSharedPtr<FSCSEditorTreeNode>>& InNodePtrs)
+{
+	for (TSharedPtr<FSCSEditorTreeNode> NodePtr : InNodePtrs)
+	{
+		if (IsChildActorSubtreeNode(NodePtr))
+		{
+			return true;
+		}
+	}
+
+	return false;
+}
+
+TSharedPtr<FSCSEditorTreeNode> FChildActorComponentEditorUtils::GetOuterChildActorComponentNode(TSharedPtr<const FSCSEditorTreeNode> InNodePtr)
+{
+	if (InNodePtr.IsValid())
+	{
+		FSCSEditorActorNodePtrType ActorTreeRootNode = InNodePtr->GetActorRootNode();
+		if (IsChildActorNode(ActorTreeRootNode))
+		{
+			return ActorTreeRootNode->GetParent();
+		}
+	}
+
+	return nullptr;
+}
+
+bool FChildActorComponentEditorUtils::IsChildActorTreeViewExpansionEnabled()
+{
+	const UBlueprintEditorProjectSettings* EditorProjectSettings = GetDefault<UBlueprintEditorProjectSettings>();
+	return EditorProjectSettings->bEnableChildActorExpansionInTreeView;
+}
+
+EChildActorComponentTreeViewVisualizationMode FChildActorComponentEditorUtils::GetProjectDefaultTreeViewVisualizationMode()
+{
+	const UBlueprintEditorProjectSettings* EditorProjectSettings = GetDefault<UBlueprintEditorProjectSettings>();
+	return EditorProjectSettings->DefaultChildActorTreeViewMode;
+}
+
+void FChildActorComponentEditorUtils::ToggleComponentNodeVisibility(UChildActorComponent* ChildActorComponent, TWeakPtr<SSCSEditor> WeakEditorPtr)
+{
+	if (ChildActorComponent)
+	{
+		const bool bIsComponentNodeVisible = ShouldShowComponentNodeInTreeView(ChildActorComponent);
+
+		if (bIsComponentNodeVisible)
+		{
+			ChildActorComponent->SetEditorTreeViewVisualizationMode(EChildActorComponentTreeViewVisualizationMode::ChildActorOnly);
+		}
+		else
+		{
+			ChildActorComponent->SetEditorTreeViewVisualizationMode(EChildActorComponentTreeViewVisualizationMode::ComponentWithChildActor);
+		}
+
+		TSharedPtr<SSCSEditor> SCSEditor = WeakEditorPtr.Pin();
+		if (SCSEditor.IsValid())
+		{
+			SCSEditor->UpdateTree();
+		}
+	}
+}
+
+void FChildActorComponentEditorUtils::ToggleChildActorNodeVisibility(UChildActorComponent* ChildActorComponent, TWeakPtr<SSCSEditor> WeakEditorPtr)
+{
+	if (ChildActorComponent)
+	{
+		const bool bIsChildActorNodeVisible = ShouldShowChildActorNodeInTreeView(ChildActorComponent);
+
+		if (bIsChildActorNodeVisible)
+		{
+			ChildActorComponent->SetEditorTreeViewVisualizationMode(EChildActorComponentTreeViewVisualizationMode::ComponentOnly);
+		}
+		else
+		{
+			ChildActorComponent->SetEditorTreeViewVisualizationMode(EChildActorComponentTreeViewVisualizationMode::ComponentWithChildActor);
+		}
+
+		TSharedPtr<SSCSEditor> SCSEditor = WeakEditorPtr.Pin();
+		if (SCSEditor.IsValid())
+		{
+			SCSEditor->UpdateTree();
+		}
+	}
+}
+
+bool FChildActorComponentEditorUtils::ShouldShowComponentNodeInTreeView(UChildActorComponent* ChildActorComponent)
+{
+	if (ChildActorComponent)
+	{
+		if (!IsChildActorTreeViewExpansionEnabled())
+		{
+			// Always show the component node when tree view expansion is disabled.
+			return true;
+		}
+
+		EChildActorComponentTreeViewVisualizationMode CurrentMode = ChildActorComponent->GetEditorTreeViewVisualizationMode();
+		if (CurrentMode == EChildActorComponentTreeViewVisualizationMode::UseDefault)
+		{
+			CurrentMode = GetProjectDefaultTreeViewVisualizationMode();
+		}
+
+		return CurrentMode != EChildActorComponentTreeViewVisualizationMode::ChildActorOnly;
+	}
+
+	return false;
+}
+
+bool FChildActorComponentEditorUtils::ShouldShowChildActorNodeInTreeView(UChildActorComponent* ChildActorComponent)
+{
+	if (ChildActorComponent)
+	{
+		if (!IsChildActorTreeViewExpansionEnabled())
+		{
+			// Never show the child actor node when tree view expansion is disabled.
+			return false;
+		}
+
+		EChildActorComponentTreeViewVisualizationMode CurrentMode = ChildActorComponent->GetEditorTreeViewVisualizationMode();
+		if (CurrentMode == EChildActorComponentTreeViewVisualizationMode::UseDefault)
+		{
+			CurrentMode = GetProjectDefaultTreeViewVisualizationMode();
+		}
+
+		return CurrentMode != EChildActorComponentTreeViewVisualizationMode::ComponentOnly;
+	}
+
+	return false;
+}
+
+void FChildActorComponentEditorUtils::FillComponentContextMenuOptions(UToolMenu* Menu, UChildActorComponent* ChildActorComponent)
+{
+	if (!IsChildActorTreeViewExpansionEnabled())
+	{
+		return;
+	}
+
+	if (!ChildActorComponent)
+	{
+		return;
+	}
+
+	TWeakPtr<SSCSEditor> WeakEditorPtr;
+	if (USSCSEditorMenuContext* MenuContext = Menu->FindContext<USSCSEditorMenuContext>())
+	{
+		WeakEditorPtr = MenuContext->SCSEditor;
+	}
+
+	FToolMenuSection& Section = Menu->AddSection("ChildActorComponent", LOCTEXT("ChildActorComponentHeading", "Child Actor Component"));
+	{
+		FText ShowOrHideItemText;
+		if (ShouldShowChildActorNodeInTreeView(ChildActorComponent))
+		{
+			ShowOrHideItemText = LOCTEXT("HideChildActorNode_Label", "Hide Child Actor Node");
+		}
+		else
+		{
+			ShowOrHideItemText = LOCTEXT("ShowChildActorNode_Label", "Show Child Actor Node");
+		}
+
+		Section.AddMenuEntry(
+			"ToggleChildActorNode",
+			ShowOrHideItemText,
+			LOCTEXT("ToggleChildActorNode_ToolTip", "Toggle visibility of this component's child actor node."),
+			FSlateIcon(),
+			FUIAction(
+				FExecuteAction::CreateStatic(&FChildActorComponentEditorUtils::ToggleChildActorNodeVisibility, ChildActorComponent, WeakEditorPtr),
+				FCanExecuteAction()
+			)
+		);
+	}
+}
+
+void FChildActorComponentEditorUtils::FillChildActorContextMenuOptions(UToolMenu* Menu, TSharedPtr<const FSCSEditorTreeNode> InNodePtr)
+{
+	if (!IsChildActorTreeViewExpansionEnabled())
+	{
+		return;
+	}
+
+	if (!IsChildActorNode(InNodePtr))
+	{
+		return;
+	}
+
+	TSharedPtr<const FSCSEditorTreeNodeChildActor> ChildActorNodePtr = StaticCastSharedPtr<const FSCSEditorTreeNodeChildActor>(InNodePtr);
+	check(ChildActorNodePtr.IsValid());
+
+	UChildActorComponent* ChildActorComponent = ChildActorNodePtr->GetChildActorComponent();
+	if (!ChildActorComponent)
+	{
+		return;
+	}
+
+	TWeakPtr<SSCSEditor> WeakEditorPtr;
+	if (USSCSEditorMenuContext* MenuContext = Menu->FindContext<USSCSEditorMenuContext>())
+	{
+		WeakEditorPtr = MenuContext->SCSEditor;
+	}
+
+	FToolMenuSection& Section = Menu->AddSection("ChildActor", LOCTEXT("ChildActorHeading", "Child Actor"));
+	{
+		FText ShowOrHideItemText;
+		if (ShouldShowComponentNodeInTreeView(ChildActorComponent))
+		{
+			ShowOrHideItemText = LOCTEXT("HideChildActorComponentNode_Label", "Hide Child Actor Component Node");
+		}
+		else
+		{
+			ShowOrHideItemText = LOCTEXT("ShowChildActorComponentNode_Label", "Show Child Actor Component Node");
+		}
+
+		Section.AddMenuEntry(
+			"ToggleChildActorComponentNode",
+			ShowOrHideItemText,
+			LOCTEXT("ToggleChildActorComponentNode_ToolTip", "Toggle visibility of this child actor's outer component node."),
+			FSlateIcon(),
+			FUIAction(
+				FExecuteAction::CreateStatic(&FChildActorComponentEditorUtils::ToggleComponentNodeVisibility, ChildActorComponent, WeakEditorPtr),
+				FCanExecuteAction()
+			)
+		);
+	}
+}
+
+#undef LOCTEXT_NAMESPACE
diff --git a/Engine/Source/Editor/UnrealEd/Private/Settings/EditorProjectSettings.cpp b/Engine/Source/Editor/UnrealEd/Private/Settings/EditorProjectSettings.cpp
--- a/Engine/Source/Editor/UnrealEd/Private/Settings/EditorProjectSettings.cpp
+++ b/Engine/Source/Editor/UnrealEd/Private/Settings/EditorProjectSettings.cpp
@@ -2,6 +2,8 @@
 
 #include "Settings/EditorProjectSettings.h"
 #include "UObject/UnrealType.h"
+#include "Toolkits/ToolkitManager.h"
+#include "BlueprintEditor.h"
 
 EUnit ConvertDefaultInputUnits(EDefaultLocationUnit In)
 {
@@ -143,5 +145,40 @@ void ULevelEditor2DSettings::PostEditChangeProperty(struct FPropertyChangedEvent
 
 UBlueprintEditorProjectSettings::UBlueprintEditorProjectSettings(const FObjectInitializer& ObjectInitializer)
 	: Super(ObjectInitializer)
+	, DefaultChildActorTreeViewMode(EChildActorComponentTreeViewVisualizationMode::ComponentOnly)
 {
 }
+
+void UBlueprintEditorProjectSettings::PostEditChangeProperty(struct FPropertyChangedEvent& PropertyChangedEvent)
+{
+	Super::PostEditChangeProperty(PropertyChangedEvent);
+
+	const FName Name = PropertyChangedEvent.Property ? PropertyChangedEvent.Property->GetFName() : NAME_None;
+	if (Name == GET_MEMBER_NAME_CHECKED(UBlueprintEditorProjectSettings, bEnableChildActorExpansionInTreeView))
+	{
+		// Find open blueprint editors and refresh them
+		UAssetEditorSubsystem* AssetEditorSubsystem = GEditor->GetEditorSubsystem<UAssetEditorSubsystem>();
+		TArray<UObject*> EditedAssets = AssetEditorSubsystem->GetAllEditedAssets();
+		for (UObject* Asset : EditedAssets)
+		{
+			if (Asset && Asset->IsA<UBlueprint>())
+			{
+				TSharedPtr<IToolkit> AssetEditorPtr = FToolkitManager::Get().FindEditorForAsset(Asset);
+				if (AssetEditorPtr.IsValid() && AssetEditorPtr->IsBlueprintEditor())
+				{
+					TSharedPtr<IBlueprintEditor> BlueprintEditorPtr = StaticCastSharedPtr<IBlueprintEditor>(AssetEditorPtr);
+					BlueprintEditorPtr->RefreshEditors();
+				}
+			}
+		}
+
+		// Deselect actors so we are forced to clear the current tree view
+		// @todo - Figure out how to update the tree view directly instead?
+		if (GEditor && GEditor->GetSelectedActorCount() > 0)
+		{
+			const bool bNoteSelectionChange = true;
+			const bool bDeselectBSPSurfaces = true;
+			GEditor->SelectNone(bNoteSelectionChange, bDeselectBSPSurfaces);
+		}
+	}
+}
diff --git a/Engine/Source/Editor/UnrealEd/Public/Kismet2/ChildActorComponentEditorUtils.h b/Engine/Source/Editor/UnrealEd/Public/Kismet2/ChildActorComponentEditorUtils.h
new file mode 100644
--- /dev/null
+++ b/Engine/Source/Editor/UnrealEd/Public/Kismet2/ChildActorComponentEditorUtils.h
@@ -0,0 +1,50 @@
+// Copyright Epic Games, Inc. All Rights Reserved.
+
+#pragma once
+
+#include "CoreMinimal.h"
+#include "Components/ChildActorComponent.h"
+
+class UToolMenu;
+class SSCSEditor;
+class FSCSEditorTreeNode;
+
+class UNREALED_API FChildActorComponentEditorUtils
+{
+public:
+	/** Returns true if the given SCS editor tree node is a child actor node */
+	static bool IsChildActorNode(TSharedPtr<const FSCSEditorTreeNode> InNodePtr);
+
+	/** Returns true if the given SCS editor tree node belongs to a child actor template */
+	static bool IsChildActorSubtreeNode(TSharedPtr<const FSCSEditorTreeNode> InNodePtr);
+
+	/** Returns true if any array element belongs to a child actor template */
+	static bool ContainsChildActorSubtreeNode(const TArray<TSharedPtr<FSCSEditorTreeNode>>& InNodePtrs);
+
+	/** Returns the child actor component root node of a child actor subtree */
+	static TSharedPtr<FSCSEditorTreeNode> GetOuterChildActorComponentNode(TSharedPtr<const FSCSEditorTreeNode> InNodePtr);
+
+	/** Returns whether child actor tree view expansion is enabled in project settings */
+	static bool IsChildActorTreeViewExpansionEnabled();
+
+	/** Returns the default visualization mode for child actors in a component tree view */
+	static EChildActorComponentTreeViewVisualizationMode GetProjectDefaultTreeViewVisualizationMode();
+
+	/** Modify editor state to toggle visibility of the given component in a tree view */
+	static void ToggleComponentNodeVisibility(UChildActorComponent* ChildActorComponent, TWeakPtr<SSCSEditor> WeakEditorPtr);
+
+	/** Modify editor state to toggle Child Actor node visibility for the given component */
+	static void ToggleChildActorNodeVisibility(UChildActorComponent* ChildActorComponent, TWeakPtr<SSCSEditor> WeakEditorPtr);
+
+	/** Whether the given component should be visible as a unique node in a component tree view */
+	static bool ShouldShowComponentNodeInTreeView(UChildActorComponent* ChildActorComponent);
+
+	/** Whether the Child Actor should be shown in a component tree view for the given component */
+	static bool ShouldShowChildActorNodeInTreeView(UChildActorComponent* ChildActorComponent);
+
+	/** Populates the given menu with options for the given Child Actor component */
+	static void FillComponentContextMenuOptions(UToolMenu* Menu, UChildActorComponent* ChildActorComponent);
+
+	/** Populates the given menu with additional options if the given SCS editor tree node represents a Child Actor node */
+	static void FillChildActorContextMenuOptions(UToolMenu* Menu, TSharedPtr<const FSCSEditorTreeNode> InNodePtr);
+};
\ No newline at end of file
diff --git a/Engine/Source/Editor/UnrealEd/Public/Settings/EditorProjectSettings.h b/Engine/Source/Editor/UnrealEd/Public/Settings/EditorProjectSettings.h
--- a/Engine/Source/Editor/UnrealEd/Public/Settings/EditorProjectSettings.h
+++ b/Engine/Source/Editor/UnrealEd/Public/Settings/EditorProjectSettings.h
@@ -6,6 +6,7 @@
 #include "UObject/ObjectMacros.h"
 #include "Engine/DeveloperSettings.h"
 #include "Math/UnitConversion.h"
+#include "Components/ChildActorComponent.h"
 
 #include "EditorProjectSettings.generated.h"
 
@@ -195,5 +196,22 @@ public:
 	// The list of namespaces to always expose in any Blueprint (for all users of the game/project)
 	UPROPERTY(EditAnywhere, config, Category=Experimental)
 	TArray<FString> NamespacesToAlwaysInclude;
+	
+	/**
+	 * Enable the option to expand child actor components within component tree views (experimental).
+	 */
+	UPROPERTY(EditAnywhere, config, Category=Experimental)
+	bool bEnableChildActorExpansionInTreeView;
+	
+	/**
+	 * Default view mode to use for child actor components in a Blueprint actor's component tree hierarchy (experimental).
+	 */
+	UPROPERTY(EditAnywhere, config, Category=Experimental, meta=(EditCondition="bEnableChildActorExpansionInTreeView"))
+	EChildActorComponentTreeViewVisualizationMode DefaultChildActorTreeViewMode;
+
+public:
+	// UObject interface
+	virtual void PostEditChangeProperty(struct FPropertyChangedEvent& PropertyChangedEvent) override;
+	// End of UObject interface
 };
 
diff --git a/Engine/Source/Runtime/Engine/Classes/Components/ChildActorComponent.h b/Engine/Source/Runtime/Engine/Classes/Components/ChildActorComponent.h
--- a/Engine/Source/Runtime/Engine/Classes/Components/ChildActorComponent.h
+++ b/Engine/Source/Runtime/Engine/Classes/Components/ChildActorComponent.h
@@ -55,6 +55,21 @@ public:
 	TSharedPtr<FComponentInstanceDataCache> ComponentInstanceData;
 };
 
+#if WITH_EDITORONLY_DATA
+UENUM()
+enum class EChildActorComponentTreeViewVisualizationMode : uint8
+{
+	/** Use the editor's default setting. */
+	UseDefault UMETA(Hidden),
+	/** Show only the outer component as a single component node. */
+	ComponentOnly,
+	/** Include the child actor hierarchy attached to the outer component as the root node. */
+	ComponentWithChildActor,
+	/** Show only as a child actor hierarchy (i.e. do not show the outer component node as the root). */
+	ChildActorOnly,
+};
+#endif
+
 /** A component that spawns an Actor when registered, and destroys it when unregistered.*/
 UCLASS(ClassGroup=Utility, hidecategories=(Object,LOD,Physics,Lighting,TextureStreaming,Activation,"Components|Activation",Collision), meta=(BlueprintSpawnableComponent))
 class ENGINE_API UChildActorComponent : public USceneComponent
@@ -108,6 +123,12 @@ private:
 	/** Cached copy of the instance data when the ChildActor is destroyed to be available when needed */
 	mutable FChildActorComponentInstanceData* CachedInstanceData;
 
+#if WITH_EDITORONLY_DATA
+	/** Indicates how this component will be visualized for editing in a tree view. Users can change this setting per instance via the context menu in the Blueprint/SCS editor. */
+	UPROPERTY()
+	EChildActorComponentTreeViewVisualizationMode EditorTreeViewVisualizationMode;
+#endif
+
 	/** Flag indicating that when the component is registered that the child actor should be recreated */
 	uint8 bNeedsRecreate:1;
 
@@ -149,6 +170,15 @@ public:
 
 	/** Kill any currently present child actor */
 	void DestroyChildActor();
+
+#if WITH_EDITOR
+	EChildActorComponentTreeViewVisualizationMode GetEditorTreeViewVisualizationMode() const
+	{
+		return EditorTreeViewVisualizationMode;
+	}
+
+	void SetEditorTreeViewVisualizationMode(EChildActorComponentTreeViewVisualizationMode InMode);
+#endif
 };
 
 
diff --git a/Engine/Source/Runtime/Engine/Classes/Engine/SCS_Node.h b/Engine/Source/Runtime/Engine/Classes/Engine/SCS_Node.h
--- a/Engine/Source/Runtime/Engine/Classes/Engine/SCS_Node.h
+++ b/Engine/Source/Runtime/Engine/Classes/Engine/SCS_Node.h
@@ -186,7 +186,7 @@ class USCS_Node : public UObject
 	ENGINE_API void SetParent(USCS_Node* InParentNode);
 
 	/** Sets parent component attributes based on the given component instance */
-	ENGINE_API void SetParent(USceneComponent* InParentComponent);
+	ENGINE_API void SetParent(const USceneComponent* InParentComponent);
 
 	/** Finds and returns the parent component template through the given Blueprint */
 	ENGINE_API USceneComponent* GetParentComponentTemplate(UBlueprint* InBlueprint) const;
diff --git a/Engine/Source/Runtime/Engine/Private/Components/ChildActorComponent.cpp b/Engine/Source/Runtime/Engine/Private/Components/ChildActorComponent.cpp
--- a/Engine/Source/Runtime/Engine/Private/Components/ChildActorComponent.cpp
+++ b/Engine/Source/Runtime/Engine/Private/Components/ChildActorComponent.cpp
@@ -14,6 +14,10 @@ UChildActorComponent::UChildActorComponent(const FObjectInitializer& ObjectIniti
 	: Super(ObjectInitializer)
 {
 	bAllowReregistration = false;
+
+#if WITH_EDITORONLY_DATA
+	EditorTreeViewVisualizationMode = EChildActorComponentTreeViewVisualizationMode::UseDefault;
+#endif
 }
 
 void UChildActorComponent::OnRegister()
@@ -724,3 +728,12 @@ void UChildActorComponent::BeginPlay()
 		ChildActor->DispatchBeginPlay(bFromLevelStreaming);
 	}
 }
+
+#if WITH_EDITOR
+void UChildActorComponent::SetEditorTreeViewVisualizationMode(EChildActorComponentTreeViewVisualizationMode InMode)
+{
+	Modify();
+
+	EditorTreeViewVisualizationMode = InMode;
+}
+#endif
\ No newline at end of file
diff --git a/Engine/Source/Runtime/Engine/Private/SCS_Node.cpp b/Engine/Source/Runtime/Engine/Private/SCS_Node.cpp
--- a/Engine/Source/Runtime/Engine/Private/SCS_Node.cpp
+++ b/Engine/Source/Runtime/Engine/Private/SCS_Node.cpp
@@ -501,7 +501,7 @@ void USCS_Node::SetParent(USCS_Node* InParentNode)
 	}
 }
 
-void USCS_Node::SetParent(USceneComponent* InParentComponent)
+void USCS_Node::SetParent(const USceneComponent* InParentComponent)
 {
 	check(InParentComponent != NULL);
 
