Fix check() firing in GetEditingObjects() when opening some animbps. It was skippable in VS, but shutting down the editor for animators.
Also fix similar issue reported here: https://udn.unrealengine.com/questions/564931/crash-when-open-animation-blueprint-editor-due-to.html
#rnx

diff --git a/Engine/Source/Editor/AnimationBlueprintEditor/Private/AnimationBlueprintEditor.cpp b/Engine/Source/Editor/AnimationBlueprintEditor/Private/AnimationBlueprintEditor.cpp
--- a/Engine/Source/Editor/AnimationBlueprintEditor/Private/AnimationBlueprintEditor.cpp
+++ b/Engine/Source/Editor/AnimationBlueprintEditor/Private/AnimationBlueprintEditor.cpp
@@ -1367,7 +1367,7 @@ void FAnimationBlueprintEditor::OnBlueprintPreCompile(UBlueprint* BlueprintToCom
 		}
 	}
 
-	if(BlueprintToCompile == GetBlueprintObj())
+	if(GetObjectsCurrentlyBeingEdited()->Num() > 0 && BlueprintToCompile == GetBlueprintObj())
 	{
 		// Grab the currently debugged object, so we can re-set it below in OnBlueprintPostCompile
 		DebuggedMeshComponent = nullptr;
diff --git a/Engine/Source/Editor/Kismet/Private/BlueprintEditor.cpp b/Engine/Source/Editor/Kismet/Private/BlueprintEditor.cpp
--- a/Engine/Source/Editor/Kismet/Private/BlueprintEditor.cpp
+++ b/Engine/Source/Editor/Kismet/Private/BlueprintEditor.cpp
@@ -3898,9 +3898,12 @@ void FBlueprintEditor::JumpToPin(const UEdGraphPin* Pin)
 
 void FBlueprintEditor::AddReferencedObjects( FReferenceCollector& Collector )
 {
-	TArray<UObject*>& LocalEditingObjects = const_cast<TArray<UObject*>&>(GetEditingObjects());
+	if (GetObjectsCurrentlyBeingEdited()->Num() > 0)
+	{
+		TArray<UObject*>& LocalEditingObjects = const_cast<TArray<UObject*>&>(GetEditingObjects());
 
-	Collector.AddReferencedObjects(LocalEditingObjects);
+		Collector.AddReferencedObjects(LocalEditingObjects);
+	}
 
 	Collector.AddReferencedObjects(StandardLibraries);
 
