DetailsView - Reconfigure Tool Button placement between the DetailsNameArea and SSCSEditor for ActorDetails and BlueprintDetails Views.

diff --git a/Engine/Source/Editor/Kismet/Private/SSCSEditor.cpp b/Engine/Source/Editor/Kismet/Private/SSCSEditor.cpp
--- a/Engine/Source/Editor/Kismet/Private/SSCSEditor.cpp
+++ b/Engine/Source/Editor/Kismet/Private/SSCSEditor.cpp
@@ -3769,7 +3769,6 @@ void SSCSEditor::Construct( const FArguments& InArgs )
 		FUIAction(FExecuteAction::CreateSP(this, &SSCSEditor::PromoteToBlueprint))
 	);
 
-	TSharedPtr<SHorizontalBox> ButtonBox;
 	TSharedPtr<SVerticalBox>   HeaderBox;
 	TSharedPtr<SWidget> SearchBar =
 		SAssignNew(FilterBox, SSearchBox)
@@ -3780,6 +3779,65 @@ void SSCSEditor::Construct( const FArguments& InArgs )
 
 	bool bHideComponentClassCombo = InArgs._HideComponentClassCombo.Get();
 
+	ButtonBox = SNew(SHorizontalBox)
+	+ SHorizontalBox::Slot()
+	.AutoWidth()
+	[
+		SNew(SComponentClassCombo)
+		.AddMetaData<FTagMetaData>(FTagMetaData(TEXT("Actor.AddComponent")))
+		.Visibility(bHideComponentClassCombo ? EVisibility::Hidden : EVisibility::Visible)
+		.OnComponentClassSelected(this, &SSCSEditor::PerformComboAddClass)
+		.ToolTipText(LOCTEXT("AddComponent_Tooltip", "Adds a new component to this actor"))
+		.IsEnabled(AllowEditing)
+	]
+
+	//
+	// horizontal slot (index) #1 => reserved for BP-editor search bar (see 'ButtonBox' usage below)
+
+	+ SHorizontalBox::Slot()
+	.AutoWidth()
+	[
+		SNew( SButton )
+		.ButtonStyle( FEditorStyle::Get(), "SimpleButton" )
+		.AddMetaData<FTagMetaData>(FTagMetaData(TEXT("Actor.ConvertToBlueprint")))
+		.Visibility( this, &SSCSEditor::GetPromoteToBlueprintButtonVisibility )
+		.OnClicked( this, &SSCSEditor::OnPromoteToBlueprintClicked )
+		.ContentPadding(FMargin(5.f))
+		.ToolTip(IDocumentation::Get()->CreateToolTip(
+			LOCTEXT("PromoteToBluerprintTooltip","Converts this actor into a reusable Blueprint Class that can have script behavior" ),
+			NULL,
+			TEXT("Shared/LevelEditor"),
+			TEXT("ConvertToBlueprint")))
+		[
+			SNew(SImage)
+			.ColorAndOpacity(FSlateColor::UseForeground())
+			.Image(FAppStyle::Get().GetBrush("Icons.Blueprints"))
+		]
+	]
+	+ SHorizontalBox::Slot()
+	.AutoWidth()
+	[
+		SNew(SComboButton)
+		.AddMetaData<FTagMetaData>(FTagMetaData(TEXT("Actor.EditBlueprint")))
+		.Visibility(this, &SSCSEditor::GetEditBlueprintButtonVisibility)
+		.ComboButtonStyle(FEditorStyle::Get(), "SimpleComboButton")
+		.HasDownArrow(false)
+		.ContentPadding(FMargin(5.f))
+		.ToolTipText(LOCTEXT("EditActorBlueprint_Tooltip", "Edit the Blueprint for this Actor"))
+		.ButtonContent()
+		[
+			SNew(SImage)
+			.ColorAndOpacity(FSlateColor::UseForeground())
+			.Image(FAppStyle::Get().GetBrush("Icons.Blueprints"))
+		]
+		.MenuContent()
+		[
+			EditBlueprintMenuBuilder.MakeWidget()
+		]
+	];
+
+
+
 	Contents = SNew(SVerticalBox)
 	+ SVerticalBox::Slot()
 	.Padding(0.0f)
@@ -3797,77 +3855,6 @@ void SSCSEditor::Construct( const FArguments& InArgs )
 			.BorderBackgroundColor( FLinearColor( .6,.6,.6, 1.0f ) )
 			[
 				SAssignNew(HeaderBox, SVerticalBox)
-					+ SVerticalBox::Slot()
-						.AutoHeight()
-						.VAlign(VAlign_Top)
-						.HAlign(HAlign_Fill)
-					[
-						SAssignNew(ButtonBox, SHorizontalBox)
-
-						+ SHorizontalBox::Slot()
-						.FillWidth(1)
-						[
-							SNew(SSpacer)
-						]
-
-						+ SHorizontalBox::Slot()
-						.AutoWidth()
-						[
-							SNew(SComponentClassCombo)
-							.AddMetaData<FTagMetaData>(FTagMetaData(TEXT("Actor.AddComponent")))
-							.Visibility(bHideComponentClassCombo ? EVisibility::Hidden : EVisibility::Visible)
-							.OnComponentClassSelected(this, &SSCSEditor::PerformComboAddClass)
-							.ToolTipText(LOCTEXT("AddComponent_Tooltip", "Adds a new component to this actor"))
-							.IsEnabled(AllowEditing)
-						]
-
-						//
-						// horizontal slot (index) #1 => reserved for BP-editor search bar (see 'ButtonBox' usage below)
-
-						+ SHorizontalBox::Slot()
-						.AutoWidth()
-						[
-							SNew( SButton )
-							.ButtonStyle( FEditorStyle::Get(), "SimpleButton" )
-							.AddMetaData<FTagMetaData>(FTagMetaData(TEXT("Actor.ConvertToBlueprint")))
-							.Visibility( this, &SSCSEditor::GetPromoteToBlueprintButtonVisibility )
-							.OnClicked( this, &SSCSEditor::OnPromoteToBlueprintClicked )
-							.ContentPadding(FMargin(5.f))
-							.ToolTip(IDocumentation::Get()->CreateToolTip(
-								LOCTEXT("PromoteToBluerprintTooltip","Converts this actor into a reusable Blueprint Class that can have script behavior" ),
-								NULL,
-								TEXT("Shared/LevelEditor"),
-								TEXT("ConvertToBlueprint")))
-							[
-								SNew(SImage)
-								.ColorAndOpacity(FSlateColor::UseForeground())
-								.Image(FAppStyle::Get().GetBrush("Icons.Blueprints"))
-							]
-						]
-						+ SHorizontalBox::Slot()
-						.AutoWidth()
-						[
-							SNew(SComboButton)
-							.AddMetaData<FTagMetaData>(FTagMetaData(TEXT("Actor.EditBlueprint")))
-							.Visibility(this, &SSCSEditor::GetEditBlueprintButtonVisibility)
-							.ComboButtonStyle(FEditorStyle::Get(), "SimpleComboButton")
-							.HasDownArrow(false)
-							.ContentPadding(FMargin(5.f))
-							.ButtonContent()
-							[
-								SNew(SImage)
-								.ColorAndOpacity(FSlateColor::UseForeground())
-								.Image(FAppStyle::Get().GetBrush("Icons.Blueprints"))
-							]
-							.MenuContent()
-							[
-								EditBlueprintMenuBuilder.MakeWidget()
-							]
-						]
-					]
-
-				//
-				// vertical slot (index) #1 => reserved for instance-editor search bar (see 'HeaderBox' usage below)
 			]
 		]
 
@@ -3884,11 +3871,10 @@ void SSCSEditor::Construct( const FArguments& InArgs )
 		]
 	];
 
-	// insert the search bar, depending on which editor this widget is in (depending on convert/edit button visibility)
-	if (bInlineSearchBarWithButtons)
+	// Only insert the buttons and search bar in the Blueprints version
+	if (bInlineSearchBarWithButtons) // Blueprints
 	{
-		const int32 SearchBarHorizontalSlotIndex = 1;
-
+		const int32 SearchBarHorizontalSlotIndex = 0;
 		ButtonBox->InsertSlot(SearchBarHorizontalSlotIndex)
 			.FillWidth(1.0f)
 			.VAlign(VAlign_Center)
@@ -3896,16 +3882,10 @@ void SSCSEditor::Construct( const FArguments& InArgs )
 		[
 			SearchBar.ToSharedRef()
 		];
-	}
-	else
-	{
-		const int32 SearchBarVerticalSlotIndex = 1;
 
-		HeaderBox->InsertSlot(SearchBarVerticalSlotIndex)
-			.VAlign(VAlign_Center)
-			.Padding(3.0f, 1.0f)
+		HeaderBox->AddSlot()
 		[
-			SearchBar.ToSharedRef()
+			ButtonBox.ToSharedRef()
 		];
 	}
 
@@ -3925,6 +3905,10 @@ void SSCSEditor::Construct( const FArguments& InArgs )
 }
 END_SLATE_FUNCTION_BUILD_OPTIMIZATION
 
+TSharedPtr<SWidget> SSCSEditor::GetToolButtonsBox()
+{
+	return ButtonBox;	
+}
 
 void SSCSEditor::OnLevelComponentRequestRename(const UActorComponent* InComponent)
 {
diff --git a/Engine/Source/Editor/Kismet/Public/SSCSEditor.h b/Engine/Source/Editor/Kismet/Public/SSCSEditor.h
--- a/Engine/Source/Editor/Kismet/Public/SSCSEditor.h
+++ b/Engine/Source/Editor/Kismet/Public/SSCSEditor.h
@@ -1158,6 +1158,9 @@ public:
 	/** Called at the end of each frame. */
 	void OnPostTick(float);
 
+	/** Return the button widgets that can add components or create/edit blueprints */
+	TSharedPtr<SWidget> GetToolButtonsBox();
+
 protected:
 	FSCSEditorTreeNodePtrType FindOrCreateParentForExistingComponent(UActorComponent* InActorComponent, FSCSEditorActorNodePtrType ActorRootNode);
 	FSCSEditorTreeNodePtrType FindParentForNewComponent(UActorComponent* NewComponent) const;
@@ -1378,4 +1381,7 @@ private:
 
 	/** The filter box that handles filtering for the tree. */
 	TSharedPtr< SSearchBox > FilterBox;
+
+	/** The tools buttons box **/
+	TSharedPtr<SHorizontalBox> ButtonBox;
 };
diff --git a/Engine/Source/Editor/LevelEditor/Private/SActorDetails.cpp b/Engine/Source/Editor/LevelEditor/Private/SActorDetails.cpp
--- a/Engine/Source/Editor/LevelEditor/Private/SActorDetails.cpp
+++ b/Engine/Source/Editor/LevelEditor/Private/SActorDetails.cpp
@@ -128,7 +128,6 @@ void SActorDetails::Construct(const FArguments& InArgs, const FName TabIdentifie
 	DetailsView->SetIsPropertyEditingEnabledDelegate(FIsPropertyEditingEnabled::CreateSP(this, &SActorDetails::IsPropertyEditingEnabled));
 	DetailsView->SetOnObjectArrayChanged(FOnObjectArrayChanged::CreateSP(this, &SActorDetails::OnDetailsViewObjectArrayChanged));
 
-
 	// Set up a delegate to call to add generic details to the view
 	DetailsView->SetGenericLayoutDetailsDelegate(FOnGetDetailCustomizationInstance::CreateStatic(&FLevelEditorGenericDetails::MakeInstance));
 
@@ -146,6 +145,10 @@ void SActorDetails::Construct(const FArguments& InArgs, const FName TabIdentifie
 		
 	ComponentsBox->SetContent(SCSEditor.ToSharedRef());
 
+	TSharedRef<SWidget> ButtonBox = SCSEditor->GetToolButtonsBox().ToSharedRef();
+	DetailsView->SetNameAreaCustomContent( ButtonBox );
+	
+
 	ChildSlot
 	[
 		SNew(SVerticalBox)
diff --git a/Engine/Source/Editor/PropertyEditor/Private/SDetailNameArea.cpp b/Engine/Source/Editor/PropertyEditor/Private/SDetailNameArea.cpp
--- a/Engine/Source/Editor/PropertyEditor/Private/SDetailNameArea.cpp
+++ b/Engine/Source/Editor/PropertyEditor/Private/SDetailNameArea.cpp
@@ -25,6 +25,7 @@ void SDetailNameArea::Construct( const FArguments& InArgs, const TArray< TWeakOb
 	SelectionTip = InArgs._SelectionTip;
 	bShowLockButton = InArgs._ShowLockButton;
 	bShowActorLabel = InArgs._ShowActorLabel;
+	CustomContent = SNullWidget::NullWidget;
 }
 
 void SDetailNameArea::Refresh( const TArray< TWeakObjectPtr<UObject> >& SelectedObjects )
@@ -137,7 +138,7 @@ TSharedRef< SWidget > SDetailNameArea::BuildObjectNameArea( const TArray< TWeakO
 			TSharedRef<IObjectNameEditableTextBox> ObjectNameBox = EdWidgetsModule.CreateObjectNameEditableTextBox(SelectedObjects);
 
 			ObjectNameArea->AddSlot()
-				.AutoWidth()
+				.FillWidth(1.0)
 				.Padding(0, 0, 3, 0)
 				[
 					SNew(SBox)
@@ -152,11 +153,18 @@ TSharedRef< SWidget > SDetailNameArea::BuildObjectNameArea( const TArray< TWeakO
 		const TWeakObjectPtr< UObject > ObjectWeakPtr = SelectedObjects.Num() == 1 ? SelectedObjects[0] : NULL;
 		BuildObjectNameAreaSelectionLabel( ObjectNameArea, ObjectWeakPtr, SelectedObjects.Num() );
 
+		ObjectNameArea->AddSlot()
+		.Expose(CustomContentSlot)
+		.AutoWidth()
+		[
+			CustomContent.ToSharedRef()
+		];
+
 		if( bShowLockButton )
 		{
 			ObjectNameArea->AddSlot()
 				.HAlign(HAlign_Right)
-				.FillWidth(1.0f)
+				.AutoWidth()
 				[
 					SNew( SButton )
 					.ButtonStyle( FEditorStyle::Get(), "SimpleButton" )
@@ -252,5 +260,17 @@ void SDetailNameArea::BuildObjectNameAreaSelectionLabel( TSharedRef< SHorizontal
 }
 
 
+void SDetailNameArea::SetCustomContent(TSharedRef<SWidget>& InCustomContent)
+{
+	CustomContent = InCustomContent;
+
+	if (CustomContentSlot != nullptr)
+	{
+		SHorizontalBox::FSlot& T = *CustomContentSlot;
+		T[InCustomContent];
+	}
+
+}
+
 #undef LOCTEXT_NAMESPACE
 
diff --git a/Engine/Source/Editor/PropertyEditor/Private/SDetailNameArea.h b/Engine/Source/Editor/PropertyEditor/Private/SDetailNameArea.h
--- a/Engine/Source/Editor/PropertyEditor/Private/SDetailNameArea.h
+++ b/Engine/Source/Editor/PropertyEditor/Private/SDetailNameArea.h
@@ -44,6 +44,12 @@ public:
 	 */
 	void Refresh( const TArray< TWeakObjectPtr<AActor> >& SelectedActors, const TArray< TWeakObjectPtr<UObject> >& SelectedObjects, FDetailsViewArgs::ENameAreaSettings NameAreaSettings  );
 
+
+	/**
+	 * Inserts Custom Content (typically tool buttons) before the lock
+	 */
+	virtual void SetCustomContent(TSharedRef<SWidget>& InCustomContent);
+
 private:
 	/** @return the Slate brush to use for the lock image */
 	const FSlateBrush* OnGetLockButtonImageResource() const;
@@ -57,6 +63,12 @@ private:
 	FOnClicked OnLockButtonClicked;
 	TAttribute<bool> IsLocked;
 	TAttribute<bool> SelectionTip;
+
+	TSharedPtr<SWidget> CustomContent;
+
 	bool bShowLockButton;
 	bool bShowActorLabel;
+
+	/** Area where the customs content resides */
+	SHorizontalBox::FSlot* CustomContentSlot;
 };
diff --git a/Engine/Source/Editor/PropertyEditor/Private/SDetailsViewBase.cpp b/Engine/Source/Editor/PropertyEditor/Private/SDetailsViewBase.cpp
--- a/Engine/Source/Editor/PropertyEditor/Private/SDetailsViewBase.cpp
+++ b/Engine/Source/Editor/PropertyEditor/Private/SDetailsViewBase.cpp
@@ -653,6 +653,11 @@ TSharedPtr<SWidget> SDetailsViewBase::GetNameAreaWidget()
 	return DetailsViewArgs.bCustomNameAreaLocation ? NameArea : nullptr;
 }
 
+void SDetailsViewBase::SetNameAreaCustomContent( TSharedRef<SWidget>& InCustomContent )
+{
+	NameArea->SetCustomContent(InCustomContent);	
+}
+
 TSharedPtr<SWidget> SDetailsViewBase::GetFilterAreaWidget()
 {
 	return DetailsViewArgs.bCustomFilterAreaLocation ? FilterRow : nullptr;
diff --git a/Engine/Source/Editor/PropertyEditor/Private/SDetailsViewBase.h b/Engine/Source/Editor/PropertyEditor/Private/SDetailsViewBase.h
--- a/Engine/Source/Editor/PropertyEditor/Private/SDetailsViewBase.h
+++ b/Engine/Source/Editor/PropertyEditor/Private/SDetailsViewBase.h
@@ -323,6 +323,7 @@ protected:
 	virtual void UpdatePropertiesWhitelist(const TSet<FPropertyPath> InWhitelistedProperties) override { CurrentFilter.WhitelistedProperties = InWhitelistedProperties; }
 
 	virtual TSharedPtr<SWidget> GetNameAreaWidget() override;
+	virtual void SetNameAreaCustomContent(TSharedRef<SWidget>& InCustomContent) override;
 	virtual TSharedPtr<SWidget> GetFilterAreaWidget() override;
 	virtual TSharedPtr<class FUICommandList> GetHostCommandList() const override;
 	virtual TSharedPtr<FTabManager> GetHostTabManager() const override;
diff --git a/Engine/Source/Editor/PropertyEditor/Public/IDetailsView.h b/Engine/Source/Editor/PropertyEditor/Public/IDetailsView.h
--- a/Engine/Source/Editor/PropertyEditor/Public/IDetailsView.h
+++ b/Engine/Source/Editor/PropertyEditor/Public/IDetailsView.h
@@ -340,6 +340,9 @@ public:
 	/** Returns the name area widget used to display object naming functionality so it can be placed in a custom location.  Note FDetailsViewArgs.bCustomNameAreaLocation must be true */
 	virtual TSharedPtr<SWidget> GetNameAreaWidget() = 0;
 
+	/** Optinoally add custom tools into the NameArea */
+	virtual void SetNameAreaCustomContent(TSharedRef<SWidget>& InCustomContent) = 0;
+
 	/** Returns the search area widget used to display search and view options so it can be placed in a custom location.  Note FDetailsViewArgs.bCustomFilterAreaLocation must be true */
 	virtual TSharedPtr<SWidget> GetFilterAreaWidget() = 0;
 
