Remove FlatButton light styles from editor style. The default button style will be used instead which matches the starship style

diff --git a/Engine/Plugins/FX/Niagara/Source/NiagaraEditor/Private/TypeEditorUtilities/NiagaraEnumTypeEditorUtilities.cpp b/Engine/Plugins/FX/Niagara/Source/NiagaraEditor/Private/TypeEditorUtilities/NiagaraEnumTypeEditorUtilities.cpp
--- a/Engine/Plugins/FX/Niagara/Source/NiagaraEditor/Private/TypeEditorUtilities/NiagaraEnumTypeEditorUtilities.cpp
+++ b/Engine/Plugins/FX/Niagara/Source/NiagaraEditor/Private/TypeEditorUtilities/NiagaraEnumTypeEditorUtilities.cpp
@@ -26,7 +26,6 @@ public:
 		[
 			SNew(SEnumComboBox, Enum)
 			.CurrentValue(this, &SNiagaraEnumParameterEditor::GetValue)
-			.ButtonStyle(FEditorStyle::Get(), "FlatButton.Light")
 			.ContentPadding(FMargin(2, 0))
 			.Font(FEditorStyle::GetFontStyle("Sequencer.AnimationOutliner.RegularFont"))
 			.OnEnumSelectionChanged(SEnumComboBox::FOnEnumSelectionChanged::CreateSP(this, &SNiagaraEnumParameterEditor::ValueChanged))
diff --git a/Engine/Source/Editor/EditorStyle/Private/StarshipStyle.cpp b/Engine/Source/Editor/EditorStyle/Private/StarshipStyle.cpp
--- a/Engine/Source/Editor/EditorStyle/Private/StarshipStyle.cpp
+++ b/Engine/Source/Editor/EditorStyle/Private/StarshipStyle.cpp
@@ -469,11 +469,6 @@ void FStarshipEditorStyle::FStyle::SetupGeneralStyles()
 			.SetPressed(BOX_BRUSH("Common/FlatButton", 2.0f / 8.0f, SelectionColor_Pressed))
 		);
 
-		Set("FlatButton.Light", FButtonStyle(Button)
-			.SetNormal(BOX_BRUSH("Common/FlatButton", 2.0f / 8.0f, FLinearColor(0.72267, 0.72267, 0.72267, 1)))
-			.SetHovered(BOX_BRUSH("Common/FlatButton", 2.0f / 8.0f, FLinearColor(0.85, 0.85, 0.85, 1)))
-			.SetPressed(BOX_BRUSH("Common/FlatButton", 2.0f / 8.0f, FLinearColor(0.58597, 0.58597, 0.58597, 1)))
-			);
 
 		Set("FlatButton.Default", GetWidgetStyle<FButtonStyle>("FlatButton.Dark"));
 
diff --git a/Engine/Source/Editor/Kismet/Private/SBlueprintEditorSelectedDebugObjectWidget.cpp b/Engine/Source/Editor/Kismet/Private/SBlueprintEditorSelectedDebugObjectWidget.cpp
--- a/Engine/Source/Editor/Kismet/Private/SBlueprintEditorSelectedDebugObjectWidget.cpp
+++ b/Engine/Source/Editor/Kismet/Private/SBlueprintEditorSelectedDebugObjectWidget.cpp
@@ -34,7 +34,6 @@ void SBlueprintEditorSelectedDebugObjectWidget::Construct(const FArguments& InAr
 	LastObjectObserved = nullptr;
 
 	DebugWorldsComboBox = SNew(STextComboBox)
-		.ButtonStyle(FEditorStyle::Get(), "FlatButton.Light")
 		.ToolTip(IDocumentation::Get()->CreateToolTip(
 		LOCTEXT("BlueprintDebugWorldTooltip", "Select a world to debug, will filter what to debug if no specific object selected"),
 		nullptr,
@@ -47,7 +46,6 @@ void SBlueprintEditorSelectedDebugObjectWidget::Construct(const FArguments& InAr
 		.OnSelectionChanged(this, &SBlueprintEditorSelectedDebugObjectWidget::DebugWorldSelectionChanged);
 
 	DebugObjectsComboBox = SNew(SComboBox<TSharedPtr<FBlueprintDebugObjectInstance>>)
-		.ButtonStyle(FEditorStyle::Get(), "FlatButton.Light")
 		.ToolTip(IDocumentation::Get()->CreateToolTip(
 		LOCTEXT("BlueprintDebugObjectTooltip", "Select an object to debug, if set to none will debug any object"),
 		nullptr,
diff --git a/Engine/Source/Editor/MovieSceneTools/Private/MovieSceneToolHelpers.cpp b/Engine/Source/Editor/MovieSceneTools/Private/MovieSceneToolHelpers.cpp
--- a/Engine/Source/Editor/MovieSceneTools/Private/MovieSceneToolHelpers.cpp
+++ b/Engine/Source/Editor/MovieSceneTools/Private/MovieSceneToolHelpers.cpp
@@ -533,7 +533,6 @@ TSharedRef<SWidget> MovieSceneToolHelpers::MakeEnumComboBox(const UEnum* InEnum,
 {
 	return SNew(SEnumComboBox, InEnum)
 		.CurrentValue(InCurrentValue)
-		.ButtonStyle(FEditorStyle::Get(), "FlatButton.Light")
 		.ContentPadding(FMargin(2, 0))
 		.Font(FEditorStyle::GetFontStyle("Sequencer.AnimationOutliner.RegularFont"))
 		.OnEnumSelectionChanged(InOnSelectionChanged);
diff --git a/Engine/Source/Editor/Persona/Private/SAnimMontageSectionsPanel.cpp b/Engine/Source/Editor/Persona/Private/SAnimMontageSectionsPanel.cpp
--- a/Engine/Source/Editor/Persona/Private/SAnimMontageSectionsPanel.cpp
+++ b/Engine/Source/Editor/Persona/Private/SAnimMontageSectionsPanel.cpp
@@ -332,7 +332,6 @@ void SAnimMontageSectionsPanel::Update()
 						SNew(SButton)
 						.IsEnabled(!bChildAnimMontage)
 						.ButtonColorAndOpacity(IsLoop(SectionIdx) ? LoopColor : NodeColor)
-						.ButtonStyle(FEditorStyle::Get(), "FlatButton.Light")
 						[
 							SNew(SHorizontalBox)
 							+SHorizontalBox::Slot()
