Add "Editor Only" to the tooltip in the SCS editor if a component is marked as Editor Only

diff --git a/Engine/Source/Editor/Kismet/Private/SSCSEditor.cpp b/Engine/Source/Editor/Kismet/Private/SSCSEditor.cpp
--- a/Engine/Source/Editor/Kismet/Private/SSCSEditor.cpp
+++ b/Engine/Source/Editor/Kismet/Private/SSCSEditor.cpp
@@ -1947,6 +1947,9 @@ TSharedRef<SToolTip> SSCS_RowWidget::CreateToolTipWidget() const
 				InfoBox->AddSlot()[SNew(SSpacer).Size(FVector2D(1.0f, 8.0f))];
 				AddToToolTipInfoBox(InfoBox, LOCTEXT("TooltipAsset", "Asset"), SNullWidget::NullWidget, TAttribute<FText>(this, &SSCS_RowWidget::GetAssetName), false);
 			}
+
+			// If the component is marked as editor only, then display that info here
+			AddToToolTipInfoBox(InfoBox, LOCTEXT("TooltipEditorOnly", "Editor Only"), SNullWidget::NullWidget, TAttribute<FText>::Create(TAttribute<FText>::FGetter::CreateSP(this, &SSCS_RowWidget::GetComponentEditorOnlyTooltipText)), false);
 		}
 	}
 
@@ -2087,6 +2090,29 @@ FText SSCS_RowWidget::GetComponentAddSourceToolTipText() const
 	return NodeType;
 }
 
+FText SSCS_RowWidget::GetComponentEditorOnlyTooltipText() const
+{
+	FText ComponentType = LOCTEXT("ComponentEditorOnlyFalse", "False");
+	if (FSCSEditorTreeNode* TreeNode = GetNode().Get())
+	{
+		if (TreeNode->IsComponentNode())
+		{
+			if (const UActorComponent* Template = TreeNode->GetComponentTemplate())
+			{
+				UBlueprint* Blueprint = GetBlueprint();
+				FObjectProperty* Prop = Blueprint ? FindFProperty<FObjectProperty>(Blueprint->SkeletonGeneratedClass, TreeNode->GetVariableName()) : nullptr;
+
+				if(Template->bIsEditorOnly || (Prop && Prop->HasAnyPropertyFlags(CPF_EditorOnly)))
+				{
+					ComponentType = LOCTEXT("ComponentEditorOnlyTrue", "True");
+				}
+			}
+		}
+	}
+	
+	return ComponentType;
+}
+
 FText SSCS_RowWidget::GetIntroducedInToolTipText() const
 {
 	FText IntroducedInTooltip = LOCTEXT("IntroducedInThisBPTooltip", "this class");
diff --git a/Engine/Source/Editor/Kismet/Public/SSCSEditor.h b/Engine/Source/Editor/Kismet/Public/SSCSEditor.h
--- a/Engine/Source/Editor/Kismet/Public/SSCSEditor.h
+++ b/Engine/Source/Editor/Kismet/Public/SSCSEditor.h
@@ -818,6 +818,13 @@ private:
 	 */
 	FText GetComponentAddSourceToolTipText() const;
 
+	/**
+	 * Retrieves a tooltip text describing if the component is marked Editor only or not
+	 *
+	 * @returns An FText object containing a description of if the component is marked Editor only or not
+	 */
+	FText GetComponentEditorOnlyTooltipText() const;
+
 public:
 	/** Pointer back to owning SCSEditor 2 tool */
 	TWeakPtr<SSCSEditor> SCSEditor;
