Blueprint editor reskin first pass and supporting work

diff --git a/Engine/Source/Editor/BehaviorTreeEditor/Private/SBehaviorTreeBlackboardView.cpp b/Engine/Source/Editor/BehaviorTreeEditor/Private/SBehaviorTreeBlackboardView.cpp
--- a/Engine/Source/Editor/BehaviorTreeEditor/Private/SBehaviorTreeBlackboardView.cpp
+++ b/Engine/Source/Editor/BehaviorTreeEditor/Private/SBehaviorTreeBlackboardView.cpp
@@ -104,7 +104,7 @@ class SBehaviorTreeBlackboardItem : public SGraphPaletteItem
 		GetPaletteItemIcon(GraphAction, IconBrush);
 
 		TSharedRef<SWidget> IconWidget = CreateIconWidget( GraphAction->GetTooltipDescription(), IconBrush, FLinearColor::White );
-		TSharedRef<SWidget> NameSlotWidget = CreateTextSlotWidget( NameFont, InCreateData, BlackboardEntryAction->bIsInherited );
+		TSharedRef<SWidget> NameSlotWidget = CreateTextSlotWidget(InCreateData, BlackboardEntryAction->bIsInherited );
 		TSharedRef<SWidget> DebugSlotWidget = CreateDebugSlotWidget( NameFont );
 
 		// Create the actual widget
@@ -151,7 +151,7 @@ private:
 		}
 	}
 
-	virtual TSharedRef<SWidget> CreateTextSlotWidget( const FSlateFontInfo& NameFont, FCreateWidgetForActionData* const InCreateData, TAttribute<bool> bInIsReadOnly ) override
+	virtual TSharedRef<SWidget> CreateTextSlotWidget(FCreateWidgetForActionData* const InCreateData, TAttribute<bool> bInIsReadOnly ) override
 	{
 		check(InCreateData);
 
@@ -169,7 +169,6 @@ private:
 		InlineRenameWidget =
 			SAssignNew(DisplayWidget, SInlineEditableTextBlock)
 			.Text(this, &SBehaviorTreeBlackboardItem::GetDisplayText)
-			.Font(NameFont)
 			.HighlightText(InCreateData->HighlightText)
 			.ToolTipText(this, &SBehaviorTreeBlackboardItem::GetItemTooltip)
 			.OnTextCommitted(this, &SBehaviorTreeBlackboardItem::OnNameTextCommitted)
diff --git a/Engine/Source/Editor/EditorStyle/Private/StarshipStyle.cpp b/Engine/Source/Editor/EditorStyle/Private/StarshipStyle.cpp
--- a/Engine/Source/Editor/EditorStyle/Private/StarshipStyle.cpp
+++ b/Engine/Source/Editor/EditorStyle/Private/StarshipStyle.cpp
@@ -733,8 +733,8 @@ void FStarshipEditorStyle::FStyle::SetupGeneralStyles()
 #if WITH_EDITOR || (IS_PROGRAM && WITH_UNREAL_DEVELOPER_TOOLS)
 
 	{
-		Set( "SystemWideCommands.FindInContentBrowser", new IMAGE_BRUSH( "Icons/icon_toolbar_genericfinder_40px", Icon40x40 ) );
-		Set( "SystemWideCommands.FindInContentBrowser.Small", new IMAGE_BRUSH( "Icons/icon_toolbar_genericfinder_40px", Icon20x20 ) );
+		Set("SystemWideCommands.FindInContentBrowser", new CORE_IMAGE_BRUSH_SVG("Starship/Common/search", Icon20x20));
+		Set("SystemWideCommands.FindInContentBrowser.Small", new CORE_IMAGE_BRUSH_SVG("Starship/Common/search", Icon16x16));
 	}
 
 	// PList Editor
@@ -3209,7 +3209,7 @@ void FStarshipEditorStyle::FStyle::SetupPropertyEditorStyles()
 		Set( "DetailsView.AdvancedDropdownBorder",      new FSlateColorBrush(FStyleColors::Background));
 		Set( "DetailsView.AdvancedDropdownBorder.Open", new FSlateColorBrush(FStyleColors::Background));
 
-		Set( "DetailsView.CategoryFontStyle", FStyleFonts::Get().SmallBold);
+		Set( "DetailsView.CategoryFontStyle", FStyleFonts::Get().NormalBold);
 		Set( "DetailsView.CategoryTextStyle", FTextBlockStyle(NormalText)
 			.SetFont(GetFontStyle("DetailsView.CategoryFontStyle"))
 			.SetColorAndOpacity(FStyleColors::ForegroundHeader)
@@ -5171,8 +5171,10 @@ void FStarshipEditorStyle::FStyle::SetupPersonaStyle()
 		Set( "Kismet.VariableList.SetTypeIconLarge", new IMAGE_BRUSH( "/Icons/pillset_40x", Icon40x40 ) );
 		Set( "Kismet.VariableList.MapValueTypeIcon", new IMAGE_BRUSH( "/Icons/pillmapvalue_16x", Icon16x16 ) );
 		Set( "Kismet.VariableList.MapKeyTypeIcon", new IMAGE_BRUSH( "/Icons/pillmapkey_16x", Icon16x16 ) );
-		Set( "Kismet.VariableList.ExposeForInstance", new IMAGE_BRUSH( "/Icons/icon_layer_visible_16x", Icon16x16 ) );
-		Set( "Kismet.VariableList.HideForInstance", new IMAGE_BRUSH( "/Icons/icon_layer_not_visible_16x", Icon16x16 ) );
+
+		Set("Kismet.VariableList.ExposeForInstance", new CORE_IMAGE_BRUSH_SVG("Starship/Common/visible", Icon16x16));
+		Set("Kismet.VariableList.HideForInstance", new CORE_IMAGE_BRUSH_SVG("Starship/Common/hidden", Icon16x16));
+	
 		Set( "Kismet.VariableList.VariableIsUsed", new IMAGE_BRUSH( "/Icons/icon_variable_used_16x", Icon16x16 ) );
 		Set( "Kismet.VariableList.VariableNotUsed", new IMAGE_BRUSH( "/Icons/icon_variable_not_used_16x", Icon16x16 ) );
 
diff --git a/Engine/Source/Editor/GraphEditor/Private/SGraphActionMenu.cpp b/Engine/Source/Editor/GraphEditor/Private/SGraphActionMenu.cpp
--- a/Engine/Source/Editor/GraphEditor/Private/SGraphActionMenu.cpp
+++ b/Engine/Source/Editor/GraphEditor/Private/SGraphActionMenu.cpp
@@ -110,7 +110,7 @@ namespace ContextMenuConsoleVariables
 //////////////////////////////////////////////////////////////////////////
 
 template<typename ItemType>
-class SCategoryHeaderTableRow : public STableRow < ItemType >
+class SCategoryHeaderTableRow : public STableRow<ItemType>
 {
 public:
 	SLATE_BEGIN_ARGS(SCategoryHeaderTableRow)
@@ -121,25 +121,23 @@ public:
 	void Construct(const FArguments& InArgs, const TSharedRef<STableViewBase>& InOwnerTableView)
 	{
 		STableRow<ItemType>::ChildSlot
-		.Padding(0.0f, 2.0f, 0.0f, 0.0f)
+		.Padding(0.0f, 2.0f, .0f, 0.0f)
 		[
 			SAssignNew(ContentBorder, SBorder)
 			.BorderImage(this, &SCategoryHeaderTableRow::GetBackgroundImage)
-			.Padding(FMargin(0.0f, 3.0f))
-			.BorderBackgroundColor(FLinearColor(.6, .6, .6, 1.0f))
+			.Padding(FMargin(3.0f, 5.0f))
 			[
 				SNew(SHorizontalBox)
-
 				+ SHorizontalBox::Slot()
 				.VAlign(VAlign_Center)
-				.Padding(2.0f, 2.0f, 2.0f, 2.0f)
+				.Padding(5.0f)
 				.AutoWidth()
 				[
 					SNew(SExpanderArrow, STableRow< ItemType >::SharedThis(this))
 				]
-
 				+ SHorizontalBox::Slot()
 				.VAlign(VAlign_Center)
+				.AutoWidth()
 				[
 					InArgs._Content.Widget
 				]
@@ -158,11 +156,11 @@ public:
 	{
 		if ( STableRow<ItemType>::IsHovered() )
 		{
-			return STableRow<ItemType>::IsItemExpanded() ? FEditorStyle::GetBrush("DetailsView.CategoryTop_Hovered") : FEditorStyle::GetBrush("DetailsView.CollapsedCategory_Hovered");
+			return FAppStyle::Get().GetBrush("Brushes.Secondary");
 		}
 		else
 		{
-			return STableRow<ItemType>::IsItemExpanded() ? FEditorStyle::GetBrush("DetailsView.CategoryTop") : FEditorStyle::GetBrush("DetailsView.CollapsedCategory");
+			return FAppStyle::Get().GetBrush("Brushes.Header");
 		}
 	}
 
@@ -181,6 +179,18 @@ public:
 		return nullptr;
 	}
 
+	FReply OnMouseButtonDown(const FGeometry& MyGeometry, const FPointerEvent& MouseEvent)
+	{
+		if (MouseEvent.GetEffectingButton() == EKeys::LeftMouseButton)
+		{
+			STableRow<ItemType>::ToggleExpansion();
+			return FReply::Handled();
+		}
+		else
+		{
+			return FReply::Unhandled();
+		}
+	}
 private:
 	TSharedPtr<SBorder> ContentBorder;
 };
@@ -1181,10 +1191,7 @@ TSharedRef<ITableRow> SGraphActionMenu::MakeWidget( TSharedPtr<FGraphActionNode>
 	}
 	else
 	{
-		const FTableRowStyle* Style = bUseSectionStyling ? &FEditorStyle::Get().GetWidgetStyle<FTableRowStyle>("TableView.DarkRow") : &FCoreStyle::Get().GetWidgetStyle<FTableRowStyle>("TableView.Row");
-
 		TableRow = SNew(STableRow< TSharedPtr<FGraphActionNode> >, OwnerTable)
-			.Style(Style)
 			.OnDragDetected(this, &SGraphActionMenu::OnItemDragDetected)
 			.ShowSelection(!InItem->IsSeparator());
 	}
@@ -1286,8 +1293,9 @@ TSharedRef<ITableRow> SGraphActionMenu::MakeWidget( TSharedPtr<FGraphActionNode>
 			[
 				SNew(SRichTextBlock)
 				.Text(SectionTitle)
+				.TransformPolicy(ETextTransformPolicy::ToUpper)
 				.DecoratorStyleSet(&FEditorStyle::Get())
-				.TextStyle(FEditorStyle::Get(), "DetailsView.CategoryTextStyle")
+				.TextStyle(FAppStyle::Get(), "DetailsView.CategoryTextStyle")
 			]
 
 			+ SHorizontalBox::Slot()
diff --git a/Engine/Source/Editor/GraphEditor/Private/SGraphPalette.cpp b/Engine/Source/Editor/GraphEditor/Private/SGraphPalette.cpp
--- a/Engine/Source/Editor/GraphEditor/Private/SGraphPalette.cpp
+++ b/Engine/Source/Editor/GraphEditor/Private/SGraphPalette.cpp
@@ -20,8 +20,6 @@
 
 void SGraphPaletteItem::Construct(const FArguments& InArgs, FCreateWidgetForActionData* const InCreateData)
 {
-	FSlateFontInfo NameFont = FCoreStyle::GetDefaultFontStyle("Regular", 10);
-
 	check(InCreateData->Action.IsValid());
 
 	TSharedPtr<FEdGraphSchemaAction> GraphAction = InCreateData->Action;
@@ -34,7 +32,7 @@ void SGraphPaletteItem::Construct(const FArguments& InArgs, FCreateWidgetForActi
 	bool bIsReadOnly = false;
 
 	TSharedRef<SWidget> IconWidget = CreateIconWidget( IconToolTip, IconBrush, IconColor );
-	TSharedRef<SWidget> NameSlotWidget = CreateTextSlotWidget( NameFont, InCreateData, bIsReadOnly );
+	TSharedRef<SWidget> NameSlotWidget = CreateTextSlotWidget(InCreateData, bIsReadOnly );
 
 	// Create the actual widget
 	this->ChildSlot
@@ -116,7 +114,7 @@ TSharedRef<SWidget> SGraphPaletteItem::CreateIconWidget(const FText& IconToolTip
 		IDocumentation::Get()->CreateToolTip(IconToolTip, NULL, DocLink, DocExcerpt));
 }
 
-TSharedRef<SWidget> SGraphPaletteItem::CreateTextSlotWidget( const FSlateFontInfo& NameFont, FCreateWidgetForActionData* const InCreateData, TAttribute<bool> bIsReadOnly )
+TSharedRef<SWidget> SGraphPaletteItem::CreateTextSlotWidget(FCreateWidgetForActionData* const InCreateData, TAttribute<bool> bIsReadOnly )
 {
 	TSharedPtr< SWidget > DisplayWidget;
 
@@ -135,7 +133,6 @@ TSharedRef<SWidget> SGraphPaletteItem::CreateTextSlotWidget( const FSlateFontInf
 	InlineRenameWidget =
 		SAssignNew(DisplayWidget, SInlineEditableTextBlock)
 		.Text(this, &SGraphPaletteItem::GetDisplayText)
-		.Font(NameFont)
 		.HighlightText(InCreateData->HighlightText)
 		.ToolTipText(this, &SGraphPaletteItem::GetItemTooltip)
 		.OnTextCommitted(this, &SGraphPaletteItem::OnNameTextCommitted)
diff --git a/Engine/Source/Editor/GraphEditor/Public/SGraphPalette.h b/Engine/Source/Editor/GraphEditor/Public/SGraphPalette.h
--- a/Engine/Source/Editor/GraphEditor/Public/SGraphPalette.h
+++ b/Engine/Source/Editor/GraphEditor/Public/SGraphPalette.h
@@ -43,7 +43,7 @@ protected:
 	TSharedRef<SWidget> CreateIconWidget(const FText& IconToolTip, const FSlateBrush* IconBrush, const FSlateColor& IconColor);
 
 	/* Create the text widget */
-	virtual TSharedRef<SWidget> CreateTextSlotWidget( const FSlateFontInfo& NameFont, FCreateWidgetForActionData* const InCreateData, TAttribute<bool> bIsReadOnly );
+	virtual TSharedRef<SWidget> CreateTextSlotWidget( FCreateWidgetForActionData* const InCreateData, TAttribute<bool> bIsReadOnly );
 
 	/** Callback when rename text is being verified on text changed */
 	virtual bool OnNameTextVerifyChanged(const FText& InNewText, FText& OutErrorMessage);
diff --git a/Engine/Source/Editor/Kismet/Private/BlueprintEditorModes.cpp b/Engine/Source/Editor/Kismet/Private/BlueprintEditorModes.cpp
--- a/Engine/Source/Editor/Kismet/Private/BlueprintEditorModes.cpp
+++ b/Engine/Source/Editor/Kismet/Private/BlueprintEditorModes.cpp
@@ -32,18 +32,11 @@ const FName FBlueprintEditorApplicationModes::BlueprintMacroMode( TEXT("MacroNam
 
 TSharedPtr<FTabManager::FLayout> GetDefaltEditorLayout(TSharedPtr<class FBlueprintEditor> InBlueprintEditor)
 {
-	return FTabManager::NewLayout( "Standalone_BlueprintEditor_Layout_v6" )
+	return FTabManager::NewLayout( "Standalone_BlueprintEditor_Layout_v7" )
 	->AddArea
 	(
 		FTabManager::NewPrimaryArea() ->SetOrientation(Orient_Vertical)
 		->Split
-		(
-			FTabManager::NewStack()
-			->SetSizeCoefficient( 0.186721f )
-			->SetHideTabWell(true)
-			->AddTab( InBlueprintEditor->GetToolbarTabId(), ETabState::OpenedTab )
-		)
-		->Split
 		(
 			FTabManager::NewSplitter() ->SetOrientation(Orient_Horizontal)
 			->Split
@@ -574,7 +567,7 @@ FBlueprintEditorUnifiedMode::FBlueprintEditorUnifiedMode(TSharedPtr<class FBluep
 
 	if ( bRegisterViewport )
 	{
-		TabLayout = FTabManager::NewLayout( "Blueprints_Unified_Components_v6" )
+		TabLayout = FTabManager::NewLayout( "Blueprints_Unified_Components_v7" )
 		->AddArea
 		(
 			FTabManager::NewPrimaryArea()->SetOrientation(Orient_Vertical)
@@ -603,13 +596,6 @@ FBlueprintEditorUnifiedMode::FBlueprintEditorUnifiedMode(TSharedPtr<class FBluep
 					FTabManager::NewSplitter()->SetOrientation( Orient_Vertical )
 					->SetSizeCoefficient(0.60f)
 					->Split
-					(
-						FTabManager::NewStack()
-						->SetSizeCoefficient(0.186721f)
-						->SetHideTabWell(true)
-						->AddTab(InBlueprintEditor->GetToolbarTabId(), ETabState::OpenedTab)
-					)
-					->Split
 					(
 						FTabManager::NewStack()
 						->SetSizeCoefficient( 0.80f )
@@ -642,7 +628,7 @@ FBlueprintEditorUnifiedMode::FBlueprintEditorUnifiedMode(TSharedPtr<class FBluep
 	}
 	else
 	{
-		TabLayout = FTabManager::NewLayout( "Blueprints_Unified_v4" )
+		TabLayout = FTabManager::NewLayout( "Blueprints_Unified_v5" )
 		->AddArea
 		(
 		FTabManager::NewPrimaryArea()->SetOrientation(Orient_Vertical)
@@ -665,13 +651,6 @@ FBlueprintEditorUnifiedMode::FBlueprintEditorUnifiedMode(TSharedPtr<class FBluep
 					FTabManager::NewSplitter()->SetOrientation( Orient_Vertical )
 					->SetSizeCoefficient(0.60f)
 					->Split
-					(
-						FTabManager::NewStack()
-						->SetSizeCoefficient(0.186721f)
-						->SetHideTabWell(true)
-						->AddTab(InBlueprintEditor->GetToolbarTabId(), ETabState::OpenedTab)
-					)
-					->Split
 					(
 						FTabManager::NewStack()
 						->SetSizeCoefficient( 0.80f )
diff --git a/Engine/Source/Editor/Kismet/Private/Debugging/SKismetDebuggingView.cpp b/Engine/Source/Editor/Kismet/Private/Debugging/SKismetDebuggingView.cpp
--- a/Engine/Source/Editor/Kismet/Private/Debugging/SKismetDebuggingView.cpp
+++ b/Engine/Source/Editor/Kismet/Private/Debugging/SKismetDebuggingView.cpp
@@ -1066,7 +1066,7 @@ void SKismetDebuggingView::Construct(const FArguments& InArgs)
 	static const FName ToolbarName = "Kismet.DebuggingViewToolBar";
 	if (!UToolMenus::Get()->IsMenuRegistered(ToolbarName))
 	{
-		UToolMenu* ToolBar = UToolMenus::Get()->RegisterMenu(ToolbarName, NAME_None, EMultiBoxType::ToolBar);
+		UToolMenu* ToolBar = UToolMenus::Get()->RegisterMenu(ToolbarName, NAME_None, EMultiBoxType::SlimHorizontalToolBar);
 
 		{
 			FToolMenuSection& Section = ToolBar->AddSection("Debug");
diff --git a/Engine/Source/Editor/Kismet/Private/SBlueprintEditorSelectedDebugObjectWidget.cpp b/Engine/Source/Editor/Kismet/Private/SBlueprintEditorSelectedDebugObjectWidget.cpp
--- a/Engine/Source/Editor/Kismet/Private/SBlueprintEditorSelectedDebugObjectWidget.cpp
+++ b/Engine/Source/Editor/Kismet/Private/SBlueprintEditorSelectedDebugObjectWidget.cpp
@@ -113,7 +113,6 @@ const FString& SBlueprintEditorSelectedDebugObjectWidget::GetDebugAllWorldsStrin
 
 TSharedRef<SWidget> SBlueprintEditorSelectedDebugObjectWidget::OnGetActiveDetailSlotContent(bool bChangedToHighDetail)
 {
-
 	const TSharedRef<SWidget> BrowseButton = PropertyCustomizationHelpers::MakeBrowseButton(FSimpleDelegate::CreateSP(this, &SBlueprintEditorSelectedDebugObjectWidget::SelectedDebugObject_OnClicked));
 	BrowseButton->SetVisibility(TAttribute<EVisibility>::Create(TAttribute<EVisibility>::FGetter::CreateSP(this, &SBlueprintEditorSelectedDebugObjectWidget::IsSelectDebugObjectButtonVisible)));
 	BrowseButton->SetToolTipText(LOCTEXT("DebugSelectActor", "Select this Actor in level"));
@@ -134,51 +133,23 @@ TSharedRef<SWidget> SBlueprintEditorSelectedDebugObjectWidget::OnGetActiveDetail
 			BrowseButton
 		];
 
-	if (!bChangedToHighDetail)
-	{
+
 		return
 			SNew(SHorizontalBox)
 			+ SHorizontalBox::Slot()
+			.VAlign(VAlign_Center)
+			.Padding(5.0f)
 			.AutoWidth()
 			[
 				DebugWorldsComboBox.ToSharedRef()
 			]
 			+ SHorizontalBox::Slot()
+			.VAlign(VAlign_Center)
+			.Padding(0.0f, 5.0f)
 			.AutoWidth()
 			[
 				DebugObjectSelectionWidget
 			];
-	}
-	else
-	{
-		return
-			SNew(SVerticalBox)
-			+ SVerticalBox::Slot()
-			.AutoHeight()
-			.VAlign(VAlign_Bottom)
-			[
-				// Vertical Layout when using normal size icons
-				SNew(SVerticalBox)
-				+ SVerticalBox::Slot()
-				.AutoHeight()
-				[
-					DebugWorldsComboBox.ToSharedRef()
-				]
-				+ SVerticalBox::Slot()
-				.AutoHeight()
-				[
-					DebugObjectSelectionWidget
-				]
-			]
-			+ SVerticalBox::Slot()
-			.AutoHeight()
-			.HAlign(HAlign_Center)
-			.Padding(2.0f)
-			[
-				SNew(STextBlock)
-				.Text(LOCTEXT("DebugSelectTitle", "Debug Filter"))
-			];
-	}
 }
 
 void SBlueprintEditorSelectedDebugObjectWidget::OnRefresh()
diff --git a/Engine/Source/Editor/Kismet/Private/SBlueprintPalette.cpp b/Engine/Source/Editor/Kismet/Private/SBlueprintPalette.cpp
--- a/Engine/Source/Editor/Kismet/Private/SBlueprintPalette.cpp
+++ b/Engine/Source/Editor/Kismet/Private/SBlueprintPalette.cpp
@@ -812,30 +812,23 @@ public:
 			bShouldHaveAVisibilityToggle = bIsBlueprintVariable && (!bIsComponentVar || FBlueprintEditorUtils::IsVariableCreatedByBlueprint(BlueprintObj, VariableObjProp));
 		}
 
-		this->ChildSlot[
+		this->ChildSlot
+		[
 			SNew(SBorder)
-				.Padding( 0.0f )
-				.BorderImage(FEditorStyle::GetBrush("NoBorder"))
-				.ColorAndOpacity(this, &SPaletteItemVisibilityToggle::GetVisibilityToggleColor)
+			.Padding(0.0f)
+			.BorderImage(FStyleDefaults::GetNoBrush())
+			.Visibility(bShouldHaveAVisibilityToggle ? EVisibility::Visible : EVisibility::Collapsed)
+			//.ForegroundColor(this, &SPaletteItemVisibilityToggle::GetVisibilityToggleColor)
 			[
-				SNew( SCheckBox )
-					.ToolTipText(this, &SPaletteItemVisibilityToggle::GetVisibilityToggleToolTip)
-					.Visibility(bShouldHaveAVisibilityToggle ? EVisibility::Visible : EVisibility::Collapsed)
-					.OnCheckStateChanged(this, &SPaletteItemVisibilityToggle::OnVisibilityToggleFlipped)
-					.IsChecked(this, &SPaletteItemVisibilityToggle::GetVisibilityToggleState)
-					// a style using the normal checkbox images but with the toggle button layout
-					.Style( FEditorStyle::Get(), "CheckboxLookToggleButtonCheckbox")	
+				SNew(SCheckBox)
+				.ToolTipText(this, &SPaletteItemVisibilityToggle::GetVisibilityToggleToolTip)
+				.OnCheckStateChanged(this, &SPaletteItemVisibilityToggle::OnVisibilityToggleFlipped)
+				.IsChecked(this, &SPaletteItemVisibilityToggle::GetVisibilityToggleState)
+				.Style(FAppStyle::Get(), "TransparentCheckBox")
 				[
-					SNew( SVerticalBox )
-					+SVerticalBox::Slot()
-						.AutoHeight()
-						.VAlign( VAlign_Center )
-						.HAlign( HAlign_Center )
-					[
-						SNew( SImage )
-							.Image( this, &SPaletteItemVisibilityToggle::GetVisibilityIcon )
-							.ColorAndOpacity( FLinearColor::Black )
-					]
+					SNew(SImage)
+					.Image(this, &SPaletteItemVisibilityToggle::GetVisibilityIcon)
+					.ColorAndOpacity(FSlateColor::UseForeground())
 				]
 			]
 		];
@@ -909,11 +902,11 @@ private:
 	 * 
 	 * @return A color denoting the item's visibility and tootip status.
 	 */
-	FLinearColor GetVisibilityToggleColor() const 
+	FSlateColor GetVisibilityToggleColor() const 
 	{
 		if ( GetVisibilityToggleState() != ECheckBoxState::Checked )
 		{
-			return FColor(64, 64, 64).ReinterpretAsLinear();
+			return FSlateColor::UseForeground();
 		}
 		else
 		{
@@ -924,11 +917,13 @@ private:
 
 			if ( !Result.IsEmpty() )
 			{
-				return FColor(130, 219, 119).ReinterpretAsLinear(); //pastel green when tooltip exists
+				static const FName TooltipExistsColor("Colors.AccentGreen");
+				return FAppStyle::Get().GetSlateColor(TooltipExistsColor);
 			}
 			else
 			{
-				return FColor(215, 219, 119).ReinterpretAsLinear(); //pastel yellow if no tooltip to alert designer 
+				static const FName TooltipDoesntExistColor("Colors.AccentYellow");
+				return FAppStyle::Get().GetSlateColor(TooltipDoesntExistColor);
 			}
 		}
 	}
@@ -1057,8 +1052,7 @@ void SBlueprintPaletteItem::Construct(const FArguments& InArgs, FCreateWidgetFor
 		TagMeta.FriendlyName = GraphAction->GetMenuDescription().ToString();
 	}
 	// construct the text widget
-	FSlateFontInfo NameFont = FCoreStyle::GetDefaultFontStyle("Regular", 10);
-	TSharedRef<SWidget> NameSlotWidget = CreateTextSlotWidget( NameFont, InCreateData, bIsReadOnly );
+	TSharedRef<SWidget> NameSlotWidget = CreateTextSlotWidget(InCreateData, bIsReadOnly );
 	
 	// Will set the icon of this property to be a Pin Type selector. 
 	auto GenerateVariableSettings = [&](FProperty* VariableProp)
@@ -1163,6 +1157,18 @@ void SBlueprintPaletteItem::Construct(const FArguments& InArgs, FCreateWidgetFor
 
 	ActionBox.Get().AddSlot()
 		.AutoWidth()
+		.Padding(FMargin(0.0f, 0.0f, 3.0f, 0.0f))
+		.HAlign(HAlign_Left)
+		.VAlign(VAlign_Center)
+		[
+			SNew(SPaletteItemVisibilityToggle, ActionPtr, InBlueprintEditor, InBlueprint)
+			.IsEnabled(bIsEditingEnabled)
+		];
+
+
+	ActionBox.Get().AddSlot()
+		.AutoWidth()
+		.VAlign(VAlign_Center)
 		[
 			IconWidget
 		];
@@ -1195,15 +1201,6 @@ void SBlueprintPaletteItem::Construct(const FArguments& InArgs, FCreateWidgetFor
 			NameSlotWidget
 		];
 
-	ActionBox.Get().AddSlot()
-		.AutoWidth()
-		.Padding(FMargin(3.0f, 0.0f))
-		.VAlign(VAlign_Center)
-		[
-			SNew(SPaletteItemVisibilityToggle, ActionPtr, InBlueprintEditor, InBlueprint)
-			.IsEnabled(bIsEditingEnabled)
-		];
-
 	// Now, create the actual widget
 	ChildSlot
 	[
@@ -1226,7 +1223,7 @@ void SBlueprintPaletteItem::OnDragEnter(const FGeometry& MyGeometry, const FDrag
 *******************************************************************************/
 
 //------------------------------------------------------------------------------
-TSharedRef<SWidget> SBlueprintPaletteItem::CreateTextSlotWidget(const FSlateFontInfo& NameFont, FCreateWidgetForActionData* const InCreateData, TAttribute<bool> bIsReadOnlyIn)
+TSharedRef<SWidget> SBlueprintPaletteItem::CreateTextSlotWidget(FCreateWidgetForActionData* const InCreateData, TAttribute<bool> bIsReadOnlyIn)
 {
 	FName const ActionTypeId = InCreateData->Action->GetTypeId();
 
@@ -1276,7 +1273,6 @@ TSharedRef<SWidget> SBlueprintPaletteItem::CreateTextSlotWidget(const FSlateFont
 		[
 			SAssignNew(EditableTextElement, SInlineEditableTextBlock)
 				.Text(this, &SBlueprintPaletteItem::GetDisplayText)
-				.Font(NameFont)
 				.HighlightText(InCreateData->HighlightText)
 				.ToolTip(ToolTipWidget)
 				.OnVerifyTextChanged(OnVerifyTextChanged)
diff --git a/Engine/Source/Editor/Kismet/Private/SBlueprintPalette.h b/Engine/Source/Editor/Kismet/Private/SBlueprintPalette.h
--- a/Engine/Source/Editor/Kismet/Private/SBlueprintPalette.h
+++ b/Engine/Source/Editor/Kismet/Private/SBlueprintPalette.h
@@ -48,7 +48,7 @@ private:
 	// End of SWidget Interface
 
 	// SGraphPaletteItem Interface
-	virtual TSharedRef<SWidget> CreateTextSlotWidget( const FSlateFontInfo& NameFont,  FCreateWidgetForActionData* const InCreateData, TAttribute<bool> bIsReadOnly ) override;
+	virtual TSharedRef<SWidget> CreateTextSlotWidget(FCreateWidgetForActionData* const InCreateData, TAttribute<bool> bIsReadOnly ) override;
 	virtual FText GetDisplayText() const override;
 	virtual bool OnNameTextVerifyChanged(const FText& InNewText, FText& OutErrorMessage) override;
 	virtual void OnNameTextCommitted(const FText& NewText, ETextCommit::Type InTextCommit) override;
diff --git a/Engine/Source/Editor/Kismet/Private/SMyBlueprint.cpp b/Engine/Source/Editor/Kismet/Private/SMyBlueprint.cpp
--- a/Engine/Source/Editor/Kismet/Private/SMyBlueprint.cpp
+++ b/Engine/Source/Editor/Kismet/Private/SMyBlueprint.cpp
@@ -367,35 +367,35 @@ void SMyBlueprint::Construct(const FArguments& InArgs, TWeakPtr<FBlueprintEditor
 	TSharedPtr<SWidget> AddNewMenu = SNullWidget::NullWidget;
 
 	AddNewMenu = SNew(SComboButton)
-		.ComboButtonStyle(FEditorStyle::Get(), "ToolbarComboButton")
-		.ButtonStyle(FEditorStyle::Get(), "FlatButton.Success")
-		.ForegroundColor(FLinearColor::White)
-		.ToolTipText(LOCTEXT("AddNewToolTip", "Add a new Variable, Graph, Function, Macro, or Event Dispatcher."))
+		.ComboButtonStyle(&FAppStyle::Get().GetWidgetStyle<FComboButtonStyle>("SimpleComboButton"))
+		.ForegroundColor(FSlateColor::UseStyle())
+		.ContentPadding(2)
 		.OnGetMenuContent(this, &SMyBlueprint::CreateAddNewMenuWidget)
-		.HasDownArrow(true)
-		.ContentPadding(FMargin(1, 0, 2, 0))
+		.HasDownArrow(false)
 		.AddMetaData<FTagMetaData>(FTagMetaData(TEXT("MyBlueprintAddNewCombo")))
 		.IsEnabled(this, &SMyBlueprint::IsEditingMode)
+		.ToolTipText(LOCTEXT("AddNewToolTip", "Add a new Variable, Graph, Function, Macro, or Event Dispatcher."))
 		.ButtonContent()
 		[
 			SNew(SHorizontalBox)
-
 			+ SHorizontalBox::Slot()
-		.AutoWidth()
-		.Padding(FMargin(0, 1))
-		[
-			SNew(SImage)
-			.Image(FEditorStyle::GetBrush("Plus"))
-		]
-
-	+ SHorizontalBox::Slot()
-		.VAlign(VAlign_Center)
-		.AutoWidth()
-		.Padding(FMargin(2, 0, 2, 0))
-		[
-			SNew(STextBlock)
-			.Text(LOCTEXT("AddNew", "Add New"))
-		]
+			.AutoWidth()
+			.HAlign(HAlign_Center)
+			.VAlign(VAlign_Center)
+			[
+				SNew(SImage)
+				.Image(FAppStyle::Get().GetBrush("Icons.Plus"))
+				.ColorAndOpacity(FSlateColor::UseForeground())
+			]
+			+ SHorizontalBox::Slot()
+			.Padding(FMargin(3, 0, 0, 0))
+			.VAlign(VAlign_Center)
+			.AutoWidth()
+			[
+				SNew(STextBlock)
+				.TextStyle(FAppStyle::Get(), "NormalText")
+				.Text(LOCTEXT("NewAssetButton", "New"))
+			]
 		];
 
 	FMenuBuilder ViewOptions(true, nullptr);
@@ -537,19 +537,19 @@ void SMyBlueprint::Construct(const FArguments& InArgs, TWeakPtr<FBlueprintEditor
 					.Padding(2, 0, 0, 0)
 					[
 						SNew(SComboButton)
-						.ComboButtonStyle(FEditorStyle::Get(), "ToolbarComboButton")
-						.ForegroundColor(FSlateColor::UseForeground())
-						.HasDownArrow(true)
-						.ContentPadding(FMargin(1, 0))
+						.ContentPadding(0)
+						.ComboButtonStyle(&FAppStyle::Get().GetWidgetStyle<FComboButtonStyle>("SimpleComboButton"))
+						.HasDownArrow(false)
 						.AddMetaData<FTagMetaData>(FTagMetaData(TEXT("ViewOptions")))
-						.MenuContent()
-						[
-							ViewOptions.MakeWidget()
-						]
 						.ButtonContent()
 						[
 							SNew(SImage)
-							.Image(FEditorStyle::GetBrush("GenericViewButton"))
+							.ColorAndOpacity(FSlateColor::UseForeground())
+							.Image(FAppStyle::Get().GetBrush("Icons.Settings"))
+						]
+						.MenuContent()
+						[
+							ViewOptions.MakeWidget()
 						]
 					]
 				]
@@ -767,23 +767,22 @@ TSharedRef<SWidget> SMyBlueprint::OnGetSectionWidget(TSharedRef<SWidget> RowWidg
 					SAssignNew(FunctionSectionButton, SComboButton)
 					.IsEnabled(this, &SMyBlueprint::IsEditingMode)
 					.Visibility(this, &SMyBlueprint::OnGetSectionTextVisibility, WeakRowWidget, InSectionID)
-					.ButtonStyle(FEditorStyle::Get(), "RoundButton")
 					.ForegroundColor(FEditorStyle::GetSlateColor("DefaultForeground"))
-					.ContentPadding(FMargin(2, 0))
 					.OnGetMenuContent(this, &SMyBlueprint::OnGetFunctionListMenu)
+					.ContentPadding(0)
 					.HasDownArrow(true)
 					.ButtonContent()
 					[
 						SNew(STextBlock)
 						.Font(IDetailLayoutBuilder::GetDetailFontBold())
 						.Text(LOCTEXT("Override", "Override"))
-						.ShadowOffset(FVector2D(1, 1))
 					]
 				]
 
 				+ SHorizontalBox::Slot()
 				.AutoWidth()
-				.Padding(2, 0,0,0)
+				.VAlign(VAlign_Center)
+				.Padding(2,0,0,0)
 				[
 					CreateAddToSectionButton(InSectionID, WeakRowWidget, AddNewText, MetaDataTag)
 				];
@@ -819,37 +818,18 @@ TSharedRef<SWidget> SMyBlueprint::OnGetSectionWidget(TSharedRef<SWidget> RowWidg
 
 TSharedRef<SWidget> SMyBlueprint::CreateAddToSectionButton(int32 InSectionID, TWeakPtr<SWidget> WeakRowWidget, FText AddNewText, FName MetaDataTag)
 {
-	return SNew(SButton)
-		.ButtonStyle(FEditorStyle::Get(), "RoundButton")
-		.ForegroundColor(FEditorStyle::GetSlateColor("DefaultForeground"))
-		.ContentPadding(FMargin(2, 0))
+	return 
+		SNew(SButton)
+		.ButtonStyle(FEditorStyle::Get(), "SimpleButton")
 		.OnClicked(this, &SMyBlueprint::OnAddButtonClickedOnSection, InSectionID)
 		.IsEnabled(this, &SMyBlueprint::CanAddNewElementToSection, InSectionID)
-		.HAlign(HAlign_Center)
-		.VAlign(VAlign_Center)
+		.ContentPadding(FMargin(1, 0))
 		.AddMetaData<FTagMetaData>(FTagMetaData(MetaDataTag))
+		.ToolTipText(AddNewText)
 		[
-			SNew(SHorizontalBox)
-
-			+ SHorizontalBox::Slot()
-			.AutoWidth()
-			.Padding(FMargin(0, 1))
-			[
-				SNew(SImage)
-				.Image(FEditorStyle::GetBrush("Plus"))
-			]
-
-			+ SHorizontalBox::Slot()
-			.VAlign(VAlign_Center)
-			.AutoWidth()
-			.Padding(FMargin(2,0,0,0))
-			[
-				SNew(STextBlock)
-				.Font(IDetailLayoutBuilder::GetDetailFontBold())
-				.Text(AddNewText)
-				.Visibility(this, &SMyBlueprint::OnGetSectionTextVisibility, WeakRowWidget, InSectionID)
-				.ShadowOffset(FVector2D(1,1))
-			]
+			SNew(SImage)
+			.Image(FAppStyle::Get().GetBrush("Icons.PlusCircle"))
+			.ColorAndOpacity(FSlateColor::UseForeground())
 		];
 }
 
@@ -933,7 +913,7 @@ EVisibility SMyBlueprint::OnGetSectionTextVisibility(TWeakPtr<SWidget> RowWidget
 	}
 	else
 	{
-		return EVisibility::Collapsed;
+		return EVisibility::Hidden;
 	}
 }
 
diff --git a/Engine/Source/Editor/Kismet/Public/BlueprintEditor.h b/Engine/Source/Editor/Kismet/Public/BlueprintEditor.h
--- a/Engine/Source/Editor/Kismet/Public/BlueprintEditor.h
+++ b/Engine/Source/Editor/Kismet/Public/BlueprintEditor.h
@@ -195,6 +195,7 @@ public:
 public:
 	//~ Begin FAssetEditorToolkit Interface
 	virtual bool OnRequestClose() override;
+	virtual bool UsesCustomToolbarPlacement() const override { return false; }
 	// End of FAssetEditorToolkit 
 
 	//~ Begin IToolkit Interface
diff --git a/Engine/Source/Editor/MaterialEditor/Private/SMaterialPalette.cpp b/Engine/Source/Editor/MaterialEditor/Private/SMaterialPalette.cpp
--- a/Engine/Source/Editor/MaterialEditor/Private/SMaterialPalette.cpp
+++ b/Engine/Source/Editor/MaterialEditor/Private/SMaterialPalette.cpp
@@ -20,8 +20,6 @@
 
 void SMaterialPaletteItem::Construct(const FArguments& InArgs, FCreateWidgetForActionData* const InCreateData)
 {
-	FSlateFontInfo NameFont = FCoreStyle::GetDefaultFontStyle("Regular", 10);
-
 	check(InCreateData->Action.IsValid());
 
 	TSharedPtr<FEdGraphSchemaAction> GraphAction = InCreateData->Action;
@@ -50,8 +48,8 @@ void SMaterialPaletteItem::Construct(const FArguments& InArgs, FCreateWidgetForA
 	bool bIsReadOnly = false;
 
 	TSharedRef<SWidget> IconWidget = CreateIconWidget( IconToolTip, IconBrush, IconColor );
-	TSharedRef<SWidget> NameSlotWidget = CreateTextSlotWidget( NameFont, InCreateData, bIsReadOnly );
-	TSharedRef<SWidget> HotkeyDisplayWidget = CreateHotkeyDisplayWidget( NameFont, HotkeyChord );
+	TSharedRef<SWidget> NameSlotWidget = CreateTextSlotWidget(InCreateData, bIsReadOnly );
+	TSharedRef<SWidget> HotkeyDisplayWidget = CreateHotkeyDisplayWidget(HotkeyChord );
 
 	// Create the actual widget
 	this->ChildSlot
@@ -81,7 +79,7 @@ void SMaterialPaletteItem::Construct(const FArguments& InArgs, FCreateWidgetForA
 	];
 }
 
-TSharedRef<SWidget> SMaterialPaletteItem::CreateHotkeyDisplayWidget(const FSlateFontInfo& NameFont, const TSharedPtr<const FInputChord> HotkeyChord)
+TSharedRef<SWidget> SMaterialPaletteItem::CreateHotkeyDisplayWidget(const TSharedPtr<const FInputChord> HotkeyChord)
 {
 	FText HotkeyText;
 	if (HotkeyChord.IsValid())
@@ -89,8 +87,7 @@ TSharedRef<SWidget> SMaterialPaletteItem::CreateHotkeyDisplayWidget(const FSlate
 		HotkeyText = HotkeyChord->GetInputText();
 	}
 	return SNew(STextBlock)
-		.Text(HotkeyText)
-		.Font(NameFont);
+		.Text(HotkeyText);
 }
 
 FText SMaterialPaletteItem::GetItemTooltip() const
diff --git a/Engine/Source/Editor/MaterialEditor/Private/SMaterialPalette.h b/Engine/Source/Editor/MaterialEditor/Private/SMaterialPalette.h
--- a/Engine/Source/Editor/MaterialEditor/Private/SMaterialPalette.h
+++ b/Engine/Source/Editor/MaterialEditor/Private/SMaterialPalette.h
@@ -22,7 +22,7 @@ public:
 
 private:
 	/* Create the hotkey display widget */
-	TSharedRef<SWidget> CreateHotkeyDisplayWidget(const FSlateFontInfo& NameFont, const TSharedPtr<const FInputChord> HotkeyChord);
+	TSharedRef<SWidget> CreateHotkeyDisplayWidget(const TSharedPtr<const FInputChord> HotkeyChord);
 
 	virtual FText GetItemTooltip() const override;
 };
diff --git a/Engine/Source/Editor/UnrealEd/Private/Toolkits/AssetEditorToolkit.cpp b/Engine/Source/Editor/UnrealEd/Private/Toolkits/AssetEditorToolkit.cpp
--- a/Engine/Source/Editor/UnrealEd/Private/Toolkits/AssetEditorToolkit.cpp
+++ b/Engine/Source/Editor/UnrealEd/Private/Toolkits/AssetEditorToolkit.cpp
@@ -42,6 +42,9 @@ TSharedPtr<FExtensibilityManager> FAssetEditorToolkit::SharedToolBarExtensibilit
 
 const FName FAssetEditorToolkit::ToolbarTabId( TEXT( "AssetEditorToolkit_Toolbar" ) );
 
+const FName FAssetEditorToolkit::DefaultAssetEditorToolBarName("AssetEditor.DefaultToolBar");
+const FName FAssetEditorToolkit::DefaultAssetEditorSlimToolBarName("AssetEditor.DefaultSlimToolBar");
+
 FAssetEditorToolkit::FAssetEditorToolkit()
 	: GCEditingObjects(*this)
 	, bCheckDirtyOnAssetSave(false)
@@ -269,6 +272,11 @@ void FAssetEditorToolkit::InitAssetEditor( const EToolkitMode::Type Mode, const
 		Toolbar = SNullWidget::NullWidget;
 	}
 
+	if (NewStandaloneHost && !UsesCustomToolbarPlacement())
+	{
+		NewStandaloneHost->SetToolbar(Toolbar);
+	}
+
 	// Create our mode manager and set it's toolkit host
 	CreateEditorModeManager();
 	if (EditorModeManager.IsValid())
@@ -1024,18 +1032,18 @@ FName FAssetEditorToolkit::GetToolMenuToolbarName() const
 
 FName FAssetEditorToolkit::GetToolMenuToolbarName(FName& OutParentName) const
 {
-	static const FName DefaultToolbarName = "AssetEditor.DefaultToolBar";
-	OutParentName = DefaultToolbarName;
+	OutParentName = UsesCustomToolbarPlacement() ? DefaultAssetEditorToolBarName : DefaultAssetEditorSlimToolBarName;
+
 	return *(TEXT("AssetEditor.") + GetToolMenuAppName().ToString() + TEXT(".ToolBar"));
 }
 
-void FAssetEditorToolkit::RegisterDefaultToolBar()
+void FAssetEditorToolkit::RegisterDefaultToolBar(bool bCustomToolbarPlacement)
 {
-	static const FName DefaultToolBarName("AssetEditor.DefaultToolBar");
+	FName DefaultToolBarName = bCustomToolbarPlacement ? DefaultAssetEditorToolBarName : DefaultAssetEditorSlimToolBarName;
 	UToolMenus* ToolMenus = UToolMenus::Get();
 	if (!ToolMenus->IsMenuRegistered(DefaultToolBarName))
 	{
-		UToolMenu* ToolbarBuilder = ToolMenus->RegisterMenu(DefaultToolBarName, NAME_None, EMultiBoxType::ToolBar);
+		UToolMenu* ToolbarBuilder = ToolMenus->RegisterMenu(DefaultToolBarName, NAME_None, bCustomToolbarPlacement ? EMultiBoxType::ToolBar : EMultiBoxType::SlimHorizontalToolBar);
 		{
 			FToolMenuSection& Section = ToolbarBuilder->AddSection("Asset");
 			Section.AddEntry(FToolMenuEntry::InitToolBarButton(FAssetEditorCommonCommands::Get().SaveAsset));
@@ -1053,7 +1061,7 @@ void FAssetEditorToolkit::GenerateToolbar()
 {
 	TSharedPtr<FExtender> Extender = FExtender::Combine(ToolbarExtenders);
 
-	RegisterDefaultToolBar();
+	RegisterDefaultToolBar(UsesCustomToolbarPlacement());
 
 	FName ParentToolbarName;
 	const FName ToolBarName = GetToolMenuToolbarName(ParentToolbarName);
@@ -1061,7 +1069,7 @@ void FAssetEditorToolkit::GenerateToolbar()
 	UToolMenu* FoundMenu = ToolMenus->FindMenu(ToolBarName);
 	if (!FoundMenu || !FoundMenu->IsRegistered())
 	{
-		FoundMenu = ToolMenus->RegisterMenu(ToolBarName, ParentToolbarName, EMultiBoxType::ToolBar);
+		FoundMenu = ToolMenus->RegisterMenu(ToolBarName, ParentToolbarName, UsesCustomToolbarPlacement() ? EMultiBoxType::ToolBar : EMultiBoxType::SlimHorizontalToolBar);
 	}
 
 	FToolMenuContext MenuContext(GetToolkitCommands(), Extender);
@@ -1100,16 +1108,9 @@ void FAssetEditorToolkit::GenerateToolbar()
 	Toolbar = 
 		SNew(SHorizontalBox)
 		+SHorizontalBox::Slot()
-		.HAlign(HAlign_Left)
 		.VAlign(VAlign_Center)
 		[
-			SNew(SVerticalBox)
-			+SVerticalBox::Slot()
-			.AutoHeight()
-			.VAlign(VAlign_Bottom)
-			[
-				ToolBarWidget
-			]
+			ToolBarWidget
 		]
 		+SHorizontalBox::Slot()
 		.HAlign(HAlign_Right)
@@ -1135,7 +1136,8 @@ void FAssetEditorToolkit::RegenerateMenusAndToolbars()
 {
 	RemoveAllToolbarWidgets();
 
-	StandaloneHost.Pin()->GenerateMenus(false);
+	TSharedPtr<SStandaloneAssetEditorToolkitHost> HostPinned = StandaloneHost.Pin();
+	HostPinned->GenerateMenus(false);
 
 	if (Toolbar != SNullWidget::NullWidget)
 	{
@@ -1143,6 +1145,12 @@ void FAssetEditorToolkit::RegenerateMenusAndToolbars()
 	}
 
 	PostRegenerateMenusAndToolbars();
+
+	if (!UsesCustomToolbarPlacement())
+	{
+		HostPinned->SetToolbar(Toolbar);
+	}
+
 }
 
 
diff --git a/Engine/Source/Editor/UnrealEd/Private/Toolkits/SStandaloneAssetEditorToolkitHost.cpp b/Engine/Source/Editor/UnrealEd/Private/Toolkits/SStandaloneAssetEditorToolkitHost.cpp
--- a/Engine/Source/Editor/UnrealEd/Private/Toolkits/SStandaloneAssetEditorToolkitHost.cpp
+++ b/Engine/Source/Editor/UnrealEd/Private/Toolkits/SStandaloneAssetEditorToolkitHost.cpp
@@ -20,6 +20,8 @@
 
 void SStandaloneAssetEditorToolkitHost::Construct( const SStandaloneAssetEditorToolkitHost::FArguments& InArgs, const TSharedPtr<FTabManager>& InTabManager, const FName InitAppName )
 {
+	ToolbarSlot = nullptr;
+
 	EditorCloseRequest = InArgs._OnRequestClose;
 	AppName = InitAppName;
 
@@ -167,6 +169,10 @@ void SStandaloneAssetEditorToolkitHost::RestoreFromLayout( const TSharedRef<FTab
 	[
 		SNew(SVerticalBox)
 		+ SVerticalBox::Slot()
+		.Padding(FMargin(0.0f, 0.0f, 0.0f, 2.0f))
+		.AutoHeight()
+		.Expose(ToolbarSlot)
+		+ SVerticalBox::Slot()
 		.Padding(0.0f, 1.0f, 0.0f, 0.0f)
 		.FillHeight(1.0f)
 		[
@@ -216,6 +222,24 @@ void SStandaloneAssetEditorToolkitHost::SetMenuOverlay( TSharedRef<SWidget> NewO
 	MenuOverlayWidgetContent->SetContent(NewOverlay);
 }
 
+void SStandaloneAssetEditorToolkitHost::SetToolbar(TSharedPtr<SWidget> Toolbar)
+{
+	if (Toolbar)
+	{
+		(*ToolbarSlot)
+		[
+			Toolbar.ToSharedRef()
+		];
+	}
+	else
+	{
+		(*ToolbarSlot)
+		[
+			SNullWidget::NullWidget
+		];
+	}
+}
+
 FEditorModeTools& SStandaloneAssetEditorToolkitHost::GetEditorModeManager() const
 {
 	check(HostedAssetEditorToolkit.IsValid());
diff --git a/Engine/Source/Editor/UnrealEd/Private/Toolkits/SStandaloneAssetEditorToolkitHost.h b/Engine/Source/Editor/UnrealEd/Private/Toolkits/SStandaloneAssetEditorToolkitHost.h
--- a/Engine/Source/Editor/UnrealEd/Private/Toolkits/SStandaloneAssetEditorToolkitHost.h
+++ b/Engine/Source/Editor/UnrealEd/Private/Toolkits/SStandaloneAssetEditorToolkitHost.h
@@ -82,6 +82,11 @@ public:
 	 */
 	void SetMenuOverlay( TSharedRef<SWidget> NewOverlay );
 
+	/**
+	 * Sets the contents of the toolbar in this asset editor
+	 */
+	void SetToolbar(TSharedPtr<SWidget> Toolbar);
+
 	virtual FEditorModeTools& GetEditorModeManager() const override;
 
 private:
@@ -115,6 +120,8 @@ private:
 
 	/** The menu extenders to populate the main toolkit host menu with */
 	TArray< TSharedPtr<FExtender> > MenuExtenders;
+
+	SVerticalBox::FSlot* ToolbarSlot;
 };
 
 
diff --git a/Engine/Source/Editor/UnrealEd/Public/Toolkits/AssetEditorToolkit.h b/Engine/Source/Editor/UnrealEd/Public/Toolkits/AssetEditorToolkit.h
--- a/Engine/Source/Editor/UnrealEd/Public/Toolkits/AssetEditorToolkit.h
+++ b/Engine/Source/Editor/UnrealEd/Public/Toolkits/AssetEditorToolkit.h
@@ -94,6 +94,8 @@ class UNREALED_API FAssetEditorToolkit
 {
 
 public:
+	static const FName DefaultAssetEditorToolBarName;
+	static const FName DefaultAssetEditorSlimToolBarName;
 
 	/** Default constructor */
 	FAssetEditorToolkit();
@@ -133,6 +135,7 @@ public:
 	virtual TWeakObjectPtr<class UEdMode> GetScriptableEditorMode() const final;
 	virtual FText GetEditorModeDisplayName() const override;
 	virtual FSlateIcon GetEditorModeIcon() const override;
+	virtual bool UsesCustomToolbarPlacement() const { return true;  }
 
 	/** IAssetEditorInstance interface */
 	virtual FName GetEditorName() const override;
@@ -186,7 +189,7 @@ public:
 	}
 
 	/** Registers default tool bar */
-	static void RegisterDefaultToolBar();
+	static void RegisterDefaultToolBar(bool bCustomToolbarPlacement);
 
 	/** Makes a default asset editing toolbar */
 	void GenerateToolbar();
diff --git a/Engine/Source/Runtime/Slate/Public/Widgets/Input/STextComboBox.h b/Engine/Source/Runtime/Slate/Public/Widgets/Input/STextComboBox.h
--- a/Engine/Source/Runtime/Slate/Public/Widgets/Input/STextComboBox.h
+++ b/Engine/Source/Runtime/Slate/Public/Widgets/Input/STextComboBox.h
@@ -29,7 +29,7 @@ public:
 		: _ComboBoxStyle(&FCoreStyle::Get().GetWidgetStyle< FComboBoxStyle >("ComboBox"))
 		, _ButtonStyle(nullptr)
 		, _ColorAndOpacity( FSlateColor::UseForeground() )
-		, _ContentPadding(FMargin(4.0, 2.0))
+		, _ContentPadding(_ComboBoxStyle->ContentPadding)
 		, _OnGetTextLabelForItem()
 		{}
 
diff --git a/Engine/Source/Runtime/SlateCore/Private/Styling/StarshipCoreStyle.cpp b/Engine/Source/Runtime/SlateCore/Private/Styling/StarshipCoreStyle.cpp
--- a/Engine/Source/Runtime/SlateCore/Private/Styling/StarshipCoreStyle.cpp
+++ b/Engine/Source/Runtime/SlateCore/Private/Styling/StarshipCoreStyle.cpp
@@ -367,7 +367,8 @@ TSharedRef<ISlateStyle> FStarshipCoreStyle::Create()
 		Style->Set("Icons.Duplicate", new IMAGE_BRUSH_SVG("Starship/Common/Duplicate", Icon16x16));
 		Style->Set("Icons.Edit", new IMAGE_BRUSH_SVG("Starship/Common/edit", Icon16x16));
 
-
+		Style->Set("Icons.Visible", new IMAGE_BRUSH_SVG("Starship/Common/visible", Icon16x16));
+		Style->Set("Icons.Hidden", new IMAGE_BRUSH_SVG("Starship/Common/hidden", Icon16x16));
 	}
 
 	// Tool panels
