Fix duplicate toolbars in derived blueprint asset editors

diff --git a/Engine/Source/Editor/Kismet/Private/BlueprintEditor.cpp b/Engine/Source/Editor/Kismet/Private/BlueprintEditor.cpp
--- a/Engine/Source/Editor/Kismet/Private/BlueprintEditor.cpp
+++ b/Engine/Source/Editor/Kismet/Private/BlueprintEditor.cpp
@@ -1569,6 +1569,8 @@ FBlueprintEditor::FBlueprintEditor()
 	AnalyticsStats.NodeKeymapCreateCount = 0;
 	AnalyticsStats.NodePasteCreateCount = 0;
 
+	bUsesCustomToolbar = true;
+
 	UEditorEngine* Editor = (UEditorEngine*)GEngine;
 	if (Editor != nullptr)
 	{
@@ -1960,6 +1962,9 @@ void FBlueprintEditor::InitBlueprintEditor(
 		// Flag the blueprint as having been opened
 		Blueprint->bIsNewlyCreated = false;
 
+		// @todo reskin, this is temp for now.  Only the base blueprint editor has had a toolbar treatment as part of the reskin.  All other derived types of BP editors need to use their custom toolbars for now.
+		bUsesCustomToolbar &= (Blueprint->GetClass() != UBlueprint::StaticClass());
+
 		Objects.Add( Blueprint );
 	}
 	
diff --git a/Engine/Source/Editor/Kismet/Public/BlueprintEditor.h b/Engine/Source/Editor/Kismet/Public/BlueprintEditor.h
--- a/Engine/Source/Editor/Kismet/Public/BlueprintEditor.h
+++ b/Engine/Source/Editor/Kismet/Public/BlueprintEditor.h
@@ -195,7 +195,7 @@ public:
 public:
 	//~ Begin FAssetEditorToolkit Interface
 	virtual bool OnRequestClose() override;
-	virtual bool UsesCustomToolbarPlacement() const override { return false; }
+	virtual bool UsesCustomToolbarPlacement() const override { return bUsesCustomToolbar; }
 	// End of FAssetEditorToolkit 
 
 	//~ Begin IToolkit Interface
@@ -1428,6 +1428,8 @@ private:
 	/** Handle to the registered OnActiveTabChanged delegate */
 	FDelegateHandle OnActiveTabChangedDelegateHandle;
 
+	bool bUsesCustomToolbar;
+
 	// Allow derived editors to add command mappings 
 	virtual void OnCreateGraphEditorCommands(TSharedPtr<FUICommandList> GraphEditorCommandsList) {}
 };
