Fix data only bps having double toolbars

diff --git a/Engine/Source/Editor/Kismet/Private/BlueprintEditorModes.cpp b/Engine/Source/Editor/Kismet/Private/BlueprintEditorModes.cpp
--- a/Engine/Source/Editor/Kismet/Private/BlueprintEditorModes.cpp
+++ b/Engine/Source/Editor/Kismet/Private/BlueprintEditorModes.cpp
@@ -156,22 +156,17 @@ FBlueprintDefaultsApplicationMode::FBlueprintDefaultsApplicationMode(TSharedPtr<
 	BlueprintDefaultsTabFactories.RegisterFactory(MakeShareable(new FFindResultsSummoner(InBlueprintEditor)));
 	BlueprintDefaultsTabFactories.RegisterFactory(MakeShareable(new FMyBlueprintSummoner(InBlueprintEditor)));
 
-	TabLayout = FTabManager::NewLayout( "Standalone_BlueprintDefaults_Layout_v6" )
+	TabLayout = FTabManager::NewLayout( "Standalone_BlueprintDefaults_Layout_v7" )
 		->AddArea
 		(
 			FTabManager::NewPrimaryArea()
 			->SetOrientation(Orient_Vertical)
 			->Split
-			(
-				FTabManager::NewStack()
-				->SetSizeCoefficient( 0.2f )
-				->SetHideTabWell(true)
-				->AddTab( InBlueprintEditor->GetToolbarTabId(), ETabState::OpenedTab )
-			)
-			->Split
 			(
 				FTabManager::NewStack()
 				->AddTab( FBlueprintEditorTabs::DefaultEditorID, ETabState::OpenedTab )
+				->SetHideTabWell(true)
+
 			)
 		);
 
@@ -215,19 +210,12 @@ FBlueprintComponentsApplicationMode::FBlueprintComponentsApplicationMode(TShared
 	BlueprintComponentsTabFactories.RegisterFactory(MakeShareable(new FDefaultsEditorSummoner(InBlueprintEditor)));
 	BlueprintComponentsTabFactories.RegisterFactory(MakeShareable(new FFindResultsSummoner(InBlueprintEditor)));
 
-	TabLayout = FTabManager::NewLayout( "Standalone_BlueprintComponents_Layout_v5" )
+	TabLayout = FTabManager::NewLayout( "Standalone_BlueprintComponents_Layout_v6" )
 		->AddArea
 		(
 			FTabManager::NewPrimaryArea()
 			->SetOrientation(Orient_Vertical)
 			->Split
-			(
-				FTabManager::NewStack()
-				->SetSizeCoefficient( 0.2f )
-				->SetHideTabWell(true)
-				->AddTab( InBlueprintEditor->GetToolbarTabId(), ETabState::OpenedTab )
-			)
-			->Split
 			(
 				FTabManager::NewSplitter()
 				->SetOrientation(Orient_Horizontal)
@@ -345,7 +333,7 @@ FBlueprintInterfaceApplicationMode::FBlueprintInterfaceApplicationMode(TSharedPt
 	BlueprintInterfaceTabFactories.RegisterFactory(MakeShareable(new FFindResultsSummoner(InBlueprintEditor)));
 	BlueprintInterfaceTabFactories.RegisterFactory(MakeShareable(new FSelectionDetailsSummoner(InBlueprintEditor)));
 
-	TabLayout = FTabManager::NewLayout( "Standalone_BlueprintInterface_Layout_v3" )
+	TabLayout = FTabManager::NewLayout( "Standalone_BlueprintInterface_Layout_v4" )
 		->AddArea
 		(
 			FTabManager::NewPrimaryArea() ->SetOrientation(Orient_Vertical)
@@ -357,13 +345,6 @@ FBlueprintInterfaceApplicationMode::FBlueprintInterfaceApplicationMode(TSharedPt
 					FTabManager::NewSplitter()->SetOrientation( Orient_Vertical )
 					->SetSizeCoefficient(0.70f)
 					->Split
-					(
-						FTabManager::NewStack()
-						->SetSizeCoefficient(0.186721f)
-						->SetHideTabWell(true)
-						->AddTab(InBlueprintEditor->GetToolbarTabId(), ETabState::OpenedTab)
-					)
-					->Split
 					(
 						FTabManager::NewStack()
 						->SetSizeCoefficient( 0.80f )
