Lighten minor tab backgrounds

diff --git a/Engine/Source/Editor/Kismet/Private/SSCSEditor.cpp b/Engine/Source/Editor/Kismet/Private/SSCSEditor.cpp
--- a/Engine/Source/Editor/Kismet/Private/SSCSEditor.cpp
+++ b/Engine/Source/Editor/Kismet/Private/SSCSEditor.cpp
@@ -4036,7 +4036,6 @@ void SSCSEditor::Construct( const FArguments& InArgs )
 			.Padding(0)
 			.BorderImage(FEditorStyle::GetBrush("DetailsView.CategoryTop"))
 			.AddMetaData<FTagMetaData>(FTagMetaData(TEXT("ComponentsPanel")))
-			.BorderBackgroundColor( FLinearColor( .6,.6,.6, 1.0f ) )
 			[
 				SAssignNew(HeaderBox, SVerticalBox)
 			]
@@ -4045,9 +4044,7 @@ void SSCSEditor::Construct( const FArguments& InArgs )
 		+ SVerticalBox::Slot()
 		.Padding(0.0f, 0.0f)
 		[
-			SNew(SBorder)
-			.Padding(2.0f)
-			.BorderImage(FEditorStyle::GetBrush("SCSEditor.TreePanel"))
+			SNew(SBox)
 			.AddMetaData<FTagMetaData>(FTagMetaData(TEXT("ComponentsPanel")))
 			.Visibility(this, &SSCSEditor::GetComponentsTreeVisibility)
 			[
diff --git a/Engine/Source/Runtime/SlateCore/Private/Styling/StarshipCoreStyle.cpp b/Engine/Source/Runtime/SlateCore/Private/Styling/StarshipCoreStyle.cpp
--- a/Engine/Source/Runtime/SlateCore/Private/Styling/StarshipCoreStyle.cpp
+++ b/Engine/Source/Runtime/SlateCore/Private/Styling/StarshipCoreStyle.cpp
@@ -1431,7 +1431,7 @@ void FStarshipCoreStyle::SetupDockingStyles(TSharedRef<FStyle>& Style)
 
 		.SetColorOverlayTabBrush(FSlateNoResource())
 		.SetColorOverlayIconBrush(FSlateNoResource())
-		.SetContentAreaBrush(FSlateColorBrush(FStyleColors::Recessed))
+		.SetContentAreaBrush(FSlateColorBrush(FStyleColors::Background))
 		.SetTabWellBrush(FSlateNoResource())
 		.SetFlashColor(TabFlashColor)
 
