Blueprints: Add option to disable all breakpoints directly from the Debug window. Users no longer have to select indivdual breakpoints to disable them or go to a different file menu to do so.

diff --git a/Engine/Source/Editor/Kismet/Private/Debugging/SKismetDebuggingView.cpp b/Engine/Source/Editor/Kismet/Private/Debugging/SKismetDebuggingView.cpp
--- a/Engine/Source/Editor/Kismet/Private/Debugging/SKismetDebuggingView.cpp
+++ b/Engine/Source/Editor/Kismet/Private/Debugging/SKismetDebuggingView.cpp
@@ -1058,6 +1058,22 @@ EVisibility SKismetDebuggingView::IsDebuggerVisible() const
 	return (GEditor->PlayWorld != NULL) ? EVisibility::Visible : EVisibility::Collapsed;
 }
 
+bool SKismetDebuggingView::CanDisableAllBreakpoints() const
+{
+	const UBlueprint* BlueprintObj = BlueprintToWatchPtr.Get();
+	return BlueprintObj && BlueprintObj->Breakpoints.Num() > 0;
+}
+
+FReply SKismetDebuggingView::OnDisableAllBreakpointsClicked()
+{
+	if(UBlueprint* BlueprintObj = BlueprintToWatchPtr.Get())
+	{
+		FDebuggingActionCallbacks::SetEnabledOnAllBreakpoints(BlueprintObj, false);		
+	}
+	
+	return FReply::Handled();
+}
+
 void SKismetDebuggingView::Construct(const FArguments& InArgs)
 {	
 	BlueprintToWatchPtr = InArgs._BlueprintToWatch;
@@ -1093,8 +1109,21 @@ void SKismetDebuggingView::Construct(const FArguments& InArgs)
 		+SVerticalBox::Slot()
 		.AutoHeight()
 		[
-			SNew(STextBlock)
-			.Text( this, &SKismetDebuggingView::GetTopText )
+			SNew( SHorizontalBox )
+			+ SHorizontalBox::Slot()
+			.HAlign( HAlign_Left )
+			[
+				SNew( STextBlock )
+					.Text( this, &SKismetDebuggingView::GetTopText )
+			]
+			+ SHorizontalBox::Slot()
+			.HAlign( HAlign_Right )
+			[
+				SNew( SButton )
+	                .IsEnabled( this, &SKismetDebuggingView::CanDisableAllBreakpoints )
+	                .Text( LOCTEXT( "DisableAllBreakPoints", "Disable All Breakpoints" ) )
+	                .OnClicked( this, &SKismetDebuggingView::OnDisableAllBreakpointsClicked )
+			]
 		]
 		+SVerticalBox::Slot()
 		[
diff --git a/Engine/Source/Editor/Kismet/Private/Debugging/SKismetDebuggingView.h b/Engine/Source/Editor/Kismet/Private/Debugging/SKismetDebuggingView.h
--- a/Engine/Source/Editor/Kismet/Private/Debugging/SKismetDebuggingView.h
+++ b/Engine/Source/Editor/Kismet/Private/Debugging/SKismetDebuggingView.h
@@ -107,6 +107,8 @@ public:
 	// End of SWidget interface
 protected:
 	FText GetTopText() const;
+	bool CanDisableAllBreakpoints() const;
+	FReply OnDisableAllBreakpointsClicked();
 
 	TSharedRef<ITableRow> OnGenerateRowForWatchTree(FDebugTreeItemPtr InItem, const TSharedRef<STableViewBase>& OwnerTable);
 	void OnGetChildrenForWatchTree(FDebugTreeItemPtr InParent, TArray<FDebugTreeItemPtr>& OutChildren);
