Style tweaks for BP editor function parameters customization

diff --git a/Engine/Source/Editor/Kismet/Private/BlueprintDetailsCustomization.cpp b/Engine/Source/Editor/Kismet/Private/BlueprintDetailsCustomization.cpp
--- a/Engine/Source/Editor/Kismet/Private/BlueprintDetailsCustomization.cpp
+++ b/Engine/Source/Editor/Kismet/Private/BlueprintDetailsCustomization.cpp
@@ -2837,13 +2837,15 @@ void FBlueprintGraphArgumentLayout::GenerateHeaderRowContent( FDetailWidgetRow&
 		.AutoWidth()
 		[
 			SNew(SButton)
+			.ButtonStyle(FAppStyle::Get(), TEXT("SimpleButton"))
 			.ContentPadding(0)
 			.IsEnabled(!IsPinEditingReadOnly())
 			.OnClicked(this, &FBlueprintGraphArgumentLayout::OnArgMoveUp)
 			.ToolTipText(LOCTEXT("FunctionArgDetailsArgMoveUpTooltip", "Move this parameter up in the list."))
 			[
 				SNew(SImage)
-				.Image(FEditorStyle::GetBrush("BlueprintEditor.Details.ArgUpButton"))
+				.Image(FEditorStyle::GetBrush("Icons.ChevronUp"))
+				.ColorAndOpacity(FSlateColor::UseForeground())
 			]
 		]
 		+ SHorizontalBox::Slot()
@@ -2851,13 +2853,15 @@ void FBlueprintGraphArgumentLayout::GenerateHeaderRowContent( FDetailWidgetRow&
 		.Padding(2, 0)
 		[
 			SNew(SButton)
+			.ButtonStyle(FAppStyle::Get(), TEXT("SimpleButton"))
 			.ContentPadding(0)
 			.IsEnabled(!IsPinEditingReadOnly())
 			.OnClicked(this, &FBlueprintGraphArgumentLayout::OnArgMoveDown)
 			.ToolTipText(LOCTEXT("FunctionArgDetailsArgMoveDownTooltip", "Move this parameter down in the list."))
 			[
 				SNew(SImage)
-				.Image(FEditorStyle::GetBrush("BlueprintEditor.Details.ArgDownButton"))
+				.Image(FEditorStyle::GetBrush("Icons.ChevronDown"))
+				.ColorAndOpacity(FSlateColor::UseForeground())
 			]
 		]
 		+ SHorizontalBox::Slot()
@@ -3699,17 +3703,14 @@ void FBlueprintGraphActionDetails::CustomizeDetails( IDetailLayoutBuilder& Detai
 
 		TSharedRef<SHorizontalBox> InputsHeaderContentWidget = SNew(SHorizontalBox);
 		TWeakPtr<SWidget> WeakInputsHeaderWidget = InputsHeaderContentWidget;
+
 		InputsHeaderContentWidget->AddSlot()
+			.HAlign(HAlign_Right)
 			[
-				SNew(SHorizontalBox)
-			];
-		InputsHeaderContentWidget->AddSlot()
-			.AutoWidth()
-			[
+
 				SNew(SButton)
-				.ButtonStyle(FEditorStyle::Get(), "RoundButton")
-				.ForegroundColor(FEditorStyle::GetSlateColor("DefaultForeground"))
-				.ContentPadding(FMargin(2, 0))
+				.ButtonStyle(FEditorStyle::Get(), "SimpleButton")
+				.ContentPadding(FMargin(1, 0))
 				.OnClicked(this, &FBlueprintGraphActionDetails::OnAddNewInputClicked)
 				.Visibility(this, &FBlueprintGraphActionDetails::GetAddNewInputOutputVisibility)
 				.HAlign(HAlign_Right)
@@ -3718,27 +3719,9 @@ void FBlueprintGraphActionDetails::CustomizeDetails( IDetailLayoutBuilder& Detai
 				.AddMetaData<FTagMetaData>(FTagMetaData(TEXT("FunctionNewInputArg")))
 				.IsEnabled(this, &FBlueprintGraphActionDetails::IsAddNewInputOutputEnabled)
 				[
-					SNew(SHorizontalBox)
-
-					+ SHorizontalBox::Slot()
-						.AutoWidth()
-						.Padding(FMargin(0, 1))
-						[
-							SNew(SImage)
-							.Image(FEditorStyle::GetBrush("Plus"))
-						]
-
-					+ SHorizontalBox::Slot()
-						.VAlign(VAlign_Center)
-						.AutoWidth()
-						.Padding(FMargin(2, 0, 0, 0))
-						[
-							SNew(STextBlock)
-							.Font(IDetailLayoutBuilder::GetDetailFontBold())
-							.Text(LOCTEXT("FunctionNewParameterInputArg", "New Parameter"))
-							.Visibility(this, &FBlueprintGraphActionDetails::OnGetSectionTextVisibility, WeakInputsHeaderWidget)
-							.ShadowOffset(FVector2D(1, 1))
-						]
+					SNew(SImage)
+					.Image(FAppStyle::Get().GetBrush("Icons.PlusCircle"))
+					.ColorAndOpacity(FSlateColor::UseForeground())
 				]
 			];
 		InputsCategory.HeaderContent(InputsHeaderContentWidget);
@@ -3753,17 +3736,13 @@ void FBlueprintGraphActionDetails::CustomizeDetails( IDetailLayoutBuilder& Detai
 		
 			TSharedRef<SHorizontalBox> OutputsHeaderContentWidget = SNew(SHorizontalBox);
 			TWeakPtr<SWidget> WeakOutputsHeaderWidget = OutputsHeaderContentWidget;
+
 			OutputsHeaderContentWidget->AddSlot()
-				[
-					SNew(SHorizontalBox)
-				];
-			OutputsHeaderContentWidget->AddSlot()
-				.AutoWidth()
+				.HAlign(HAlign_Right)
 				[
 					SNew(SButton)
-					.ButtonStyle(FEditorStyle::Get(), "RoundButton")
-					.ForegroundColor(FEditorStyle::GetSlateColor("DefaultForeground"))
-					.ContentPadding(FMargin(2, 0))
+					.ButtonStyle(FEditorStyle::Get(), "SimpleButton")
+					.ContentPadding(FMargin(1, 0))
 					.OnClicked(this, &FBlueprintGraphActionDetails::OnAddNewOutputClicked)
 					.Visibility(this, &FBlueprintGraphActionDetails::GetAddNewInputOutputVisibility)
 					.HAlign(HAlign_Right)
@@ -3771,30 +3750,12 @@ void FBlueprintGraphActionDetails::CustomizeDetails( IDetailLayoutBuilder& Detai
 					.VAlign(VAlign_Center)
 					.AddMetaData<FTagMetaData>(FTagMetaData(TEXT("FunctionNewOutputArg")))
 					.IsEnabled(this, &FBlueprintGraphActionDetails::IsAddNewInputOutputEnabled)
-					[
-						SNew(SHorizontalBox)
-
-						+ SHorizontalBox::Slot()
-							.AutoWidth()
-							.Padding(FMargin(0, 1))
 					[
 						SNew(SImage)
-							.Image(FEditorStyle::GetBrush("Plus"))
+						.Image(FAppStyle::Get().GetBrush("Icons.PlusCircle"))
+						.ColorAndOpacity(FSlateColor::UseForeground())
 					]
-
-				+ SHorizontalBox::Slot()
-					.VAlign(VAlign_Center)
-					.AutoWidth()
-					.Padding(FMargin(2, 0, 0, 0))
-					[
-						SNew(STextBlock)
-						.Font(IDetailLayoutBuilder::GetDetailFontBold())
-						.Text(LOCTEXT("FunctionNewOutputArg", "New Parameter"))
-						.Visibility(this, &FBlueprintGraphActionDetails::OnGetSectionTextVisibility, WeakOutputsHeaderWidget)
-						.ShadowOffset(FVector2D(1, 1))
-					]
-				]
-			];
+				];
 			OutputsCategory.HeaderContent(OutputsHeaderContentWidget);
 		}
 	}
diff --git a/Engine/Source/Editor/PropertyEditor/Private/PropertyCustomizationHelpers.cpp b/Engine/Source/Editor/PropertyEditor/Private/PropertyCustomizationHelpers.cpp
--- a/Engine/Source/Editor/PropertyEditor/Private/PropertyCustomizationHelpers.cpp
+++ b/Engine/Source/Editor/PropertyEditor/Private/PropertyCustomizationHelpers.cpp
@@ -63,10 +63,9 @@ namespace PropertyCustomizationHelpers
 			OnClickAction = InArgs._OnClickAction;
 
 			SButton::FArguments ButtonArgs = SButton::FArguments()
-				.ButtonStyle( FEditorStyle::Get(), "HoverHintOnly" )
+				.ButtonStyle( FAppStyle::Get(), "SimpleButton" )
 				.OnClicked( this, &SPropertyEditorButton::OnClick )
 				.ToolTipText( InArgs._Text )
-				.ContentPadding( 4.0f )
 				.ForegroundColor( FSlateColor::UseForeground() )
 				.IsFocusable(InArgs._IsFocusable)
 				[ 
@@ -169,7 +168,7 @@ namespace PropertyCustomizationHelpers
 			SNew( SPropertyEditorButton )
 			.Text( LOCTEXT( "ClearButtonLabel", "Clear") )
 			.ToolTipText( OptionalToolTipText.Get().IsEmpty() ? LOCTEXT( "ClearButtonToolTipText", "Clear Path") : OptionalToolTipText )
-			.Image( FEditorStyle::GetBrush("PropertyWindow.Button_Clear") )
+			.Image(FAppStyle::Get().GetBrush("Icons.X"))
 			.OnClickAction( OnClearClicked )
 			.IsEnabled(IsEnabled)
 			.IsFocusable( false );
diff --git a/Engine/Source/Runtime/SlateCore/Private/Styling/StarshipCoreStyle.cpp b/Engine/Source/Runtime/SlateCore/Private/Styling/StarshipCoreStyle.cpp
--- a/Engine/Source/Runtime/SlateCore/Private/Styling/StarshipCoreStyle.cpp
+++ b/Engine/Source/Runtime/SlateCore/Private/Styling/StarshipCoreStyle.cpp
@@ -334,6 +334,7 @@ TSharedRef<ISlateStyle> FStarshipCoreStyle::Create()
 		Style->Set("Icons.Cross", new IMAGE_BRUSH_SVG("Starship/Common/plus", Icon16x16));
 		Style->Set("Icons.Plus", new IMAGE_BRUSH_SVG("Starship/Common/plus", Icon16x16));
 		Style->Set("Icons.PlusCircle", new IMAGE_BRUSH_SVG("Starship/Common/plus-circle", Icon16x16));
+		Style->Set("Icons.X", new IMAGE_BRUSH_SVG("Starship/Common/close", Icon16x16));
 
 		Style->Set("Icons.Save", new IMAGE_BRUSH_SVG("Starship/Common/save", Icon16x16));
 
@@ -369,6 +370,8 @@ TSharedRef<ISlateStyle> FStarshipCoreStyle::Create()
 
 		Style->Set("Icons.Visible", new IMAGE_BRUSH_SVG("Starship/Common/visible", Icon16x16));
 		Style->Set("Icons.Hidden", new IMAGE_BRUSH_SVG("Starship/Common/hidden", Icon16x16));
+
+
 	}
 
 	// Tool panels
