Fix utility blueprints not having a toolbar

diff --git a/Engine/Plugins/Experimental/ControlRig/Source/ControlRigEditor/Private/Editor/ControlRigEditor.h b/Engine/Plugins/Experimental/ControlRig/Source/ControlRigEditor/Private/Editor/ControlRigEditor.h
--- a/Engine/Plugins/Experimental/ControlRig/Source/ControlRigEditor/Private/Editor/ControlRigEditor.h
+++ b/Engine/Plugins/Experimental/ControlRig/Source/ControlRigEditor/Private/Editor/ControlRigEditor.h
@@ -93,6 +93,7 @@ public:
 	virtual FText GetToolkitToolTipText() const override;
 	virtual FString GetWorldCentricTabPrefix() const override;
 	virtual FLinearColor GetWorldCentricTabColorScale() const override;	
+	virtual bool UsesCustomToolbarPlacement() const override { return true; }
 	virtual FString GetDocumentationLink() const override
 	{
 		return FString(TEXT("Engine/Animation/ControlRig"));
diff --git a/Engine/Source/Editor/AnimationBlueprintEditor/Private/AnimationBlueprintEditor.h b/Engine/Source/Editor/AnimationBlueprintEditor/Private/AnimationBlueprintEditor.h
--- a/Engine/Source/Editor/AnimationBlueprintEditor/Private/AnimationBlueprintEditor.h
+++ b/Engine/Source/Editor/AnimationBlueprintEditor/Private/AnimationBlueprintEditor.h
@@ -103,6 +103,7 @@ public:
 	virtual void OnActiveTabChanged(TSharedPtr<SDockTab> PreviouslyActive, TSharedPtr<SDockTab> NewlyActivated) override;
 	virtual void OnSelectedNodesChangedImpl(const TSet<class UObject*>& NewSelection) override;
 	virtual void HandleSetObjectBeingDebugged(UObject* InObject) override;
+	virtual bool UsesCustomToolbarPlacement() const override { return true; }
 
 	// Gets the Anim Blueprint being edited/viewed by this Persona instance
 	UAnimBlueprint* GetAnimBlueprint() const;
diff --git a/Engine/Source/Editor/Kismet/Private/BlueprintEditor.cpp b/Engine/Source/Editor/Kismet/Private/BlueprintEditor.cpp
--- a/Engine/Source/Editor/Kismet/Private/BlueprintEditor.cpp
+++ b/Engine/Source/Editor/Kismet/Private/BlueprintEditor.cpp
@@ -1576,8 +1576,6 @@ FBlueprintEditor::FBlueprintEditor()
 	AnalyticsStats.NodeKeymapCreateCount = 0;
 	AnalyticsStats.NodePasteCreateCount = 0;
 
-	bUsesCustomToolbar = true;
-
 	UEditorEngine* Editor = (UEditorEngine*)GEngine;
 	if (Editor != nullptr)
 	{
@@ -1969,11 +1967,6 @@ void FBlueprintEditor::InitBlueprintEditor(
 		// Flag the blueprint as having been opened
 		Blueprint->bIsNewlyCreated = false;
 
-		const bool bBaseBlueprintType = (Blueprint->GetClass() == UBlueprint::StaticClass() || Blueprint->GetClass() == ULevelScriptBlueprint::StaticClass());
-
-		// @todo reskin, this is temp for now.  Only the base blueprint editor has had a toolbar treatment as part of the reskin.  All other derived types of BP editors need to use their custom toolbars for now.
-		bUsesCustomToolbar &= (!bBaseBlueprintType);
-
 		Objects.Add( Blueprint );
 	}
 	
diff --git a/Engine/Source/Editor/Kismet/Public/BlueprintEditor.h b/Engine/Source/Editor/Kismet/Public/BlueprintEditor.h
--- a/Engine/Source/Editor/Kismet/Public/BlueprintEditor.h
+++ b/Engine/Source/Editor/Kismet/Public/BlueprintEditor.h
@@ -195,7 +195,7 @@ public:
 public:
 	//~ Begin FAssetEditorToolkit Interface
 	virtual bool OnRequestClose() override;
-	virtual bool UsesCustomToolbarPlacement() const override { return bUsesCustomToolbar; }
+	virtual bool UsesCustomToolbarPlacement() const override { return false; }
 	// End of FAssetEditorToolkit 
 
 	//~ Begin IToolkit Interface
@@ -1431,8 +1431,6 @@ private:
 	/** Handle to the registered OnActiveTabChanged delegate */
 	FDelegateHandle OnActiveTabChangedDelegateHandle;
 
-	bool bUsesCustomToolbar;
-
 	// Allow derived editors to add command mappings 
 	virtual void OnCreateGraphEditorCommands(TSharedPtr<FUICommandList> GraphEditorCommandsList) {}
 };
