Blueprints: Enable the new CDO moving by defalut which will elimate recursion when marking a BP as structurally modified.

#rnx

diff --git a/Engine/Source/Editor/Kismet/Private/BlueprintCompilationManager.cpp b/Engine/Source/Editor/Kismet/Private/BlueprintCompilationManager.cpp
--- a/Engine/Source/Editor/Kismet/Private/BlueprintCompilationManager.cpp
+++ b/Engine/Source/Editor/Kismet/Private/BlueprintCompilationManager.cpp
@@ -509,7 +509,7 @@ FReinstancingJob::FReinstancingJob(TPair<UClass*, UClass*> InOldToNew)
 namespace SkelReinstUtils
 {
 	/** Flag to use the new FBlueprintCompileReinstancer::MoveDependentSkelToReinst and avoid problematic recursion during reinstancing */
-	static bool bEnableSkelReinstUpdate = false;
+	static bool bEnableSkelReinstUpdate = true;
 	static FAutoConsoleVariableRef CVarEnableSkelReinstUpdate(
 		TEXT("BP.bEnableSkelReinstUpdate"), bEnableSkelReinstUpdate,
 		TEXT("If true the Reinstancing of SKEL classes will use the new FBlueprintCompileReinstancer::MoveDependentSkelToReinst(o(n)) instead of the old MoveSkelCDOAside (o(n^2))"),
