Fixed crash in FindInBlueprintManager keeping dangling UClass pointer in a cache.

diff --git a/Engine/Source/Editor/Kismet/Private/FindInBlueprintManager.cpp b/Engine/Source/Editor/Kismet/Private/FindInBlueprintManager.cpp
--- a/Engine/Source/Editor/Kismet/Private/FindInBlueprintManager.cpp
+++ b/Engine/Source/Editor/Kismet/Private/FindInBlueprintManager.cpp
@@ -2075,17 +2075,17 @@ void FFindInBlueprintSearchManager::OnAssetAdded(const FAssetData& InAssetData)
 {
 	const UClass* AssetClass = nullptr;
 	{
-		const UClass** FoundClass = CachedAssetClasses.Find(InAssetData.AssetClass);
-		if (FoundClass)
+		TWeakObjectPtr<const UClass> FoundClass = CachedAssetClasses.FindRef(InAssetData.AssetClass);
+		if (FoundClass.IsValid())
 		{
-			AssetClass = *FoundClass;
+			AssetClass = FoundClass.Get();
 		}
 		else
 		{
 			AssetClass = InAssetData.GetClass();
 			if (AssetClass)
 			{
-				CachedAssetClasses.Add(InAssetData.AssetClass, AssetClass);
+				CachedAssetClasses.Add(InAssetData.AssetClass, TWeakObjectPtr<const UClass>(AssetClass));
 			}
 		}
 	}
diff --git a/Engine/Source/Editor/Kismet/Public/FindInBlueprintManager.h b/Engine/Source/Editor/Kismet/Public/FindInBlueprintManager.h
--- a/Engine/Source/Editor/Kismet/Public/FindInBlueprintManager.h
+++ b/Engine/Source/Editor/Kismet/Public/FindInBlueprintManager.h
@@ -795,7 +795,7 @@ private:
 	EFiBCacheOpType CurrentCacheOpType;
 
 	/** Mapping between a class name and its UClass instance - used for faster look up in FFindInBlueprintSearchManager::OnAssetAdded */
-	TMap<FName, const UClass*> CachedAssetClasses;
+	TMap<FName, TWeakObjectPtr<const UClass>> CachedAssetClasses;
 
 	/** The tab identifier/instance name for global find results */
 	FName GlobalFindResultsTabIDs[MAX_GLOBAL_FIND_RESULTS];
