PR #7572: UE-103069: Use GetDisplayNameText instead of PropertyName directly in the My Blueprint Panel (Contributed by [REDACTED_PERSON])

diff --git a/Engine/Source/Editor/Kismet/Private/SMyBlueprint.cpp b/Engine/Source/Editor/Kismet/Private/SMyBlueprint.cpp
--- a/Engine/Source/Editor/Kismet/Private/SMyBlueprint.cpp
+++ b/Engine/Source/Editor/Kismet/Private/SMyBlueprint.cpp
@@ -1223,7 +1223,7 @@ void SMyBlueprint::CollectAllActions(FGraphActionListBuilderBase& OutAllActions)
 
 		const FText PropertyTooltip = Property->GetToolTipText();
 		const FName PropertyName = Property->GetFName();
-		const FText PropertyDesc = FText::FromName(PropertyName);
+		const FText PropertyDesc = Property->GetDisplayNameText();
 
 		FText CategoryName = FObjectEditorUtils::GetCategoryText(Property);
 		FText PropertyCategory = FObjectEditorUtils::GetCategoryText(Property);
