Blueprint tab factories use the asset editor subsystem to close editors, and keep a weak object ptr to the editing blueprint.

#rnx

diff --git a/Engine/Source/Editor/Kismet/Private/BlueprintEditorTabFactories.cpp b/Engine/Source/Editor/Kismet/Private/BlueprintEditorTabFactories.cpp
--- a/Engine/Source/Editor/Kismet/Private/BlueprintEditorTabFactories.cpp
+++ b/Engine/Source/Editor/Kismet/Private/BlueprintEditorTabFactories.cpp
@@ -137,8 +137,9 @@ TSharedRef<SWidget> FDebugInfoSummoner::CreateTabBody(const FWorkflowTabSpawnInf
 	return BlueprintEditorPtr->GetDebuggingView();
 }
 
-FDefaultsEditorSummoner::FDefaultsEditorSummoner(TSharedPtr<class FAssetEditorToolkit> InHostingApp)
+FDefaultsEditorSummoner::FDefaultsEditorSummoner(TSharedPtr<FBlueprintEditor> InHostingApp)
 	: FWorkflowTabFactory(FBlueprintEditorTabs::DefaultEditorID, InHostingApp)
+	, EditingBlueprint(InHostingApp->GetBlueprintObj())
 {
 	TabLabel = LOCTEXT("ClassDefaultsTabTitle", "Class Defaults");
 	TabIcon = FSlateIcon(FEditorStyle::GetStyleSetName(), "Kismet.Tabs.BlueprintDefaults");
@@ -181,14 +182,12 @@ TSharedRef<SWidget> FDefaultsEditorSummoner::CreateTabBody(const FWorkflowTabSpa
 
 TSharedRef<SWidget> FDefaultsEditorSummoner::CreateOptionalEditableWarning() const
 {
-	TSharedPtr<FBlueprintEditor> BlueprintEditorPtr = StaticCastSharedPtr<FBlueprintEditor>(HostingApp.Pin());
-
 	bool bHasUneditableBlueprintComponent = false;
-	if (UBlueprint* Blueprint = BlueprintEditorPtr->GetBlueprintObj())
+	if (EditingBlueprint.IsValid())
 	{
-		if (Blueprint->GeneratedClass && FBlueprintEditorUtils::IsDataOnlyBlueprint(Blueprint))
+		if (EditingBlueprint->GeneratedClass && FBlueprintEditorUtils::IsDataOnlyBlueprint(EditingBlueprint.Get()))
 		{
-			if (AActor* Actor = Cast<AActor>(Blueprint->GeneratedClass->GetDefaultObject()))
+			if (AActor* Actor = Cast<AActor>(EditingBlueprint->GeneratedClass->GetDefaultObject()))
 			{
 				for (UActorComponent* Component : Actor->GetComponents())
 				{
@@ -235,37 +234,32 @@ TSharedRef<SWidget> FDefaultsEditorSummoner::CreateOptionalEditableWarning() con
 
 TSharedRef<SWidget> FDefaultsEditorSummoner::CreateOptionalDataOnlyMessage() const
 {
-	TSharedPtr<FBlueprintEditor> BlueprintEditorPtr = StaticCastSharedPtr<FBlueprintEditor>(HostingApp.Pin());
-
 	TSharedRef<SWidget> Message = SNullWidget::NullWidget;
-	if ( UBlueprint* Blueprint = BlueprintEditorPtr->GetBlueprintObj() )
+	if (EditingBlueprint.IsValid() && FBlueprintEditorUtils::IsDataOnlyBlueprint(EditingBlueprint.Get()))
 	{
-		if ( FBlueprintEditorUtils::IsDataOnlyBlueprint(Blueprint) )
-		{
-			Message = SNew(SBorder)
-				.Padding(FMargin(5))
-				.BorderImage(FEditorStyle::GetBrush("ToolPanel.GroupBorder"))
-				[
-					SNew(SWrapBox)
-					.UseAllottedSize(true)
-
-					+ SWrapBox::Slot()
-					[
-						SNew(STextBlock)
-						.Font(FEditorStyle::GetFontStyle("BoldFont"))
-						.Text(LOCTEXT("DataOnlyMessage_Part1", "NOTE: This is a data only blueprint, so only the default values are shown.  It does not have any script or variables.  If you want to add some, "))
-					]
-
-					+ SWrapBox::Slot()
-					[
-						SNew(SHyperlink)
-						.Style(FEditorStyle::Get(), "Common.GotoBlueprintHyperlink")
-						.OnNavigate(const_cast<FDefaultsEditorSummoner*>(this), &FDefaultsEditorSummoner::OnChangeBlueprintToNotDataOnly)
-						.Text(LOCTEXT("FullEditor", "Open Full Blueprint Editor"))
-						.ToolTipText(LOCTEXT("FullEditorToolTip", "This opens the blueprint in the full editor."))
-					]
-				];
-		}
+		Message = SNew(SBorder)
+			.Padding(FMargin(5))
+			.BorderImage(FEditorStyle::GetBrush("ToolPanel.GroupBorder"))
+			[
+				SNew(SWrapBox)
+				.UseAllottedSize(true)
+
+			+ SWrapBox::Slot()
+			[
+				SNew(STextBlock)
+				.Font(FEditorStyle::GetFontStyle("BoldFont"))
+				.Text(LOCTEXT("DataOnlyMessage_Part1", "NOTE: This is a data only blueprint, so only the default values are shown.  It does not have any script or variables.  If you want to add some, "))
+			]
+
+			+ SWrapBox::Slot()
+			[
+				SNew(SHyperlink)
+				.Style(FEditorStyle::Get(), "Common.GotoBlueprintHyperlink")
+				.OnNavigate(const_cast<FDefaultsEditorSummoner*>(this), &FDefaultsEditorSummoner::OnChangeBlueprintToNotDataOnly)
+				.Text(LOCTEXT("FullEditor", "Open Full Blueprint Editor"))
+				.ToolTipText(LOCTEXT("FullEditorToolTip", "This opens the blueprint in the full editor."))
+			]
+		];
 	}
 
 	return Message;
@@ -273,23 +267,20 @@ TSharedRef<SWidget> FDefaultsEditorSummoner::CreateOptionalDataOnlyMessage() con
 
 void FDefaultsEditorSummoner::OnChangeBlueprintToNotDataOnly()
 {
-	UBlueprint* Blueprint = nullptr;
-
+	if (!GEditor || !EditingBlueprint.IsValid())
 	{
-		TSharedPtr<FBlueprintEditor> BlueprintEditorPtr = StaticCastSharedPtr<FBlueprintEditor>(HostingApp.Pin());
-		Blueprint = BlueprintEditorPtr->GetBlueprintObj();
-		if ( Blueprint )
-		{
-			BlueprintEditorPtr->CloseWindow();
-		}
+		return;
 	}
 
-	if ( Blueprint )
+	UAssetEditorSubsystem* AssetEditorSubsystem = GEditor->GetEditorSubsystem<UAssetEditorSubsystem>();
+	if (!AssetEditorSubsystem)
 	{
-		Blueprint->bForceFullEditor = true;
-
-		GEditor->EditObject(Blueprint);
+		return;
 	}
+
+	AssetEditorSubsystem->CloseAllEditorsForAsset(EditingBlueprint.Get());
+	EditingBlueprint->bForceFullEditor = true;
+	AssetEditorSubsystem->OpenEditorForAsset(EditingBlueprint.Get());
 }
 
 FConstructionScriptEditorSummoner::FConstructionScriptEditorSummoner(TSharedPtr<class FAssetEditorToolkit> InHostingApp)
diff --git a/Engine/Source/Editor/Kismet/Private/BlueprintEditorTabFactories.h b/Engine/Source/Editor/Kismet/Private/BlueprintEditorTabFactories.h
--- a/Engine/Source/Editor/Kismet/Private/BlueprintEditorTabFactories.h
+++ b/Engine/Source/Editor/Kismet/Private/BlueprintEditorTabFactories.h
@@ -72,7 +72,7 @@ public:
 struct FDefaultsEditorSummoner : public FWorkflowTabFactory
 {
 public:
-	FDefaultsEditorSummoner(TSharedPtr<class FAssetEditorToolkit> InHostingApp);
+	FDefaultsEditorSummoner(TSharedPtr<class FBlueprintEditor> InHostingApp);
 
 	virtual TSharedRef<SWidget> CreateTabBody(const FWorkflowTabSpawnInfo& Info) const override;
 
@@ -86,6 +86,8 @@ private:
 	TSharedRef<SWidget> CreateOptionalEditableWarning() const;
 
 	void OnChangeBlueprintToNotDataOnly();
+
+	TWeakObjectPtr<UBlueprint> EditingBlueprint;
 };
 
 /////////////////////////////////////////////////////
