Blueprints: Fix blueprint variables not getting the correct display name.

#rnx

diff --git a/Engine/Source/Editor/Kismet/Private/SMyBlueprint.cpp b/Engine/Source/Editor/Kismet/Private/SMyBlueprint.cpp
--- a/Engine/Source/Editor/Kismet/Private/SMyBlueprint.cpp
+++ b/Engine/Source/Editor/Kismet/Private/SMyBlueprint.cpp
@@ -1223,7 +1223,7 @@ void SMyBlueprint::CollectAllActions(FGraphActionListBuilderBase& OutAllActions)
 
 		const FText PropertyTooltip = Property->GetToolTipText();
 		const FName PropertyName = Property->GetFName();
-		const FText PropertyDesc = Property->GetDisplayNameText();
+		const FText PropertyDesc = FText::FromName(PropertyName);
 
 		FText CategoryName = FObjectEditorUtils::GetCategoryText(Property);
 		FText PropertyCategory = FObjectEditorUtils::GetCategoryText(Property);
