Blueprint UI cosmetic fixes

diff --git a/Engine/Source/Editor/Kismet/Private/BlueprintDetailsCustomization.cpp b/Engine/Source/Editor/Kismet/Private/BlueprintDetailsCustomization.cpp
--- a/Engine/Source/Editor/Kismet/Private/BlueprintDetailsCustomization.cpp
+++ b/Engine/Source/Editor/Kismet/Private/BlueprintDetailsCustomization.cpp
@@ -3426,6 +3426,7 @@ void FBlueprintGraphActionDetails::CustomizeDetails( IDetailLayoutBuilder& Detai
 					SAssignNew( ColorBlock, SColorBlock )
 						.Color( this, &FBlueprintGraphActionDetails::GetNodeTitleColor )
 						.CornerRadius(FVector4(4.0f, 4.0f, 4.0f, 4.0f))
+						.Size(FVector2D(70.0f, 22.0f))
 						.AlphaDisplayMode(EColorBlockAlphaDisplayMode::Ignore)
 						.OnMouseButtonDown( this, &FBlueprintGraphActionDetails::ColorBlock_OnMouseButtonDown )
 				];
