Use the FProperty Display name in the SCS editor if it is available instead of the native C++ name to make it more consistent with the My Blueprint window

diff --git a/Engine/Source/Editor/Kismet/Private/SSCSEditor.cpp b/Engine/Source/Editor/Kismet/Private/SSCSEditor.cpp
--- a/Engine/Source/Editor/Kismet/Private/SSCSEditor.cpp
+++ b/Engine/Source/Editor/Kismet/Private/SSCSEditor.cpp
@@ -1033,10 +1033,11 @@ FString FSCSEditorTreeNodeComponentBase::GetDisplayString() const
 
 	UBlueprint* Blueprint = GetBlueprint();
 	UClass* VariableOwner = (Blueprint != nullptr) ? Blueprint->SkeletonGeneratedClass : nullptr;
+	FProperty* VariableProperty = FindFProperty<FProperty>(VariableOwner, VariableName);
 
 	bool const bHasValidVarName = (VariableName != NAME_None);
 	bool const bIsArrayVariable = bHasValidVarName && (VariableOwner != nullptr) && 
-		FindFProperty<FArrayProperty>(VariableOwner, VariableName);
+		VariableProperty && VariableProperty->IsA<FArrayProperty>();
 
 	// Only display SCS node variable names in the tree if they have not been autogenerated
 	if (bHasValidVarName && !bIsArrayVariable)
@@ -1044,7 +1045,7 @@ FString FSCSEditorTreeNodeComponentBase::GetDisplayString() const
 		if (IsNativeComponent() && GetDefault<UEditorStyleSettings>()->bShowNativeComponentNames)
 		{
 			FStringFormatNamedArguments Args;
-			Args.Add(TEXT("VarName"), VariableName.ToString());
+			Args.Add(TEXT("VarName"), VariableProperty && VariableProperty->IsNative() ? VariableProperty->GetDisplayNameText().ToString() : VariableName.ToString());
 			Args.Add(TEXT("CompName"), ComponentTemplate->GetName());
 			return FString::Format(TEXT("{VarName} ({CompName})"), Args);
 		}
