Use the BP editor lib utility functions instead to reduce code duplpication for delete unused variables

#rnx

diff --git a/Engine/Source/Editor/Kismet/Kismet.Build.cs b/Engine/Source/Editor/Kismet/Kismet.Build.cs
--- a/Engine/Source/Editor/Kismet/Kismet.Build.cs
+++ b/Engine/Source/Editor/Kismet/Kismet.Build.cs
@@ -42,6 +42,7 @@ public class Kismet : ModuleRules
 				"KismetWidgets",
 				"KismetCompiler",
 				"BlueprintGraph",
+				"BlueprintEditorLibrary",
 				"AnimGraph",
 				"PropertyEditor",
 				"SourceControl",
diff --git a/Engine/Source/Editor/Kismet/Private/BlueprintEditor.cpp b/Engine/Source/Editor/Kismet/Private/BlueprintEditor.cpp
--- a/Engine/Source/Editor/Kismet/Private/BlueprintEditor.cpp
+++ b/Engine/Source/Editor/Kismet/Private/BlueprintEditor.cpp
@@ -100,6 +100,7 @@
 #include "IDetailsView.h"
 #include "HAL/PlatformApplicationMisc.h"
 #include "Stats/StatsHierarchical.h"
+#include "BlueprintEditorLibrary.h"
 
 #include "BlueprintEditorTabs.h"
 
@@ -3636,33 +3637,9 @@ void FBlueprintEditor::DeleteUnusedVariables_OnClicked()
 {
 	UBlueprint* BlueprintObj = GetBlueprintObj();
 	
-	bool bHasAtLeastOneVariableToCheck = false;
+	// Gather FProperties from this BP and see if we can remove any
 	TArray<FProperty*> VariableProperties;
-	for (TFieldIterator<FProperty> PropertyIt(BlueprintObj->SkeletonGeneratedClass, EFieldIteratorFlags::ExcludeSuper); PropertyIt; ++PropertyIt)
-	{
-		FProperty* Property = *PropertyIt;
-		// Don't show delegate properties, there is special handling for these
-		const bool bDelegateProp = Property->IsA(FDelegateProperty::StaticClass()) || Property->IsA(FMulticastDelegateProperty::StaticClass());
-		const bool bShouldShowProp = (!Property->HasAnyPropertyFlags(CPF_Parm) && Property->HasAllPropertyFlags(CPF_BlueprintVisible) && !bDelegateProp);
-
-		if (bShouldShowProp)
-		{
-			bHasAtLeastOneVariableToCheck = true;
-			FName VarName = Property->GetFName();
-			
-			const int32 VarInfoIndex = FBlueprintEditorUtils::FindNewVariableIndex(BlueprintObj, VarName);
-			const bool bHasVarInfo = (VarInfoIndex != INDEX_NONE);
-			
-			const FObjectPropertyBase* ObjectProperty = CastField<const FObjectPropertyBase>(Property);
-			bool bIsTimeline = ObjectProperty &&
-				ObjectProperty->PropertyClass &&
-				ObjectProperty->PropertyClass->IsChildOf(UTimelineComponent::StaticClass());
-			if (!FBlueprintEditorUtils::IsVariableUsed(BlueprintObj, VarName) && !bIsTimeline && bHasVarInfo)
-			{
-				VariableProperties.Add(Property);
-			}
-		}
-	}
+	bool bHasAtLeastOneVariableToCheck = UBlueprintEditorLibrary::GatherUnusedVariables(BlueprintObj, VariableProperties);
 
 	if (VariableProperties.Num() > 0)
 	{
