My Blueprint and the Components list in both the level editor's and blueprint editor details views now use SEditorHeaderButton.

diff --git a/Engine/Source/Editor/Kismet/Private/SMyBlueprint.cpp b/Engine/Source/Editor/Kismet/Private/SMyBlueprint.cpp
--- a/Engine/Source/Editor/Kismet/Private/SMyBlueprint.cpp
+++ b/Engine/Source/Editor/Kismet/Private/SMyBlueprint.cpp
@@ -27,9 +27,7 @@
 #include "K2Node_EventNodeInterface.h"
 #include "ScopedTransaction.h"
 #include "HAL/PlatformApplicationMisc.h"
-
 #include "DetailLayoutBuilder.h"
-
 #include "SKismetInspector.h"
 #include "SSCSEditor.h"
 #include "GraphEditorDragDropAction.h"
@@ -39,12 +37,8 @@
 #include "SBlueprintPalette.h"
 #include "BlueprintEditorCommands.h"
 #include "GraphEditorActions.h"
-
 #include "AnimationGraph.h"
-
-
 #include "SBlueprintEditorToolbar.h"
-
 #include "Kismet2/BlueprintEditorUtils.h"
 #include "ObjectEditorUtils.h"
 #include "GraphEditor/Private/GraphActionNode.h"
@@ -52,13 +46,12 @@
 #include "EditorCategoryUtils.h"
 #include "Widgets/Input/SSearchBox.h"
 #include "Framework/Commands/GenericCommands.h"
-
 #include "BlueprintEditorSettings.h"
 #include "SReplaceNodeReferences.h"
 #include "ReplaceNodeReferencesHelper.h"
 #include "Animation/AnimClassInterface.h"
-
 #include "BPGraphClipboardData.h"
+#include "SEditorHeaderButton.h"
 
 #define LOCTEXT_NAMESPACE "MyBlueprint"
 
@@ -370,39 +363,13 @@ void SMyBlueprint::Construct(const FArguments& InArgs, TWeakPtr<FBlueprintEditor
 		ToolbarBuilderWidget = SNew(SBox);
 	}
 
-	TSharedPtr<SWidget> AddNewMenu = SNullWidget::NullWidget;
-
-	AddNewMenu = SNew(SComboButton)
-		.ComboButtonStyle(&FAppStyle::Get().GetWidgetStyle<FComboButtonStyle>("SimpleComboButton"))
-		.ForegroundColor(FSlateColor::UseStyle())
-		.ContentPadding(2)
-		.OnGetMenuContent(this, &SMyBlueprint::CreateAddNewMenuWidget)
-		.HasDownArrow(false)
+	TSharedPtr<SWidget> AddNewMenu = SNew(SEditorHeaderButton)
 		.AddMetaData<FTagMetaData>(FTagMetaData(TEXT("MyBlueprintAddNewCombo")))
-		.IsEnabled(this, &SMyBlueprint::IsEditingMode)
+		.Icon(FAppStyle::Get().GetBrush("Icons.Plus"))
+		.Text(LOCTEXT("AddNewLabel", "Add"))
 		.ToolTipText(LOCTEXT("AddNewToolTip", "Add a new Variable, Graph, Function, Macro, or Event Dispatcher."))
-		.ButtonContent()
-		[
-			SNew(SHorizontalBox)
-			+ SHorizontalBox::Slot()
-			.AutoWidth()
-			.HAlign(HAlign_Center)
-			.VAlign(VAlign_Center)
-			[
-				SNew(SImage)
-				.Image(FAppStyle::Get().GetBrush("Icons.Plus"))
-				.ColorAndOpacity(FSlateColor::UseForeground())
-			]
-			+ SHorizontalBox::Slot()
-			.Padding(FMargin(3, 0, 0, 0))
-			.VAlign(VAlign_Center)
-			.AutoWidth()
-			[
-				SNew(STextBlock)
-				.TextStyle(FAppStyle::Get(), "NormalText")
-				.Text(LOCTEXT("NewAssetButton", "New"))
-			]
-		];
+		.IsEnabled(this, &SMyBlueprint::IsEditingMode)
+		.OnGetMenuContent(this, &SMyBlueprint::CreateAddNewMenuWidget);
 
 	FMenuBuilder ViewOptions(true, nullptr);
 
@@ -527,7 +494,8 @@ void SMyBlueprint::Construct(const FArguments& InArgs, TWeakPtr<FBlueprintEditor
 
 					+ SHorizontalBox::Slot()
 					.AutoWidth()
-					.Padding(0, 0, 2, 0)
+					.VAlign(VAlign_Center)
+					.Padding(0, 0, 4, 0)
 					[
 						AddNewMenu.ToSharedRef()
 					]
diff --git a/Engine/Source/Editor/Kismet/Private/SSCSEditor.cpp b/Engine/Source/Editor/Kismet/Private/SSCSEditor.cpp
--- a/Engine/Source/Editor/Kismet/Private/SSCSEditor.cpp
+++ b/Engine/Source/Editor/Kismet/Private/SSCSEditor.cpp
@@ -85,6 +85,7 @@
 #include "ISCSEditorUICustomization.h"
 
 #include "Logging/MessageLog.h"
+#include "SEditorHeaderButton.h"
 
 #define LOCTEXT_NAMESPACE "SSCSEditor"
 
@@ -3984,6 +3985,8 @@ void SSCSEditor::Construct( const FArguments& InArgs )
 
 	ButtonBox = SNew(SHorizontalBox)
 	+ SHorizontalBox::Slot()
+	.VAlign(VAlign_Center)
+	.Padding(0,0,4,0)
 	.AutoWidth()
 	[
 		SNew(SComponentClassCombo)
@@ -3994,7 +3997,8 @@ void SSCSEditor::Construct( const FArguments& InArgs )
 		.IsEnabled(AllowEditing)
 	]
 	+ SHorizontalBox::Slot()
-	.Padding(0)
+	.VAlign(VAlign_Center)
+	.Padding(0,0,4,0)
 	.AutoWidth()
 	[
 		SAssignNew(ExtensionPanel, SExtensionPanel)
@@ -4005,41 +4009,30 @@ void SSCSEditor::Construct( const FArguments& InArgs )
 	// horizontal slot index #2 => reserved for BP-editor search bar (see 'ButtonBox' and 'SearchBarHorizontalSlotIndex' usage below)
 
 	+ SHorizontalBox::Slot()
+	.VAlign(VAlign_Center)
+	.Padding(0,0,4,0)
 	.AutoWidth()
 	[
-		SNew( SButton )
-		.ButtonStyle( FEditorStyle::Get(), "SimpleButton" )
-		.AddMetaData<FTagMetaData>(FTagMetaData(TEXT("Actor.ConvertToBlueprint")))
+		SNew(SEditorHeaderButton)
+		.AddMetaData<FTagMetaData>( FTagMetaData(TEXT("Actor.ConvertToBlueprint")) )
 		.Visibility( this, &SSCSEditor::GetPromoteToBlueprintButtonVisibility )
 		.OnClicked( this, &SSCSEditor::OnPromoteToBlueprintClicked )
-		.ContentPadding(FMargin(5.f))
+		.Icon(FAppStyle::Get().GetBrush("Icons.Blueprints"))
 		.ToolTip(IDocumentation::Get()->CreateToolTip(
-			LOCTEXT("PromoteToBluerprintTooltip","Converts this actor into a reusable Blueprint Class that can have script behavior" ),
+			LOCTEXT("PromoteToBluerprintTooltip", "Converts this actor into a reusable Blueprint Class that can have script behavior" ),
 			NULL,
 			TEXT("Shared/LevelEditor"),
 			TEXT("ConvertToBlueprint")))
-		[
-			SNew(SImage)
-			.ColorAndOpacity(FSlateColor::UseForeground())
-			.Image(FAppStyle::Get().GetBrush("Icons.Blueprints"))
-		]
 	]
 	+ SHorizontalBox::Slot()
+	.VAlign(VAlign_Center)
 	.AutoWidth()
 	[
-		SNew(SComboButton)
+		SNew(SEditorHeaderButton)
 		.AddMetaData<FTagMetaData>(FTagMetaData(TEXT("Actor.EditBlueprint")))
 		.Visibility(this, &SSCSEditor::GetEditBlueprintButtonVisibility)
-		.ComboButtonStyle(FEditorStyle::Get(), "SimpleComboButton")
-		.HasDownArrow(false)
-		.ContentPadding(FMargin(5.f))
 		.ToolTipText(LOCTEXT("EditActorBlueprint_Tooltip", "Edit the Blueprint for this Actor"))
-		.ButtonContent()
-		[
-			SNew(SImage)
-			.ColorAndOpacity(FSlateColor::UseForeground())
-			.Image(FAppStyle::Get().GetBrush("Icons.Blueprints"))
-		]
+		.Icon(FAppStyle::Get().GetBrush("Icons.Blueprints"))
 		.MenuContent()
 		[
 			EditBlueprintMenuBuilder.MakeWidget()
@@ -4077,8 +4070,7 @@ void SSCSEditor::Construct( const FArguments& InArgs )
 	// Only insert the buttons and search bar in the Blueprints version
 	if (bInlineSearchBarWithButtons) // Blueprints
 	{
-		const int32 SearchBarHorizontalSlotIndex = 0;
-		ButtonBox->InsertSlot(SearchBarHorizontalSlotIndex)
+		ButtonBox->AddSlot()
 			.FillWidth(1.0f)
 			.VAlign(VAlign_Center)
 			.Padding(3.0f, 3.0f)
@@ -4087,6 +4079,8 @@ void SSCSEditor::Construct( const FArguments& InArgs )
 		];
 
 		HeaderBox->AddSlot()
+			.VAlign(VAlign_Center)
+			.AutoHeight()
 		[
 			ButtonBox.ToSharedRef()
 		];
diff --git a/Engine/Source/Editor/UnrealEd/Private/SComponentClassCombo.cpp b/Engine/Source/Editor/UnrealEd/Private/SComponentClassCombo.cpp
--- a/Engine/Source/Editor/UnrealEd/Private/SComponentClassCombo.cpp
+++ b/Engine/Source/Editor/UnrealEd/Private/SComponentClassCombo.cpp
@@ -11,7 +11,6 @@
 #include "Engine/Blueprint.h"
 #include "Engine/Selection.h"
 #include "Editor.h"
-
 #include "Styling/SlateIconFinder.h"
 #include "ComponentAssetBroker.h"
 #include "ComponentTypeRegistry.h"
@@ -19,6 +18,7 @@
 #include "Widgets/Input/SSearchBox.h"
 #include "SListViewSelectorDropdownMenu.h"
 #include "Misc/TextFilterExpressionEvaluator.h"
+#include "SEditorHeaderButton.h"
 
 #define LOCTEXT_NAMESPACE "ComponentClassCombo"
 
@@ -64,52 +64,46 @@ void SComponentClassCombo::Construct(const FArguments& InArgs)
 		.OnTextChanged( this, &SComponentClassCombo::OnSearchBoxTextChanged )
 		.OnTextCommitted( this, &SComponentClassCombo::OnSearchBoxTextCommitted );
 
-	// Create the Construct arguments for the parent class (SComboButton)
-	SComboButton::FArguments Args;
-	Args.ComboButtonStyle(FAppStyle::Get(), "SimpleComboButton")
-	.HasDownArrow(false)
-	.IsFocusable(true)
-	.OnComboBoxOpened(this, &SComponentClassCombo::ClearSelection)
-	.ButtonContent()
-	[
-		SNew(SImage)
-		.ColorAndOpacity(FSlateColor::UseForeground())
-		.Image(FAppStyle::Get().GetBrush("Icons.Plus"))
-	]
-	.MenuContent()
+	ChildSlot
 	[
-
-		SNew(SListViewSelectorDropdownMenu<FComponentClassComboEntryPtr>, SearchBox, ComponentClassListView)
+		SAssignNew(AddNewButton, SEditorHeaderButton)
+		.Icon(FAppStyle::Get().GetBrush("Icons.Plus"))
+		.Text(LOCTEXT("Add", "Add"))
+		.OnComboBoxOpened(this, &SComponentClassCombo::ClearSelection)
+		.MenuContent()
 		[
-			SNew(SBorder)
-			.BorderImage(FEditorStyle::GetBrush("Menu.Background"))
-			.Padding(2)
+			SNew(SListViewSelectorDropdownMenu<FComponentClassComboEntryPtr>, SearchBox, ComponentClassListView)
 			[
-				SNew(SBox)
-				.WidthOverride(250)
-				[				
-					SNew(SVerticalBox)
-					+SVerticalBox::Slot()
-					.Padding(1.f)
-					.AutoHeight()
-					[
-						SearchBox.ToSharedRef()
-					]
-					+SVerticalBox::Slot()
-					.MaxHeight(400)
-					[
-						ComponentClassListView.ToSharedRef()
+				SNew(SBorder)
+				.BorderImage(FEditorStyle::GetBrush("Menu.Background"))
+				.Padding(2)
+				[
+					SNew(SBox)
+					.WidthOverride(250)
+					[				
+						SNew(SVerticalBox)
+						+SVerticalBox::Slot()
+						.Padding(1.f)
+						.AutoHeight()
+						[
+							SearchBox.ToSharedRef()
+						]
+						+SVerticalBox::Slot()
+						.MaxHeight(400)
+						[
+							ComponentClassListView.ToSharedRef()
+						]
 					]
 				]
 			]
 		]
 	];
 
-	SComboButton::Construct(Args);
+	
 
 	ComponentClassListView->EnableToolTipForceField( true );
-	// The base class can automatically handle setting focus to a specified control when the combo button is opened
-	SetMenuContentWidgetToFocus( SearchBox );
+	// The button can automatically handle setting focus to a specified control when the combo button is opened
+	AddNewButton->SetMenuContentWidgetToFocus( SearchBox );
 }
 
 SComponentClassCombo::~SComponentClassCombo()
@@ -245,7 +239,7 @@ void SComponentClassCombo::OnAddComponentSelectionChanged( FComponentClassComboE
 		if ( InItem->IsClass() )
 		{
 			// Neither do we want the combo dropdown staying open once the user has clicked on a valid option
-			SetIsOpen(false, false);
+			AddNewButton->SetIsMenuOpen(false, false);
 
 			if( OnComponentClassSelected.IsBound() )
 			{
diff --git a/Engine/Source/Editor/UnrealEd/Public/SComponentClassCombo.h b/Engine/Source/Editor/UnrealEd/Public/SComponentClassCombo.h
--- a/Engine/Source/Editor/UnrealEd/Public/SComponentClassCombo.h
+++ b/Engine/Source/Editor/UnrealEd/Public/SComponentClassCombo.h
@@ -167,7 +167,7 @@ private:
 
 //////////////////////////////////////////////////////////////////////////
 
-class UNREALED_API SComponentClassCombo : public SComboButton
+class UNREALED_API SComponentClassCombo : public SCompoundWidget
 {
 public:
 	SLATE_BEGIN_ARGS( SComponentClassCombo )
@@ -232,6 +232,9 @@ private:
 	/** The search box control - part of the combo drop down */
 	TSharedPtr<SSearchBox> SearchBox;
 
+	/** The Add combo button. */
+	TSharedPtr<class SEditorHeaderButton> AddNewButton;
+
 	/** The component list control - part of the combo drop down */
 	TSharedPtr< SListView<FComponentClassComboEntryPtr> > ComponentClassListView;
 
