Fix merge error that caused a static analysis warning

#rnx

diff --git a/Engine/Source/Editor/Kismet/Private/BlueprintEditor.cpp b/Engine/Source/Editor/Kismet/Private/BlueprintEditor.cpp
--- a/Engine/Source/Editor/Kismet/Private/BlueprintEditor.cpp
+++ b/Engine/Source/Editor/Kismet/Private/BlueprintEditor.cpp
@@ -3717,7 +3717,7 @@ void FBlueprintEditor::DeleteUnusedVariables_OnClicked()
 			}
 		}
 	}
-	else if (VariableProperties.Num() > 0)
+	else if (bHasAtLeastOneVariableToCheck)
 	{
 		LogSimpleMessage(LOCTEXT("AllVariablesInUseMessage", "All variables are currently in use."));
 	}
