BlendSpace 2.0: Blendspace Graph Node

diff --git a/Engine/Source/Editor/AnimationBlueprintEditor/Private/AnimationBlueprintEditorMode.cpp b/Engine/Source/Editor/AnimationBlueprintEditor/Private/AnimationBlueprintEditorMode.cpp
--- a/Engine/Source/Editor/AnimationBlueprintEditor/Private/AnimationBlueprintEditorMode.cpp
+++ b/Engine/Source/Editor/AnimationBlueprintEditor/Private/AnimationBlueprintEditorMode.cpp
@@ -129,8 +129,8 @@ FAnimationBlueprintEditorMode::FAnimationBlueprintEditorMode(const TSharedRef<FA
 	TabFactories.RegisterFactory(PersonaModule.CreateAnimationAssetBrowserTabFactory(InAnimationBlueprintEditor, InAnimationBlueprintEditor->GetPersonaToolkit(), FOnOpenNewAsset::CreateSP(&InAnimationBlueprintEditor.Get(), &FAnimationBlueprintEditor::HandleOpenNewAsset), FOnAnimationSequenceBrowserCreated(), true));
 	TabFactories.RegisterFactory(PersonaModule.CreateAnimBlueprintPreviewTabFactory(InAnimationBlueprintEditor, InAnimationBlueprintEditor->GetPersonaToolkit()->GetPreviewScene()));
 	TabFactories.RegisterFactory(PersonaModule.CreateAnimBlueprintAssetOverridesTabFactory(InAnimationBlueprintEditor, InAnimationBlueprintEditor->GetPersonaToolkit()->GetAnimBlueprint(), InAnimationBlueprintEditor->OnPostUndo));
-	TabFactories.RegisterFactory(PersonaModule.CreateSkeletonSlotNamesTabFactory(InAnimationBlueprintEditor, InAnimationBlueprintEditor->GetSkeletonTree()->GetEditableSkeleton(), InAnimationBlueprintEditor->OnPostUndo, FOnObjectSelected::CreateSP(&InAnimationBlueprintEditor.Get(), &FAnimationBlueprintEditor::HandleObjectSelected)));
-	TabFactories.RegisterFactory(PersonaModule.CreateCurveViewerTabFactory(InAnimationBlueprintEditor, InAnimationBlueprintEditor->GetSkeletonTree()->GetEditableSkeleton(), InAnimationBlueprintEditor->GetPersonaToolkit()->GetPreviewScene(), InAnimationBlueprintEditor->OnPostUndo, FOnObjectsSelected::CreateSP(&InAnimationBlueprintEditor.Get(), &FAnimationBlueprintEditor::HandleObjectsSelected)));
+	TabFactories.RegisterFactory(PersonaModule.CreateSkeletonSlotNamesTabFactory(InAnimationBlueprintEditor, InAnimationBlueprintEditor->GetSkeletonTree()->GetEditableSkeleton(), FOnObjectSelected::CreateSP(&InAnimationBlueprintEditor.Get(), &FAnimationBlueprintEditor::HandleObjectSelected)));
+	TabFactories.RegisterFactory(PersonaModule.CreateCurveViewerTabFactory(InAnimationBlueprintEditor, InAnimationBlueprintEditor->GetSkeletonTree()->GetEditableSkeleton(), InAnimationBlueprintEditor->GetPersonaToolkit()->GetPreviewScene(), FOnObjectsSelected::CreateSP(&InAnimationBlueprintEditor.Get(), &FAnimationBlueprintEditor::HandleObjectsSelected)));
 
 	// setup toolbar - clear existing toolbar extender from the BP mode
 	//@TODO: Keep this in sync with BlueprintEditorModes.cpp
diff --git a/Engine/Source/Editor/AnimationBlueprintEditor/Private/AnimationGraphFactory.cpp b/Engine/Source/Editor/AnimationBlueprintEditor/Private/AnimationGraphFactory.cpp
--- a/Engine/Source/Editor/AnimationBlueprintEditor/Private/AnimationGraphFactory.cpp
+++ b/Engine/Source/Editor/AnimationBlueprintEditor/Private/AnimationGraphFactory.cpp
@@ -33,6 +33,8 @@
 #include "StateMachineConnectionDrawingPolicy.h"
 
 #include "KismetPins/SGraphPinExec.h"
+#include "AnimGraphNode_BlendSpaceGraph.h"
+#include "AnimationNodes/SGraphNodeBlendSpaceGraph.h"
 
 TSharedPtr<class SGraphNode> FAnimationGraphNodeFactory::CreateNode(class UEdGraphNode* InNode) const 
 {
@@ -58,6 +60,10 @@ TSharedPtr<class SGraphNode> FAnimationGraphNodeFactory::CreateNode(class UEdGra
 		{
 			return SNew(SGraphNodeBlendSpacePlayer, BlendSpacePlayer);
 		}
+		else if (UAnimGraphNode_BlendSpaceGraphBase* BlendSpaceGraph = Cast<UAnimGraphNode_BlendSpaceGraphBase>(InNode))
+		{
+			return SNew(SGraphNodeBlendSpaceGraph, BlendSpaceGraph);
+		}
 		else
 		{
 			return SNew(SAnimationGraphNode, BaseAnimNode);
diff --git a/Engine/Source/Editor/AnimationBlueprintEditor/Private/AnimationNodes/SAnimationGraphNode.cpp b/Engine/Source/Editor/AnimationBlueprintEditor/Private/AnimationNodes/SAnimationGraphNode.cpp
--- a/Engine/Source/Editor/AnimationBlueprintEditor/Private/AnimationNodes/SAnimationGraphNode.cpp
+++ b/Engine/Source/Editor/AnimationBlueprintEditor/Private/AnimationNodes/SAnimationGraphNode.cpp
@@ -132,7 +132,7 @@ void SAnimationGraphNode::Construct(const FArguments& InArgs, UAnimGraphNode_Bas
 
 	this->UpdateGraphNode();
 
-	ReconfigurePinWidgetsForPropertyBindings();
+	ReconfigurePinWidgetsForPropertyBindings(CastChecked<UAnimGraphNode_Base>(GraphNode), [this](UEdGraphPin* InPin){ return FindWidgetForPin(InPin); });
 
 	const FSlateBrush* ImageBrush = FEditorStyle::Get().GetBrush(TEXT("Graph.AnimationFastPathIndicator"));
 
@@ -271,16 +271,14 @@ void SAnimationGraphNode::GetNodeInfoPopups(FNodeInfoContext* Context, TArray<FG
 	}
 }
 
-void SAnimationGraphNode::ReconfigurePinWidgetsForPropertyBindings()
+void SAnimationGraphNode::ReconfigurePinWidgetsForPropertyBindings(UAnimGraphNode_Base* InAnimGraphNode, TFunctionRef<TSharedPtr<SGraphPin>(UEdGraphPin*)> InFindWidgetForPin)
 {
-	UAnimGraphNode_Base* AnimGraphNode = CastChecked<UAnimGraphNode_Base>(GraphNode);
-
-	for(UEdGraphPin* Pin : AnimGraphNode->Pins)
+	for(UEdGraphPin* Pin : InAnimGraphNode->Pins)
 	{
 		FEdGraphPinType PinType = Pin->PinType;
 		if(Pin->Direction == EGPD_Input && !UAnimationGraphSchema::IsPosePin(PinType))
 		{
-			TSharedPtr<SGraphPin> PinWidget = FindWidgetForPin(Pin);
+			TSharedPtr<SGraphPin> PinWidget = InFindWidgetForPin(Pin);
 
 			if(PinWidget.IsValid())
 			{
@@ -293,7 +291,7 @@ void SAnimationGraphNode::ReconfigurePinWidgetsForPropertyBindings()
 				{
 					TWeakPtr<SGraphPin> WeakPinWidget = PinWidget;
 
-					PinWidget->GetValueWidget()->SetVisibility(MakeAttributeLambda([ComparisonName, AnimGraphNode, WeakPinWidget]()
+					PinWidget->GetValueWidget()->SetVisibility(MakeAttributeLambda([ComparisonName, InAnimGraphNode, WeakPinWidget]()
 					{
 						EVisibility Visibility = EVisibility::Collapsed;
 
@@ -301,7 +299,7 @@ void SAnimationGraphNode::ReconfigurePinWidgetsForPropertyBindings()
 						{
 							Visibility = WeakPinWidget.Pin()->GetDefaultValueVisibility();
 
-							if (FAnimGraphNodePropertyBinding* BindingPtr = AnimGraphNode->PropertyBindings.Find(ComparisonName))
+							if (FAnimGraphNodePropertyBinding* BindingPtr = InAnimGraphNode->PropertyBindings.Find(ComparisonName))
 							{
 								Visibility = EVisibility::Collapsed;
 							}
@@ -315,18 +313,18 @@ void SAnimationGraphNode::ReconfigurePinWidgetsForPropertyBindings()
 				PinWidget->GetLabelAndValue()->AddSlot()
 				[
 					SNew(SHorizontalBox)
-					.ToolTipText_Lambda([ComparisonName, AnimGraphNode]()
+					.ToolTipText_Lambda([ComparisonName, InAnimGraphNode]()
 					{
-						if (FAnimGraphNodePropertyBinding* BindingPtr = AnimGraphNode->PropertyBindings.Find(ComparisonName))
+						if (FAnimGraphNodePropertyBinding* BindingPtr = InAnimGraphNode->PropertyBindings.Find(ComparisonName))
 						{
 							return FText::Format(LOCTEXT("BindingTooltipFormat", "Pin is bound to property '{0}'"), BindingPtr->PathAsText);
 						}
 
 						return FText::GetEmpty();
 					})
-					.Visibility_Lambda([ComparisonName, AnimGraphNode]()
+					.Visibility_Lambda([ComparisonName, InAnimGraphNode]()
 					{
-						return AnimGraphNode->PropertyBindings.Contains(ComparisonName) ? EVisibility::Visible : EVisibility::Collapsed;
+						return InAnimGraphNode->PropertyBindings.Contains(ComparisonName) ? EVisibility::Visible : EVisibility::Collapsed;
 					})
 					+SHorizontalBox::Slot()
 					.AutoWidth()
@@ -334,9 +332,9 @@ void SAnimationGraphNode::ReconfigurePinWidgetsForPropertyBindings()
 					.Padding(3.0f, 2.0f)
 					[
 						SNew(SImage)
-						.Image_Lambda([ComparisonName, AnimGraphNode, PinType]() -> const FSlateBrush*
+						.Image_Lambda([ComparisonName, InAnimGraphNode, PinType]() -> const FSlateBrush*
 						{
-							if (FAnimGraphNodePropertyBinding* BindingPtr = AnimGraphNode->PropertyBindings.Find(ComparisonName))
+							if (FAnimGraphNodePropertyBinding* BindingPtr = InAnimGraphNode->PropertyBindings.Find(ComparisonName))
 							{
 								static FName FunctionIcon(TEXT("GraphEditor.Function_16x"));
 
@@ -345,11 +343,11 @@ void SAnimationGraphNode::ReconfigurePinWidgetsForPropertyBindings()
 
 							return nullptr;
 						})
-						.ColorAndOpacity_Lambda([AnimGraphNode, ComparisonName]()
+						.ColorAndOpacity_Lambda([InAnimGraphNode, ComparisonName]()
 						{
-							if(const UEdGraphSchema* Schema = AnimGraphNode->GetSchema())
+							if(const UEdGraphSchema* Schema = InAnimGraphNode->GetSchema())
 							{
-								if (FAnimGraphNodePropertyBinding* BindingPtr = AnimGraphNode->PropertyBindings.Find(ComparisonName))
+								if (FAnimGraphNodePropertyBinding* BindingPtr = InAnimGraphNode->PropertyBindings.Find(ComparisonName))
 								{
 									return Schema->GetPinTypeColor(BindingPtr->bIsPromotion ? BindingPtr->PromotedPinType : BindingPtr->PinType);
 								}
@@ -363,9 +361,9 @@ void SAnimationGraphNode::ReconfigurePinWidgetsForPropertyBindings()
 					.Padding(3.0f, 2.0f)
 					[
 						SNew(STextBlock)
-						.Text_Lambda([ComparisonName, AnimGraphNode]()
+						.Text_Lambda([ComparisonName, InAnimGraphNode]()
 						{
-							if (const FAnimGraphNodePropertyBinding* BindingPtr = AnimGraphNode->PropertyBindings.Find(ComparisonName))
+							if (const FAnimGraphNodePropertyBinding* BindingPtr = InAnimGraphNode->PropertyBindings.Find(ComparisonName))
 							{
 								return BindingPtr->PathAsText;
 							}
diff --git a/Engine/Source/Editor/AnimationBlueprintEditor/Private/AnimationNodes/SBlendSpacePreview.h b/Engine/Source/Editor/AnimationBlueprintEditor/Private/AnimationNodes/SBlendSpacePreview.h
new file mode 100644
--- /dev/null
+++ b/Engine/Source/Editor/AnimationBlueprintEditor/Private/AnimationNodes/SBlendSpacePreview.h
@@ -0,0 +1,31 @@
+// Copyright Epic Games, Inc. All Rights Reserved.
+
+#pragma once
+
+#include "CoreMinimal.h"
+#include "Layout/Visibility.h"
+#include "Widgets/DeclarativeSyntaxSupport.h"
+#include "PersonaDelegates.h"
+
+class UBlendSpaceBase;
+class UAnimGraphNode_Base;
+
+class SBlendSpacePreview : public SCompoundWidget
+{
+public:
+	SLATE_BEGIN_ARGS(SBlendSpacePreview){}
+
+	SLATE_ARGUMENT(FOnGetBlendSpaceSampleName, OnGetBlendSpaceSampleName)
+
+	SLATE_END_ARGS()
+
+	void Construct(const FArguments& InArgs, UAnimGraphNode_Base* InNode);
+
+protected:
+	EVisibility GetBlendSpaceVisibility() const;
+	bool GetBlendSpaceInfo(TWeakObjectPtr<const UBlendSpaceBase>& OutBlendSpace, FVector& OutPosition) const;
+
+	TWeakObjectPtr<const UAnimGraphNode_Base> Node;
+	TWeakObjectPtr<const UBlendSpaceBase> CachedBlendSpace;
+	FVector CachedPosition;
+};
diff --git a/Engine/Source/Editor/AnimationBlueprintEditor/Private/AnimationNodes/SGraphNodeBlendSpacePlayer.cpp b/Engine/Source/Editor/AnimationBlueprintEditor/Private/AnimationNodes/SGraphNodeBlendSpacePlayer.cpp
--- a/Engine/Source/Editor/AnimationBlueprintEditor/Private/AnimationNodes/SGraphNodeBlendSpacePlayer.cpp
+++ b/Engine/Source/Editor/AnimationBlueprintEditor/Private/AnimationNodes/SGraphNodeBlendSpacePlayer.cpp
@@ -1,15 +1,10 @@
 // Copyright Epic Games, Inc. All Rights Reserved.
 
 #include "AnimationNodes/SGraphNodeBlendSpacePlayer.h"
-#include "Animation/AnimBlueprintGeneratedClass.h"
-#include "Kismet2/KismetDebugUtilities.h"
-#include "Kismet2/BlueprintEditorUtils.h"
-#include "AnimNodes/AnimNode_BlendSpacePlayer.h"
-#include "AnimGraphNode_BlendSpaceBase.h"
-#include "PersonaModule.h"
-#include "Widgets/Layout/SBox.h"
-#include "Modules/ModuleManager.h"
+#include "SBlendSpacePreview.h"
+#include "AnimGraphNode_Base.h"
 #include "AnimGraphNode_BlendSpacePlayer.h"
+#include "Kismet2/BlueprintEditorUtils.h"
 
 void SGraphNodeBlendSpacePlayer::Construct(const FArguments& InArgs, UAnimGraphNode_Base* InNode)
 {
@@ -25,7 +20,6 @@ void SGraphNodeBlendSpacePlayer::Construct(const FArguments& InArgs, UAnimGraphN
 
 	RegisterActiveTimer(1.0f / 60.0f, FWidgetActiveTimerDelegate::CreateLambda([this](double InCurrentTime, float InDeltaTime)
 	{
-		GetBlendSpaceInfo(CachedBlendSpace, CachedPosition);
 		UpdateGraphSyncLabel();
 		return EActiveTimerReturnType::Continue;
 	}));
@@ -33,89 +27,15 @@ void SGraphNodeBlendSpacePlayer::Construct(const FArguments& InArgs, UAnimGraphN
 
 void SGraphNodeBlendSpacePlayer::CreateBelowWidgetControls(TSharedPtr<SVerticalBox> MainBox)
 {
-	FPersonaModule& PersonaModule = FModuleManager::LoadModuleChecked<FPersonaModule>("Persona");
-
 	MainBox->AddSlot()
 	.AutoHeight()
 	.VAlign(VAlign_Fill)
 	.Padding(0.0f)
 	[
-		SNew(SBox)
-		.HeightOverride_Lambda([WeakMainBox = TWeakPtr<SVerticalBox>(MainBox)]()
-		{
-			if(TSharedPtr<SVerticalBox> LocalMainBox = WeakMainBox.Pin())
-			{
-				float Size = LocalMainBox->GetDesiredSize().X;
-				return FMath::Min(Size, 100.0f);
-			}
-
-			return 0.0f;
-		})
-		.Visibility(this, &SGraphNodeBlendSpacePlayer::GetBlendSpaceVisibility)
-		[
-			PersonaModule.CreateBlendSpacePreviewWidget(
-			MakeAttributeLambda([this]()
-			{
-				return CachedBlendSpace.Get();
-			}),
-			MakeAttributeLambda([this]()
-			{
-				return CachedPosition;
-			}))
-		]
+		SNew(SBlendSpacePreview, CastChecked<UAnimGraphNode_Base>(GraphNode))
 	];
 }
 
-EVisibility SGraphNodeBlendSpacePlayer::GetBlendSpaceVisibility() const
-{
-	if (UBlueprint* Blueprint = FBlueprintEditorUtils::FindBlueprintForNode(GraphNode))
-	{
-		if (FProperty* Property = FKismetDebugUtilities::FindClassPropertyForNode(Blueprint, GraphNode))
-		{
-			if (UObject* ActiveObject = Blueprint->GetObjectBeingDebugged())
-			{
-				return EVisibility::Visible;
-			}
-		}
-	}
-
-	return EVisibility::Collapsed;
-}
-
-bool SGraphNodeBlendSpacePlayer::GetBlendSpaceInfo(TWeakObjectPtr<const UBlendSpaceBase>& OutBlendSpace, FVector& OutPosition) const
-{
-	if (UBlueprint* Blueprint = FBlueprintEditorUtils::FindBlueprintForNode(GraphNode))
-	{
-		if (UObject* ActiveObject = Blueprint->GetObjectBeingDebugged())
-		{
-			if (UAnimGraphNode_BlendSpaceBase* VisualBlendSpacePlayer = Cast<UAnimGraphNode_BlendSpaceBase>(GraphNode))
-			{
-				if (UAnimBlueprintGeneratedClass* Class = Cast<UAnimBlueprintGeneratedClass>((UObject*)ActiveObject->GetClass()))
-				{
-					if(int32* NodeIndexPtr = Class->GetAnimBlueprintDebugData().NodePropertyToIndexMap.Find(TWeakObjectPtr<UAnimGraphNode_Base>(Cast<UAnimGraphNode_Base>(GraphNode))))
-					{
-						int32 AnimNodeIndex = *NodeIndexPtr;
-						// reverse node index temporarily because of a bug in NodeGuidToIndexMap
-						AnimNodeIndex = Class->GetAnimNodeProperties().Num() - AnimNodeIndex - 1;
-
-						if (FAnimBlueprintDebugData::FBlendSpacePlayerRecord* DebugInfo = Class->GetAnimBlueprintDebugData().BlendSpacePlayerRecordsThisFrame.FindByPredicate([AnimNodeIndex](const FAnimBlueprintDebugData::FBlendSpacePlayerRecord& InRecord){ return InRecord.NodeID == AnimNodeIndex; }))
-						{
-							OutBlendSpace = DebugInfo->BlendSpace.Get();
-							OutPosition = FVector(DebugInfo->PositionX, DebugInfo->PositionY, DebugInfo->PositionZ);
-							return true;
-						}
-					}
-				}
-			}
-		}
-	}
-
-	OutBlendSpace = nullptr;
-	OutPosition = FVector::ZeroVector;
-	return false;
-}
-
-
 void SGraphNodeBlendSpacePlayer::UpdateGraphSyncLabel()
 {
 	if (UAnimGraphNode_BlendSpacePlayer* VisualBlendSpacePlayer = Cast<UAnimGraphNode_BlendSpacePlayer>(GraphNode))
@@ -150,4 +70,4 @@ void SGraphNodeBlendSpacePlayer::UpdateGraphSyncLabel()
 			CachedSyncGroupName = CurrentSyncGroupName;
 		}
 	}
-}
\ No newline at end of file
+}
diff --git a/Engine/Source/Editor/AnimationBlueprintEditor/Private/AnimationNodes/SGraphNodeBlendSpacePlayer.h b/Engine/Source/Editor/AnimationBlueprintEditor/Private/AnimationNodes/SGraphNodeBlendSpacePlayer.h
--- a/Engine/Source/Editor/AnimationBlueprintEditor/Private/AnimationNodes/SGraphNodeBlendSpacePlayer.h
+++ b/Engine/Source/Editor/AnimationBlueprintEditor/Private/AnimationNodes/SGraphNodeBlendSpacePlayer.h
@@ -25,15 +25,9 @@ public:
 	// End of SGraphNode interface
 
 protected:
-	EVisibility GetBlendSpaceVisibility() const;
-	bool GetBlendSpaceInfo(TWeakObjectPtr<const UBlendSpaceBase>& OutBlendSpace, FVector& OutPosition) const;
-
 	// Invalidates the node's label if we are syncing based on graph context
 	void UpdateGraphSyncLabel();
 
-	TWeakObjectPtr<const UBlendSpaceBase> CachedBlendSpace;
-	FVector CachedPosition;
-
 	// Cached name to display when sync groups are dynamic
 	FName CachedSyncGroupName;
 };
diff --git a/Engine/Source/Editor/AnimationBlueprintEditor/Private/AnimationPins/SGraphPinPose.cpp b/Engine/Source/Editor/AnimationBlueprintEditor/Private/AnimationPins/SGraphPinPose.cpp
--- a/Engine/Source/Editor/AnimationBlueprintEditor/Private/AnimationPins/SGraphPinPose.cpp
+++ b/Engine/Source/Editor/AnimationBlueprintEditor/Private/AnimationPins/SGraphPinPose.cpp
@@ -67,7 +67,7 @@ public:
 
 	SLATE_END_ARGS()
 
-	void Construct(const FArguments& InArgs, UAnimBlueprint* InAnimBlueprint, const FAnimGraphAttributeDesc* InAttributeDesc, const UEdGraphPin* InPin, EAttributeUsage InUsage, UAnimGraphNode_Base* InAnimGraphNode)
+	void Construct(const FArguments& InArgs, UAnimBlueprint* InAnimBlueprint, const FAnimGraphAttributeDesc* InAttributeDesc, const UEdGraphPin* InPin, EAttributeUsage InUsage, const UAnimGraphNode_Base* InAnimGraphNode)
 	{
 		AnimBlueprint = InAnimBlueprint;
 		AttributeDesc = InAttributeDesc;
@@ -149,7 +149,7 @@ public:
 	const FAnimGraphAttributeDesc* AttributeDesc;
 	const UEdGraphPin* Pin;
 	EAttributeUsage Usage;
-	UAnimGraphNode_Base* AnimGraphNode;
+	const UAnimGraphNode_Base* AnimGraphNode;
 	FLinearColor ActiveColor;
 	FLinearColor Color;
 	FLinearColor BackgroundColor;
@@ -185,10 +185,12 @@ void SGraphPinPose::ReconfigureWidgetForAttributes()
 					// Skip displaying passthrough attributes on pins
 					if(InUsage != EAttributeUsage::Passthrough && AnimGraphNode->ShouldShowAttributesOnPins())
 					{
+						const UAnimGraphNode_Base* ProxyGraphNode = AnimGraphNode->GetProxyNodeForAttributes();
+
 						GetLabelAndValue()->AddSlot()
 						.Padding(2.0f, 0.0f, 0.0f, 0.0f)
 						[
-							SNew(SAttributeIndicator, AnimBlueprint, AttributeDesc, GraphPinObj, InUsage, AnimGraphNode)
+							SNew(SAttributeIndicator, AnimBlueprint, AttributeDesc, GraphPinObj, InUsage, ProxyGraphNode)
 						];
 					}
 				}
diff --git a/Engine/Source/Editor/AnimationBlueprintEditor/Public/AnimationNodes/SAnimationGraphNode.h b/Engine/Source/Editor/AnimationBlueprintEditor/Public/AnimationNodes/SAnimationGraphNode.h
--- a/Engine/Source/Editor/AnimationBlueprintEditor/Public/AnimationNodes/SAnimationGraphNode.h
+++ b/Engine/Source/Editor/AnimationBlueprintEditor/Public/AnimationNodes/SAnimationGraphNode.h
@@ -20,6 +20,9 @@ public:
 
 	void Construct(const FArguments& InArgs, UAnimGraphNode_Base* InNode);
 
+	// Tweak any created pin widgets so they respond to bindings
+	static void ReconfigurePinWidgetsForPropertyBindings(UAnimGraphNode_Base* InAnimGraphNode, TFunctionRef<TSharedPtr<SGraphPin>(UEdGraphPin*)> InFindWidgetForPin);
+
 protected:
 	// SWidget interface
 	virtual void Tick(const FGeometry& AllottedGeometry, const double InCurrentTime, const float InDeltaTime) override;
@@ -40,9 +43,6 @@ private:
 	// Handle the node informing us that the title has changed
 	void HandleNodeTitleChanged();
 
-	// Tweak any created pin widgets so they respond to bindings
-	void ReconfigurePinWidgetsForPropertyBindings();
-
 	/** Keep a reference to the indicator widget handing around */
 	TSharedPtr<SWidget> IndicatorWidget;
 
diff --git a/Engine/Source/Editor/AnimationEditor/Private/AnimationEditor.cpp b/Engine/Source/Editor/AnimationEditor/Private/AnimationEditor.cpp
--- a/Engine/Source/Editor/AnimationEditor/Private/AnimationEditor.cpp
+++ b/Engine/Source/Editor/AnimationEditor/Private/AnimationEditor.cpp
@@ -75,23 +75,8 @@ DEFINE_LOG_CATEGORY(LogAnimationEditor);
 
 #define LOCTEXT_NAMESPACE "AnimationEditor"
 
-FAnimationEditor::FAnimationEditor()
-{
-	UEditorEngine* Editor = Cast<UEditorEngine>(GEngine);
-	if (Editor != nullptr)
-	{
-		Editor->RegisterForUndo(this);
-	}
-}
-
 FAnimationEditor::~FAnimationEditor()
 {
-	UEditorEngine* Editor = Cast<UEditorEngine>(GEngine);
-	if (Editor != nullptr)
-	{
-		Editor->UnregisterForUndo(this);
-	}
-
 	GEditor->GetEditorSubsystem<UImportSubsystem>()->OnAssetPostImport.RemoveAll(this);
 	FReimportManager::Instance()->OnPostReimport().RemoveAll(this);
 }
@@ -371,16 +356,6 @@ void FAnimationEditor::HandleObjectSelected(UObject* InObject)
 	}
 }
 
-void FAnimationEditor::PostUndo(bool bSuccess)
-{
-	OnPostUndo.Broadcast();
-}
-
-void FAnimationEditor::PostRedo(bool bSuccess)
-{
-	OnPostUndo.Broadcast();
-}
-
 void FAnimationEditor::HandleDetailsCreated(const TSharedRef<IDetailsView>& InDetailsView)
 {
 	DetailsView = InDetailsView;
@@ -394,7 +369,7 @@ TSharedPtr<SDockTab> FAnimationEditor::OpenNewAnimationDocumentTab(UAnimationAss
 	{
 		FString	DocumentLink;
 
-		FAnimDocumentArgs Args(PersonaToolkit->GetPreviewScene(), GetPersonaToolkit(), GetSkeletonTree()->GetEditableSkeleton(), OnPostUndo, OnSectionsChanged);
+		FAnimDocumentArgs Args(PersonaToolkit->GetPreviewScene(), GetPersonaToolkit(), GetSkeletonTree()->GetEditableSkeleton(), OnSectionsChanged);
 		Args.OnDespatchObjectsSelected = FOnObjectsSelected::CreateSP(this, &FAnimationEditor::HandleObjectsSelected);
 		Args.OnDespatchInvokeTab = FOnInvokeTab::CreateSP(this, &FAssetEditorToolkit::InvokeTab);
 		Args.OnDespatchSectionsChanged = FSimpleDelegate::CreateSP(this, &FAnimationEditor::HandleSectionsChanged);
diff --git a/Engine/Source/Editor/AnimationEditor/Private/AnimationEditor.h b/Engine/Source/Editor/AnimationEditor/Private/AnimationEditor.h
--- a/Engine/Source/Editor/AnimationEditor/Private/AnimationEditor.h
+++ b/Engine/Source/Editor/AnimationEditor/Private/AnimationEditor.h
@@ -9,7 +9,6 @@
 #include "Toolkits/IToolkitHost.h"
 #include "IAnimationEditor.h"
 #include "TickableEditorObject.h"
-#include "EditorUndoClient.h"
 #include "Containers/ArrayView.h"
 
 struct FAssetData;
@@ -50,11 +49,9 @@ namespace AnimationEditorTabs
 	extern const FName AnimMontageSectionsTab;
 }
 
-class FAnimationEditor : public IAnimationEditor, public FGCObject, public FEditorUndoClient, public FTickableEditorObject
+class FAnimationEditor : public IAnimationEditor, public FGCObject, public FTickableEditorObject
 {
 public:
-	FAnimationEditor();
-
 	virtual ~FAnimationEditor();
 
 	/** Edits the specified Skeleton object */
@@ -82,10 +79,6 @@ public:
 	virtual void Tick(float DeltaTime) override;
 	virtual TStatId GetStatId() const override;
 	virtual ETickableTickType GetTickableTickType() const override { return ETickableTickType::Always; }
-	
-	/** FEditorUndoClient interface */
-	virtual void PostUndo(bool bSuccess) override;
-	virtual void PostRedo(bool bSuccess) override;
 
 	/** @return the documentation location for this editor */
 	virtual FString GetDocumentationLink() const override
@@ -167,9 +160,6 @@ private:
 
 	bool RecordMeshToAnimation(USkeletalMeshComponent* PreviewComponent, UAnimSequence* NewAsset) const;
 public:
-	/** Multicast delegate fired on global undo/redo */
-	FSimpleMulticastDelegate OnPostUndo;
-
 	/** Multicast delegate fired on global undo/redo */
 	FSimpleMulticastDelegate OnLODChanged;
 
diff --git a/Engine/Source/Editor/AnimationEditor/Private/AnimationEditorMode.cpp b/Engine/Source/Editor/AnimationEditor/Private/AnimationEditorMode.cpp
--- a/Engine/Source/Editor/AnimationEditor/Private/AnimationEditorMode.cpp
+++ b/Engine/Source/Editor/AnimationEditor/Private/AnimationEditorMode.cpp
@@ -32,8 +32,8 @@ FAnimationEditorMode::FAnimationEditorMode(TSharedRef<FWorkflowCentricApplicatio
 	TabFactories.RegisterFactory(PersonaModule.CreateAdvancedPreviewSceneTabFactory(InHostingApp, AnimationEditor->GetPersonaToolkit()->GetPreviewScene()));
 	TabFactories.RegisterFactory(PersonaModule.CreateAnimationAssetBrowserTabFactory(InHostingApp, AnimationEditor->GetPersonaToolkit(), FOnOpenNewAsset::CreateSP(&AnimationEditor.Get(), &FAnimationEditor::HandleOpenNewAsset), FOnAnimationSequenceBrowserCreated::CreateSP(&AnimationEditor.Get(), &FAnimationEditor::HandleAnimationSequenceBrowserCreated), true));
 	TabFactories.RegisterFactory(PersonaModule.CreateAssetDetailsTabFactory(InHostingApp, FOnGetAsset::CreateSP(&AnimationEditor.Get(), &FAnimationEditor::HandleGetAsset), FOnDetailsCreated()));
-	TabFactories.RegisterFactory(PersonaModule.CreateCurveViewerTabFactory(InHostingApp, InSkeletonTree->GetEditableSkeleton(), AnimationEditor->GetPersonaToolkit()->GetPreviewScene(), AnimationEditor->OnPostUndo, OnObjectsSelected));
-	TabFactories.RegisterFactory(PersonaModule.CreateSkeletonSlotNamesTabFactory(InHostingApp, InSkeletonTree->GetEditableSkeleton(), AnimationEditor->OnPostUndo, FOnObjectSelected::CreateSP(&AnimationEditor.Get(), &FAnimationEditor::HandleObjectSelected)));
+	TabFactories.RegisterFactory(PersonaModule.CreateCurveViewerTabFactory(InHostingApp, InSkeletonTree->GetEditableSkeleton(), AnimationEditor->GetPersonaToolkit()->GetPreviewScene(), OnObjectsSelected));
+	TabFactories.RegisterFactory(PersonaModule.CreateSkeletonSlotNamesTabFactory(InHostingApp, InSkeletonTree->GetEditableSkeleton(), FOnObjectSelected::CreateSP(&AnimationEditor.Get(), &FAnimationEditor::HandleObjectSelected)));
 	TabFactories.RegisterFactory(PersonaModule.CreateAnimNotifiesTabFactory(InHostingApp, InSkeletonTree->GetEditableSkeleton(), OnObjectsSelected));
 	TabFactories.RegisterFactory(PersonaModule.CreateAnimMontageSectionsTabFactory(InHostingApp, AnimationEditor->GetPersonaToolkit(), AnimationEditor->OnSectionsChanged));
 
diff --git a/Engine/Source/Editor/EditorStyle/Private/StarshipStyle.cpp b/Engine/Source/Editor/EditorStyle/Private/StarshipStyle.cpp
--- a/Engine/Source/Editor/EditorStyle/Private/StarshipStyle.cpp
+++ b/Engine/Source/Editor/EditorStyle/Private/StarshipStyle.cpp
@@ -971,6 +971,8 @@ void FStarshipEditorStyle::FStyle::SetupGeneralStyles()
 		Set( "BlendSpace.SamplePoint", new IMAGE_BRUSH( "Old/AnimEditor/BlendSpace_Sample", Icon16x16 ) );
 		Set( "BlendSpace.SamplePoint_Highlight", new IMAGE_BRUSH( "Old/AnimEditor/BlendSpace_Sample_Highlight", Icon16x16 ) );
 		Set( "BlendSpace.SamplePoint_Invalid", new IMAGE_BRUSH( "Old/AnimEditor/BlendSpace_Sample_Invalid", Icon16x16 ) );
+		Set( "BlendSpace.Graph", new IMAGE_BRUSH_SVG("Starship/Animation/BlendSpace", Icon16x16));
+		Set( "BlendSpace.SampleGraph", new IMAGE_BRUSH_SVG("Starship/Animation/BlendSpace", Icon16x16));
 
 		Set( "AnimEditor.EditPreviewParameters", new IMAGE_BRUSH( "Icons/icon_adjust_parameters_40x", Icon40x40) );		
 		Set( "AnimEditor.EditPreviewParameters.Small", new IMAGE_BRUSH( "Icons/icon_adjust_parameters_40x", Icon20x20) );		
diff --git a/Engine/Source/Editor/GraphEditor/Private/GraphEditorActions.cpp b/Engine/Source/Editor/GraphEditor/Private/GraphEditorActions.cpp
--- a/Engine/Source/Editor/GraphEditor/Private/GraphEditorActions.cpp
+++ b/Engine/Source/Editor/GraphEditor/Private/GraphEditorActions.cpp
@@ -84,24 +84,6 @@ void FGraphEditorCommandsImpl::RegisterCommands()
 	UI_COMMAND( SelectAllInputNodes, "Select All Input Nodes", "Adds all input Nodes linked to this Pin to selection", EUserInterfaceActionType::Button, FInputChord() )
 	UI_COMMAND( SelectAllOutputNodes, "Select All Output Nodes", "Adds all output Nodes linked to this Pin to selection", EUserInterfaceActionType::Button, FInputChord() )
 
-	UI_COMMAND( SelectBone, "Select Bone", "Assign or change the bone for SkeletalControls", EUserInterfaceActionType::Button, FInputChord() )
-	UI_COMMAND( AddBlendListPin, "Add Blend Pin", "Add Blend Pin to BlendList", EUserInterfaceActionType::Button, FInputChord() )
-	UI_COMMAND( RemoveBlendListPin, "Remove Blend Pin", "Remove Blend Pin", EUserInterfaceActionType::Button, FInputChord() )
-
-	UI_COMMAND( ConvertToSeqEvaluator, "Convert To Single Frame Animation", "Convert to one frame animation that requires position", EUserInterfaceActionType::Button, FInputChord() )
-	UI_COMMAND( ConvertToSeqPlayer, "Convert to Sequence Player", "Convert back to sequence player without manual position set up", EUserInterfaceActionType::Button, FInputChord() )
-
-	UI_COMMAND( ConvertToBSEvaluator, "Convert To Single Frame BlendSpace", "Convert to one frame BlendSpace that requires position", EUserInterfaceActionType::Button, FInputChord() )
-	UI_COMMAND( ConvertToBSPlayer, "Convert to BlendSpace Player", "Convert back to BlendSpace player without manual position set up", EUserInterfaceActionType::Button, FInputChord() )
-
-	UI_COMMAND( ConvertToAimOffsetLookAt, "Convert To LookAt AimOffset", "Convert to one AimOffset that automatically tracks a Target", EUserInterfaceActionType::Button, FInputChord())
-	UI_COMMAND( ConvertToAimOffsetSimple, "Convert to Simple AimOffset", "Convert to a manual AimOffets", EUserInterfaceActionType::Button, FInputChord())
-
-	UI_COMMAND(ConvertToPoseBlender, "Convert To Pose Blender", "Convert to pose blender that can blend by source curves", EUserInterfaceActionType::Button, FInputChord())
-	UI_COMMAND(ConvertToPoseByName, "Convert to Pose By Name", "Convert to pose node that returns by name", EUserInterfaceActionType::Button, FInputChord())
-
-	UI_COMMAND( OpenRelatedAsset, "Open Asset", "Opens the asset related to this node", EUserInterfaceActionType::Button, FInputChord() )
-
 	UI_COMMAND( CreateComment, "Create Comment", "Create a comment box", EUserInterfaceActionType::Button, FInputChord(EKeys::C))
 
 	UI_COMMAND( ZoomIn, "Zoom In", "Zoom in on the graph editor", EUserInterfaceActionType::Button, FInputChord(EKeys::Add))
diff --git a/Engine/Source/Editor/GraphEditor/Public/GraphEditorActions.h b/Engine/Source/Editor/GraphEditor/Public/GraphEditorActions.h
--- a/Engine/Source/Editor/GraphEditor/Public/GraphEditorActions.h
+++ b/Engine/Source/Editor/GraphEditor/Public/GraphEditorActions.h
@@ -112,31 +112,6 @@ public:
 	TSharedPtr< FUICommandInfo > SelectAllInputNodes;
 	TSharedPtr< FUICommandInfo > SelectAllOutputNodes;
 
-	// SkeletalControl specific commands
-	TSharedPtr< FUICommandInfo > SelectBone;
-	// Blend list options
-	TSharedPtr< FUICommandInfo > AddBlendListPin;
-	TSharedPtr< FUICommandInfo > RemoveBlendListPin;
-
-	// options for sequence/evaluator converter
-	TSharedPtr< FUICommandInfo > ConvertToSeqEvaluator;
-	TSharedPtr< FUICommandInfo > ConvertToSeqPlayer;
-
-	// options for blendspace sequence/evaluator converter
-	TSharedPtr< FUICommandInfo > ConvertToBSEvaluator;
-	TSharedPtr< FUICommandInfo > ConvertToBSPlayer;
-
-	// options for aimoffset converter
-	TSharedPtr< FUICommandInfo > ConvertToAimOffsetLookAt;
-	TSharedPtr< FUICommandInfo > ConvertToAimOffsetSimple;
-
-	// options for sequence/evaluator converter
-	TSharedPtr< FUICommandInfo > ConvertToPoseBlender;
-	TSharedPtr< FUICommandInfo > ConvertToPoseByName;
-
-	// option for opening the asset related to the graph node
-	TSharedPtr< FUICommandInfo > OpenRelatedAsset;
-
 	//create a comment node
 	TSharedPtr< FUICommandInfo > CreateComment;
 
diff --git a/Engine/Source/Editor/Kismet/Private/BlueprintEditor.cpp b/Engine/Source/Editor/Kismet/Private/BlueprintEditor.cpp
--- a/Engine/Source/Editor/Kismet/Private/BlueprintEditor.cpp
+++ b/Engine/Source/Editor/Kismet/Private/BlueprintEditor.cpp
@@ -170,6 +170,9 @@
 #include "Widgets/Input/SNumericEntryBox.h"
 #include "Subsystems/AssetEditorSubsystem.h"
 
+#include "BlendSpaceGraph.h"
+#include "AnimationBlendSpaceSampleGraph.h"
+
 DEFINE_LOG_CATEGORY_STATIC(LogBlueprintEditor, Log, All);
 
 #define LOCTEXT_NAMESPACE "BlueprintEditor"
@@ -561,6 +564,14 @@ const FSlateBrush* FBlueprintEditor::GetGlyphForGraph(const UEdGraph* Graph, boo
 				{
 					ReturnValue = FEditorStyle::GetBrush( bInLargeIcon ? TEXT("GraphEditor.State_24x") : TEXT("GraphEditor.State_16x") );
 				}
+				else if ( Graph->IsA(UBlendSpaceGraph::StaticClass()) )
+				{
+					ReturnValue = FEditorStyle::GetBrush(TEXT("BlendSpace.Graph") );
+				}
+				else if ( Graph->IsA(UAnimationBlendSpaceSampleGraph::StaticClass()) )
+				{
+					ReturnValue = FEditorStyle::GetBrush(TEXT("BlendSpace.SampleGraph") );
+				}
 				else
 				{
 					// If it has overriden an interface, show it as a function
@@ -1035,6 +1046,16 @@ void FBlueprintEditor::OnComponentDoubleClicked(TSharedPtr<class FSCSEditorTreeN
 	}
 }
 
+TSharedRef<SWidget> FBlueprintEditor::CreateGraphTitleBarWidget(TSharedRef<FTabInfo> InTabInfo, UEdGraph* InGraph)
+{
+	// Create the title bar widget
+	return SNew(SGraphTitleBar)
+		.EdGraphObj(InGraph)
+		.Kismet2(SharedThis(this))
+		.OnDifferentGraphCrumbClicked(this, &FBlueprintEditor::OnChangeBreadCrumbGraph)
+		.HistoryNavigationWidget(InTabInfo->CreateHistoryNavigationWidget());
+}
+
 /** Create new tab for the supplied graph - don't call this directly.*/
 TSharedRef<SGraphEditor> FBlueprintEditor::CreateGraphEditorWidget(TSharedRef<FTabInfo> InTabInfo, UEdGraph* InGraph)
 {
@@ -1319,57 +1340,6 @@ TSharedRef<SGraphEditor> FBlueprintEditor::CreateGraphEditorWidget(TSharedRef<FT
 				FCanExecuteAction::CreateSP( this, &FBlueprintEditor::CanStopWatchingPin )
 				);
 
-			GraphEditorCommands->MapAction( FGraphEditorCommands::Get().SelectBone,
-				FExecuteAction::CreateSP( this, &FBlueprintEditor::OnSelectBone ),
-				FCanExecuteAction::CreateSP( this, &FBlueprintEditor::CanSelectBone )
-				);
-
-			GraphEditorCommands->MapAction( FGraphEditorCommands::Get().AddBlendListPin,
-				FExecuteAction::CreateSP( this, &FBlueprintEditor::OnAddPosePin ),
-				FCanExecuteAction::CreateSP( this, &FBlueprintEditor::CanAddPosePin )
-				);
-
-			GraphEditorCommands->MapAction( FGraphEditorCommands::Get().RemoveBlendListPin,
-				FExecuteAction::CreateSP( this, &FBlueprintEditor::OnRemovePosePin ),
-				FCanExecuteAction::CreateSP( this, &FBlueprintEditor::CanRemovePosePin )
-				);
-
-			GraphEditorCommands->MapAction( FGraphEditorCommands::Get().ConvertToSeqEvaluator,
-				FExecuteAction::CreateSP( this, &FBlueprintEditor::OnConvertToSequenceEvaluator )
-				);
-
-			GraphEditorCommands->MapAction( FGraphEditorCommands::Get().ConvertToSeqPlayer,
-				FExecuteAction::CreateSP( this, &FBlueprintEditor::OnConvertToSequencePlayer )
-				);
-
-			GraphEditorCommands->MapAction( FGraphEditorCommands::Get().ConvertToBSEvaluator,
-				FExecuteAction::CreateSP( this, &FBlueprintEditor::OnConvertToBlendSpaceEvaluator )
-				);
-
-			GraphEditorCommands->MapAction( FGraphEditorCommands::Get().ConvertToBSPlayer,
-				FExecuteAction::CreateSP( this, &FBlueprintEditor::OnConvertToBlendSpacePlayer )
-				);
-
-			GraphEditorCommands->MapAction(FGraphEditorCommands::Get().ConvertToAimOffsetLookAt,
-				FExecuteAction::CreateSP(this, &FBlueprintEditor::OnConvertToAimOffsetLookAt)
-			);
-
-			GraphEditorCommands->MapAction(FGraphEditorCommands::Get().ConvertToAimOffsetSimple,
-				FExecuteAction::CreateSP(this, &FBlueprintEditor::OnConvertToAimOffsetSimple)
-			);
-
-			GraphEditorCommands->MapAction(FGraphEditorCommands::Get().ConvertToPoseBlender,
-				FExecuteAction::CreateSP(this, &FBlueprintEditor::OnConvertToPoseBlender)
-				);
-
-			GraphEditorCommands->MapAction(FGraphEditorCommands::Get().ConvertToPoseByName,
-				FExecuteAction::CreateSP(this, &FBlueprintEditor::OnConvertToPoseByName)
-				);
-
-			GraphEditorCommands->MapAction( FGraphEditorCommands::Get().OpenRelatedAsset,
-				FExecuteAction::CreateSP( this, &FBlueprintEditor::OnOpenRelatedAsset )
-				);
-
 			GraphEditorCommands->MapAction( FGraphEditorCommands::Get().CreateComment,
 				FExecuteAction::CreateSP( this, &FBlueprintEditor::OnCreateComment )
 				);
@@ -1430,11 +1400,7 @@ TSharedRef<SGraphEditor> FBlueprintEditor::CreateGraphEditorWidget(TSharedRef<FT
 	}
 
 	// Create the title bar widget
-	TSharedPtr<SWidget> TitleBarWidget = SNew(SGraphTitleBar)
-		.EdGraphObj(InGraph)
-		.Kismet2(SharedThis(this))
-		.OnDifferentGraphCrumbClicked(this, &FBlueprintEditor::OnChangeBreadCrumbGraph)
-		.HistoryNavigationWidget(InTabInfo->CreateHistoryNavigationWidget());
+	TSharedPtr<SWidget> TitleBarWidget = CreateGraphTitleBarWidget(InTabInfo, InGraph);
 
 	SGraphEditor::FGraphEditorEvents InEvents;
 	SetupGraphEditorEvents(InGraph, InEvents);
@@ -1547,9 +1513,9 @@ FGraphAppearanceInfo FBlueprintEditor::GetGraphAppearance(UEdGraph* InGraph) con
 // Open the editor for a given graph
 void FBlueprintEditor::OnChangeBreadCrumbGraph(UEdGraph* InGraph)
 {
-	if (InGraph && FocusedGraphEdPtr.IsValid())
+	if (InGraph)
 	{
-		OpenDocument(InGraph, FDocumentTracker::NavigatingCurrentDocument);
+		JumpToHyperlink(InGraph, false);
 	}
 }
 
@@ -3816,7 +3782,7 @@ void FBlueprintEditor::JumpToHyperlink(const UObject* ObjectReference, bool bReq
 	{
 		// Navigating into things should re-use the current tab when it makes sense
 		FDocumentTracker::EOpenDocumentCause OpenMode = FDocumentTracker::OpenNewDocument;
-		if ((Graph->GetSchema()->GetGraphType(Graph) == GT_Ubergraph) || Cast<UK2Node_Composite>(Graph->GetOuter()))
+		if ((Graph->GetSchema()->GetGraphType(Graph) == GT_Ubergraph) || Cast<UK2Node>(Graph->GetOuter()) || Cast<UEdGraph>(Graph->GetOuter()))
 		{
 			// Ubergraphs directly reuse the current graph
 			OpenMode = FDocumentTracker::NavigatingCurrentDocument;
@@ -6430,9 +6396,7 @@ void FBlueprintEditor::DeleteSelectedNodes()
 				};
 
 
-				if (Node->IsA<UAnimStateNodeBase>() || 
-					Node->IsA<UAnimGraphNode_StateMachineBase>() ||
-					Node->IsA<UK2Node_Composite>())
+				if (Node->GetSubGraphs().Num() > 0)
 				{
 					CloseAllDocumentsTab(Node);
 				}
@@ -6453,13 +6417,9 @@ void FBlueprintEditor::DeleteSelectedNodes()
 					bNeedToModifyStructurally = true;
 				}
 
-				if (UK2Node_Composite* SelectedNode = Cast<UK2Node_Composite>(*NodeIt))
+				if (Node->GetSubGraphs().Num() > 0)
 				{
-					//Close the tab for the composite if it was open
-					if (SelectedNode->BoundGraph)
-					{
-						DocumentManager->CleanInvalidTabs();
-					}
+					DocumentManager->CleanInvalidTabs();
 				}
 				else if (UK2Node_Timeline* TimelineNode = Cast<UK2Node_Timeline>(*NodeIt))
 				{
diff --git a/Engine/Source/Editor/Kismet/Private/SBlueprintPalette.cpp b/Engine/Source/Editor/Kismet/Private/SBlueprintPalette.cpp
--- a/Engine/Source/Editor/Kismet/Private/SBlueprintPalette.cpp
+++ b/Engine/Source/Editor/Kismet/Private/SBlueprintPalette.cpp
@@ -57,6 +57,8 @@
 #include "UObject/WeakFieldPtr.h"
 #include "BlueprintNodeSpawner.h"
 #include "Dialogs/Dialogs.h"
+#include "BlendSpaceGraph.h"
+#include "AnimationBlendSpaceSampleGraph.h"
 
 #define LOCTEXT_NAMESPACE "BlueprintPalette"
 
@@ -205,6 +207,16 @@ static void GetSubGraphIcon(FEdGraphSchemaAction_K2Graph const* const ActionIn,
 					ToolTipOut = LOCTEXT("AnimationTransitionGraph_ToolTip", "Animation Transition Rule");
 				}
 			}
+			else if (Cast<UBlendSpaceGraph>(ActionIn->EdGraph))
+			{
+				IconOut = FEditorStyle::GetBrush(TEXT("BlendSpace.Graph") );
+				ToolTipOut = LOCTEXT("BlendSpace_ToolTip", "BlendSpace");
+			}
+			else if (Cast<UAnimationBlendSpaceSampleGraph>(ActionIn->EdGraph))
+			{
+				IconOut = FEditorStyle::GetBrush(TEXT("BlendSpace.SampleGraph") );
+				ToolTipOut = LOCTEXT("BlendSpaceSample_ToolTip", "BlendSpace Sample");
+			}
 			else
 			{
 				IconOut = FEditorStyle::GetBrush(TEXT("GraphEditor.SubGraph_16x") );
diff --git a/Engine/Source/Editor/Kismet/Private/SGraphTitleBarAddNewBookmark.cpp b/Engine/Source/Editor/Kismet/Private/SGraphTitleBarAddNewBookmark.cpp
--- a/Engine/Source/Editor/Kismet/Private/SGraphTitleBarAddNewBookmark.cpp
+++ b/Engine/Source/Editor/Kismet/Private/SGraphTitleBarAddNewBookmark.cpp
@@ -124,6 +124,21 @@ void SGraphTitleBarAddNewBookmark::Construct(const FArguments& InArgs)
 
 	SetMenuContentWidgetToFocus(NameEntryWidget);
 
+	SetVisibility(MakeAttributeLambda([this]()
+	{
+		TSharedPtr<FBlueprintEditor> EditorContext = EditorContextPtr.Pin();
+		if (EditorContext.IsValid())
+		{
+			UEdGraph* FocusedGraph = EditorContext->GetFocusedGraph();
+			if (FocusedGraph != nullptr || CurrentViewBookmarkId.IsValid())
+			{
+				return EVisibility::Visible;
+			}
+		}
+
+		return EVisibility::Hidden;
+	}));
+	
 	SComboButton::Construct(Args);
 }
 
diff --git a/Engine/Source/Editor/Kismet/Private/SMyBlueprint.cpp b/Engine/Source/Editor/Kismet/Private/SMyBlueprint.cpp
--- a/Engine/Source/Editor/Kismet/Private/SMyBlueprint.cpp
+++ b/Engine/Source/Editor/Kismet/Private/SMyBlueprint.cpp
@@ -1930,7 +1930,7 @@ void SMyBlueprint::ExecuteAction(TSharedPtr<FEdGraphSchemaAction> InAction)
 
 			if (GraphAction->EdGraph)
 			{
-				BlueprintEditorPtr.Pin()->OpenDocument(GraphAction->EdGraph, OpenMode);
+				BlueprintEditorPtr.Pin()->JumpToHyperlink(GraphAction->EdGraph);
 			}
 		}
 		if (InAction->GetTypeId() == FEdGraphSchemaAction_K2Delegate::StaticGetTypeId())
@@ -1939,7 +1939,7 @@ void SMyBlueprint::ExecuteAction(TSharedPtr<FEdGraphSchemaAction> InAction)
 
 			if (DelegateAction->EdGraph)
 			{
-				BlueprintEditorPtr.Pin()->OpenDocument(DelegateAction->EdGraph, OpenMode);
+				BlueprintEditorPtr.Pin()->JumpToHyperlink(DelegateAction->EdGraph);
 			}
 		}
 		else if(InAction->GetTypeId() == FEdGraphSchemaAction_K2Var::StaticGetTypeId())
@@ -1957,7 +1957,7 @@ void SMyBlueprint::ExecuteAction(TSharedPtr<FEdGraphSchemaAction> InAction)
 					// Convert the Timeline's name to a variable name before comparing it to the variable
 					if (BlueprintObj->Timelines[i]->GetVariableName() == VarAction->GetVariableName())
 					{
-						BlueprintEditorPtr.Pin()->OpenDocument(BlueprintObj->Timelines[i], OpenMode);
+						BlueprintEditorPtr.Pin()->JumpToHyperlink(BlueprintObj->Timelines[i]);
 					}
 				}
 			}
diff --git a/Engine/Source/Editor/Kismet/Public/BlueprintEditor.h b/Engine/Source/Editor/Kismet/Public/BlueprintEditor.h
--- a/Engine/Source/Editor/Kismet/Public/BlueprintEditor.h
+++ b/Engine/Source/Editor/Kismet/Public/BlueprintEditor.h
@@ -955,33 +955,6 @@ protected:
 	virtual void OnStopWatchingPin();
 	virtual bool CanStopWatchingPin() const;
 
-	/**  BEGIN PERSONA related callback functions */
-	virtual void OnSelectBone() {};
-	virtual bool CanSelectBone() const { return false; }
-
-	virtual void OnAddPosePin() {};
-	virtual bool CanAddPosePin() const { return false; }
-
-	virtual void OnMergeAnimStateTransitions() {};
-	virtual bool CanMergeAnimStateTransitions() const { return false; }
-
-	virtual void OnRemovePosePin() {};
-	virtual bool CanRemovePosePin() const { return false; }
-
-	// convert functions between evaluator and player
-	virtual void OnConvertToSequenceEvaluator() {};
-	virtual void OnConvertToSequencePlayer() {};
-	virtual void OnConvertToBlendSpaceEvaluator() {};
-	virtual void OnConvertToBlendSpacePlayer() {};
-	virtual void OnConvertToAimOffsetLookAt() {};
-	virtual void OnConvertToAimOffsetSimple() {};
-	virtual void OnConvertToPoseBlender() {};
-	virtual void OnConvertToPoseByName() {};
-
-	// Opens the associated asset of the selected nodes
-	virtual void OnOpenRelatedAsset() {};
-	/** END PERSONA related callback functions */
-
 	void ToggleSaveIntermediateBuildProducts();
 	bool GetSaveIntermediateBuildProducts() const;
 
@@ -1127,6 +1100,9 @@ protected:
 	/** Unbinds all graph editor "quick jump" commands */
 	void ClearAllGraphEditorQuickJumps();
 
+	/** Create a graph title bar widget */
+	TSharedRef<SWidget> CreateGraphTitleBarWidget(TSharedRef<FTabInfo> InTabInfo, UEdGraph* InGraph);
+
 private:
 
 	/** Returns true if modules can be recompiled */
diff --git a/Engine/Source/Editor/Persona/Private/Customization/BlendSampleDetails.cpp b/Engine/Source/Editor/Persona/Private/Customization/BlendSampleDetails.cpp
--- a/Engine/Source/Editor/Persona/Private/Customization/BlendSampleDetails.cpp
+++ b/Engine/Source/Editor/Persona/Private/Customization/BlendSampleDetails.cpp
@@ -23,12 +23,19 @@
 #include "PackageTools.h"
 #include "IDetailGroup.h"
 #include "Editor.h"
+#include "AnimGraphNode_BlendSpaceGraphBase.h"
+#include "Animation/AnimBlueprint.h"
+#include "Toolkits/ToolkitManager.h"
+#include "BlueprintEditorModule.h"
+#include "AnimationGraph.h"
+#include "BlendSpaceGraph.h"
 
 #define LOCTEXT_NAMESPACE "BlendSampleDetails"
 
-FBlendSampleDetails::FBlendSampleDetails(const UBlendSpaceBase* InBlendSpace, class SBlendSpaceGridWidget* InGridWidget)
+FBlendSampleDetails::FBlendSampleDetails(const UBlendSpaceBase* InBlendSpace, class SBlendSpaceGridWidget* InGridWidget, int32 InSampleIndex)
 	: BlendSpace(InBlendSpace)
 	, GridWidget(InGridWidget)
+	, SampleIndex(InSampleIndex)
 {
 	// Retrieve the additive animation type enum
 	const UEnum* AdditiveTypeEnum = StaticEnum<EAdditiveAnimationType>();
@@ -80,16 +87,34 @@ void FBlendSampleDetails::CustomizeDetails(class IDetailLayoutBuilder& DetailBui
 		}
 	}
 
+	const UBlendSpaceBase* BlendspaceToUse = BlendSpaceBase != nullptr ? BlendSpaceBase : BlendSpace;
+	UAnimGraphNode_BlendSpaceGraphBase* BlendSpaceNode = nullptr;
+	if(UBlendSpaceGraph* BlendSpaceGraph = Cast<UBlendSpaceGraph>(BlendspaceToUse->GetOuter()))
+	{
+		check(BlendspaceToUse == BlendSpaceGraph->BlendSpace);
+		BlendSpaceNode = CastChecked<UAnimGraphNode_BlendSpaceGraphBase>(BlendSpaceGraph->GetOuter());
+	}
+
 	FBlendSampleDetails::GenerateBlendSampleWidget([&CategoryBuilder]() -> FDetailWidgetRow& { return CategoryBuilder.AddCustomRow(FText::FromString(TEXT("SampleValue"))); }, GridWidget->OnSampleMoved, (BlendSpaceBase != nullptr) ? BlendSpaceBase : BlendSpace, GridWidget->GetSelectedSampleIndex(), false );
 
-	TSharedPtr<IPropertyHandle> AnimationProperty = DetailBuilder.GetProperty(GET_MEMBER_NAME_CHECKED(FBlendSample, Animation), (UClass*)(FBlendSample::StaticStruct()));
-	FDetailWidgetRow& AnimationRow = CategoryBuilder.AddCustomRow(FText::FromString(TEXT("Animation")));
-	FBlendSampleDetails::GenerateAnimationWidget(AnimationRow, BlendSpaceBase != nullptr ? BlendSpaceBase : BlendSpace, AnimationProperty);
+	TSharedPtr<IPropertyHandle> AnimationProperty = DetailBuilder.GetProperty(GET_MEMBER_NAME_CHECKED(FBlendSample, Animation), FBlendSample::StaticStruct());
+
+	if(BlendspaceToUse->IsAsset())
+	{
+		FDetailWidgetRow& AnimationRow = CategoryBuilder.AddCustomRow(FText::FromString(TEXT("Animation")));
+		FBlendSampleDetails::GenerateAnimationWidget(AnimationRow, BlendspaceToUse, AnimationProperty);
+	}
+	else
+	{
+		check(BlendSpaceNode);
+		FDetailWidgetRow& GraphRow = CategoryBuilder.AddCustomRow(FText::FromString(TEXT("Graph")));
+		FBlendSampleDetails::GenerateSampleGraphWidget(GraphRow, BlendSpaceNode, SampleIndex);
+	}
 
-	TSharedPtr<IPropertyHandle> RateScaleProperty = DetailBuilder.GetProperty(GET_MEMBER_NAME_CHECKED(FBlendSample, RateScale), (UClass*)(FBlendSample::StaticStruct()));
+	TSharedPtr<IPropertyHandle> RateScaleProperty = DetailBuilder.GetProperty(GET_MEMBER_NAME_CHECKED(FBlendSample, RateScale), FBlendSample::StaticStruct());
 	CategoryBuilder.AddProperty(RateScaleProperty);
 
-	TSharedPtr<IPropertyHandle> SnappedProperty = DetailBuilder.GetProperty(GET_MEMBER_NAME_CHECKED(FBlendSample, bSnapToGrid), (UClass*)(FBlendSample::StaticStruct()));
+	TSharedPtr<IPropertyHandle> SnappedProperty = DetailBuilder.GetProperty(GET_MEMBER_NAME_CHECKED(FBlendSample, bSnapToGrid), FBlendSample::StaticStruct());
 	CategoryBuilder.AddProperty(SnappedProperty);
 }
 
@@ -202,6 +227,56 @@ void FBlendSampleDetails::GenerateAnimationWidget(FDetailWidgetRow& Row, const U
 		];
 }
 
+void FBlendSampleDetails::GenerateSampleGraphWidget(FDetailWidgetRow& InRow, UAnimGraphNode_BlendSpaceGraphBase* InBlendSpaceNode, int32 InSampleIndex)
+{
+	InRow
+		.NameContent()
+		[
+			SNew(STextBlock)
+			.Font(FEditorStyle::GetFontStyle(TEXT("PropertyWindow.NormalFont")))
+			.Text(LOCTEXT("GraphLabel", "Graph"))
+		]
+		.ValueContent()
+		[
+			SNew(SBox)
+			.WidthOverride(125.0f)
+			[
+				SNew(SButton)
+				.ToolTipText(LOCTEXT("BlendSampleGraphButtonToolTip", "Edit the graph associated with this sample point"))
+				.OnClicked_Lambda([InSampleIndex, WeakBlendSpaceNode = TWeakObjectPtr<UAnimGraphNode_BlendSpaceGraphBase>(InBlendSpaceNode)]()
+				{
+					FSlateApplication::Get().DismissAllMenus();
+
+					if(UAnimGraphNode_BlendSpaceGraphBase* BlendSpaceNode = WeakBlendSpaceNode.Get())
+					{
+						UAnimBlueprint* AnimBlueprint = BlendSpaceNode->GetAnimBlueprint();
+						TSharedPtr<IToolkit> FoundAssetEditor = FToolkitManager::Get().FindEditorForAsset(AnimBlueprint);
+						if (FoundAssetEditor.IsValid() && FoundAssetEditor->IsBlueprintEditor())
+						{
+							TSharedPtr<IBlueprintEditor> BlueprintEditor = StaticCastSharedPtr<IBlueprintEditor>(FoundAssetEditor);
+							if(BlendSpaceNode->GetGraphs().IsValidIndex(InSampleIndex))
+							{
+								BlueprintEditor->JumpToHyperlink(BlendSpaceNode->GetGraphs()[InSampleIndex], false);
+							}
+						}
+					}
+
+					return FReply::Handled();
+				})
+				[
+					SNew(SBox)
+					.VAlign(VAlign_Center)
+					.HAlign(HAlign_Center)
+					[
+						SNew(STextBlock)
+						.Font(FEditorStyle::GetFontStyle(TEXT("PropertyWindow.NormalFont")))
+						.Text(LOCTEXT("BlendSampleGraphButtonLabel", "Edit Graph"))
+					]
+				]
+			]
+		];
+}
+
 bool FBlendSampleDetails::ShouldFilterAssetStatic(const FAssetData& AssetData, const UBlendSpaceBase* BlendSpaceBase)
 {
 	/** Cached flags to check whether or not an additive animation type is compatible with the blend space*/
diff --git a/Engine/Source/Editor/Persona/Private/Customization/BlendSampleDetails.h b/Engine/Source/Editor/Persona/Private/Customization/BlendSampleDetails.h
--- a/Engine/Source/Editor/Persona/Private/Customization/BlendSampleDetails.h
+++ b/Engine/Source/Editor/Persona/Private/Customization/BlendSampleDetails.h
@@ -12,15 +12,16 @@
 struct FAssetData;
 class FDetailWidgetRow;
 class UBlendSpaceBase;
+class UAnimGraphNode_BlendSpaceGraphBase;
 
 class FBlendSampleDetails : public IDetailCustomization
 {
 public:
-	FBlendSampleDetails(const class UBlendSpaceBase* InBlendSpace, class SBlendSpaceGridWidget* InGridWidget);
+	FBlendSampleDetails(const class UBlendSpaceBase* InBlendSpace, class SBlendSpaceGridWidget* InGridWidget, int32 InSampleIndex);
 
-	static TSharedRef<IDetailCustomization> MakeInstance(const class UBlendSpaceBase* InBlendSpace, class SBlendSpaceGridWidget* InGridWidget)
+	static TSharedRef<IDetailCustomization> MakeInstance(const class UBlendSpaceBase* InBlendSpace, class SBlendSpaceGridWidget* InGridWidget, int32 InSampleIndex)
 	{
-		return MakeShareable( new FBlendSampleDetails(InBlendSpace, InGridWidget) );
+		return MakeShareable( new FBlendSampleDetails(InBlendSpace, InGridWidget, InSampleIndex) );
 	}
 
 	// Begin IDetailCustomization interface
@@ -29,9 +30,11 @@ public:
 	
 	static void GenerateBlendSampleWidget(TFunction<FDetailWidgetRow& (void)>InFunctor, FOnSampleMoved OnSampleMoved, const class UBlendSpaceBase* BlendSpace, const int32 SampleIndex, bool bShowLabel);
 
-	static void GenerateAnimationWidget(FDetailWidgetRow& Row, const class UBlendSpaceBase* BlendSpace, TSharedPtr<IPropertyHandle> AnimationProperty);
+	static void GenerateAnimationWidget(FDetailWidgetRow& Row, const UBlendSpaceBase* BlendSpace, TSharedPtr<IPropertyHandle> AnimationProperty);
 
-	static bool ShouldFilterAssetStatic(const FAssetData& AssetData, const class UBlendSpaceBase* BlendSpaceBase);
+	static void GenerateSampleGraphWidget(FDetailWidgetRow& Row, UAnimGraphNode_BlendSpaceGraphBase* BlendSpaceNode, int32 SampleIndex);
+
+	static bool ShouldFilterAssetStatic(const FAssetData& AssetData, const UBlendSpaceBase* BlendSpaceBase);
 
 protected:
 	/** Checks whether or not the specified asset should not be shown in the mini content browser when changing the animation */
@@ -41,6 +44,8 @@ private:
 	const class UBlendSpaceBase* BlendSpace;
 	/** Parent grid widget object */
 	SBlendSpaceGridWidget* GridWidget;
+	/** Current sample index */
+	int32 SampleIndex;
 	/** Cached flags to check whether or not an additive animation type is compatible with the blend space*/	
 	TMap<FString, bool> bValidAdditiveTypes;
 };
diff --git a/Engine/Source/Editor/Persona/Private/Customization/BlendSpaceDetails.cpp b/Engine/Source/Editor/Persona/Private/Customization/BlendSpaceDetails.cpp
--- a/Engine/Source/Editor/Persona/Private/Customization/BlendSpaceDetails.cpp
+++ b/Engine/Source/Editor/Persona/Private/Customization/BlendSpaceDetails.cpp
@@ -16,6 +16,8 @@
 #include "Widgets/Input/SNumericEntryBox.h"
 #include "PropertyCustomizationHelpers.h"
 #include "ScopedTransaction.h"
+#include "BlendSpaceGraph.h"
+#include "AnimGraphNode_BlendSpaceGraph.h"
 
 #define LOCTEXT_NAMESPACE "BlendSpaceDetails"
 
@@ -23,6 +25,7 @@ FBlendSpaceDetails::FBlendSpaceDetails()
 {
 	Builder = nullptr;
 	BlendSpaceBase = nullptr;
+	BlendSpaceNode = nullptr;
 }
 
 FBlendSpaceDetails::~FBlendSpaceDetails()
@@ -39,6 +42,22 @@ void FBlendSpaceDetails::CustomizeDetails(class IDetailLayoutBuilder& DetailBuil
 	if (WeakPtr)
 	{
 		BlendSpaceBase = Cast<UBlendSpaceBase>(WeakPtr->Get());
+
+		if(!BlendSpaceBase->IsAsset())
+		{
+			// Hide various properties when we are 'internal'
+			DetailBuilder.HideCategory("MetaData");
+			DetailBuilder.HideCategory("AnimationNotifies");
+			DetailBuilder.HideCategory("Thumbnail");
+			DetailBuilder.HideCategory("Animation");
+			DetailBuilder.HideCategory("AdditiveSettings");
+		}
+
+		if(UBlendSpaceGraph* BlendSpaceGraph = Cast<UBlendSpaceGraph>(BlendSpaceBase->GetOuter()))
+		{
+			check(BlendSpaceBase == BlendSpaceGraph->BlendSpace);
+			BlendSpaceNode = CastChecked<UAnimGraphNode_BlendSpaceGraphBase>(BlendSpaceGraph->GetOuter());
+		}
 		const bool b1DBlendSpace = BlendSpaceBase->IsA<UBlendSpace1D>();
 
 		IDetailCategoryBuilder& CategoryBuilder = DetailBuilder.EditCategory(FName("Axis Settings"));
@@ -100,18 +119,29 @@ void FBlendSpaceDetails::CustomizeDetails(class IDetailLayoutBuilder& DetailBuil
 
 			IDetailGroup& Group = SampleCategoryBuilder.AddGroup(FName("GroupName"), FText::GetEmpty());
 			Group.HeaderRow()
+			.NameContent()
 			[
 				SNew(SHorizontalBox)
 				+ SHorizontalBox::Slot()
+				.Padding(FMargin(0,2,2,2))
 				.FillWidth(1.0f)
+				.HAlign(HAlign_Right)
 				[
 					SNew(STextBlock)
 					.Font(DetailBuilder.GetDetailFont())
-					.Text_Lambda([AnimationProperty, SampleIndex]() -> FText
+					.Text_Lambda([this, AnimationProperty, SampleIndex]() -> FText
 					{
 						FAssetData AssetData;
 						AnimationProperty->GetValue(AssetData);
-						return AssetData.IsValid() ? FText::Format(LOCTEXT("BlendSpaceAnimationNameLabel", "{0} ({1})"), FText::FromString(AssetData.GetAsset()->GetName()), FText::FromString(FString::FromInt(SampleIndex))) : FText::FromString("No Animation");
+						if(AssetData.IsValid())
+						{
+							return FText::Format(LOCTEXT("BlendSpaceAnimationNameLabel", "{0} ({1})"), FText::FromString(AssetData.GetAsset()->GetName()), FText::FromString(FString::FromInt(SampleIndex)));
+						}
+						else if(BlendSpaceNode.Get() && BlendSpaceNode->GetGraphs().IsValidIndex(SampleIndex))
+						{
+							return FText::Format(LOCTEXT("BlendSpaceAnimationNameLabel", "{0} ({1})"), FText::FromName(BlendSpaceNode->GetGraphs()[SampleIndex]->GetFName()), FText::FromString(FString::FromInt(SampleIndex)));
+						}
+						return LOCTEXT("NoAnimation", "No Animation");
 					})
 				]
 			];
@@ -131,9 +161,19 @@ void FBlendSpaceDetails::CustomizeDetails(class IDetailLayoutBuilder& DetailBuil
 					}
 				}
 			}), BlendSpaceBase, SampleIndex, false);
-			FDetailWidgetRow& AnimationRow = Group.AddWidgetRow();
-			FBlendSampleDetails::GenerateAnimationWidget(AnimationRow, BlendSpaceBase, AnimationProperty);
-			Group.AddPropertyRow(RateScaleProperty.ToSharedRef());
+			
+			if(BlendSpaceBase->IsAsset())
+			{
+				FDetailWidgetRow& AnimationRow = Group.AddWidgetRow();
+				FBlendSampleDetails::GenerateAnimationWidget(AnimationRow, BlendSpaceBase, AnimationProperty);
+				Group.AddPropertyRow(RateScaleProperty.ToSharedRef());
+			}
+			else if(BlendSpaceNode.Get())
+			{
+				FDetailWidgetRow& GraphRow = Group.AddWidgetRow();
+				FBlendSampleDetails::GenerateSampleGraphWidget(GraphRow, BlendSpaceNode.Get(), SampleIndex);
+			}
+			
 			Group.AddPropertyRow(SnappedProperty.ToSharedRef());
 		}
 	}
diff --git a/Engine/Source/Editor/Persona/Private/Customization/BlendSpaceDetails.h b/Engine/Source/Editor/Persona/Private/Customization/BlendSpaceDetails.h
--- a/Engine/Source/Editor/Persona/Private/Customization/BlendSpaceDetails.h
+++ b/Engine/Source/Editor/Persona/Private/Customization/BlendSpaceDetails.h
@@ -5,6 +5,10 @@
 #include "CoreMinimal.h"
 #include "IDetailCustomization.h"
 
+class IDetailLayoutBuilder;
+class UBlendSpaceBase;
+class UAnimGraphNode_BlendSpaceGraphBase;
+
 class FBlendSpaceDetails : public IDetailCustomization
 {
 public:
@@ -19,6 +23,7 @@ public:
 	// IDetailCustomization interface
 	virtual void CustomizeDetails(class IDetailLayoutBuilder& DetailBuilder) override;
 private:
-	class IDetailLayoutBuilder* Builder;
-	class UBlendSpaceBase* BlendSpaceBase;
+	IDetailLayoutBuilder* Builder;
+	UBlendSpaceBase* BlendSpaceBase;
+	TWeakObjectPtr<UAnimGraphNode_BlendSpaceGraphBase> BlendSpaceNode;
 };
diff --git a/Engine/Source/Editor/Persona/Private/PersonaModule.cpp b/Engine/Source/Editor/Persona/Private/PersonaModule.cpp
--- a/Engine/Source/Editor/Persona/Private/PersonaModule.cpp
+++ b/Engine/Source/Editor/Persona/Private/PersonaModule.cpp
@@ -312,7 +312,12 @@ TSharedRef<class FWorkflowTabFactory> FPersonaModule::CreateAnimNotifiesTabFacto
 
 TSharedRef<class FWorkflowTabFactory> FPersonaModule::CreateCurveViewerTabFactory(const TSharedRef<class FWorkflowCentricApplication>& InHostingApp, const TSharedRef<class IEditableSkeleton>& InEditableSkeleton, const TSharedRef<IPersonaPreviewScene>& InPreviewScene, FSimpleMulticastDelegate& InOnPostUndo, FOnObjectsSelected InOnObjectsSelected) const
 {
-	return MakeShareable(new FAnimCurveViewerTabSummoner(InHostingApp, InEditableSkeleton, InPreviewScene, InOnPostUndo, InOnObjectsSelected));
+	return MakeShareable(new FAnimCurveViewerTabSummoner(InHostingApp, InEditableSkeleton, InPreviewScene, InOnObjectsSelected));
+}
+
+TSharedRef<class FWorkflowTabFactory> FPersonaModule::CreateCurveViewerTabFactory(const TSharedRef<class FWorkflowCentricApplication>& InHostingApp, const TSharedRef<class IEditableSkeleton>& InEditableSkeleton, const TSharedRef<IPersonaPreviewScene>& InPreviewScene, FOnObjectsSelected InOnObjectsSelected) const
+{
+	return MakeShareable(new FAnimCurveViewerTabSummoner(InHostingApp, InEditableSkeleton, InPreviewScene, InOnObjectsSelected));
 }
 
 TSharedRef<class FWorkflowTabFactory> FPersonaModule::CreateRetargetManagerTabFactory(const TSharedRef<class FWorkflowCentricApplication>& InHostingApp, const TSharedRef<class IEditableSkeleton>& InEditableSkeleton, const TSharedRef<IPersonaPreviewScene>& InPreviewScene, FSimpleMulticastDelegate& InOnPostUndo) const
@@ -352,19 +357,36 @@ TSharedRef<class FWorkflowTabFactory> FPersonaModule::CreateAnimBlueprintAssetOv
 
 TSharedRef<class FWorkflowTabFactory> FPersonaModule::CreateSkeletonSlotNamesTabFactory(const TSharedRef<class FWorkflowCentricApplication>& InHostingApp, const TSharedRef<class IEditableSkeleton>& InEditableSkeleton, FSimpleMulticastDelegate& InOnPostUndo, FOnObjectSelected InOnObjectSelected) const
 {
-	return MakeShareable(new FSkeletonSlotNamesSummoner(InHostingApp, InEditableSkeleton, InOnPostUndo, InOnObjectSelected));
+	return MakeShareable(new FSkeletonSlotNamesSummoner(InHostingApp, InEditableSkeleton, InOnObjectSelected));
 }
 
-TSharedRef<SWidget> FPersonaModule::CreateBlendSpacePreviewWidget(TAttribute<const UBlendSpaceBase*> InBlendSpace, TAttribute<FVector> InPosition) const
+TSharedRef<class FWorkflowTabFactory> FPersonaModule::CreateSkeletonSlotNamesTabFactory(const TSharedRef<class FWorkflowCentricApplication>& InHostingApp, const TSharedRef<class IEditableSkeleton>& InEditableSkeleton, FOnObjectSelected InOnObjectSelected) const
+{
+	return MakeShareable(new FSkeletonSlotNamesSummoner(InHostingApp, InEditableSkeleton, InOnObjectSelected));
+}
+
+TSharedRef<SWidget> FPersonaModule::CreateBlendSpacePreviewWidget(const FBlendSpacePreviewArgs& InArgs) const
 {
 	return
 		SNew(SBlendSpaceGridWidget)
 		.Cursor(EMouseCursor::Crosshairs)
-		.BlendSpaceBase(InBlendSpace)
-		.Position(InPosition)
+		.BlendSpaceBase(InArgs.PreviewBlendSpace)
+		.Position(InArgs.PreviewPosition)
 		.ReadOnly(true)
 		.ShowAxisLabels(false)
-		.ShowSettingsButtons(false);
+		.ShowSettingsButtons(false)
+		.OnGetBlendSpaceSampleName(InArgs.OnGetBlendSpaceSampleName);
+}
+
+TSharedRef<SWidget> FPersonaModule::CreateBlendSpacePreviewWidget(TAttribute<const UBlendSpaceBase*> InBlendSpace, TAttribute<FVector> InPosition) const
+{
+PRAGMA_DISABLE_DEPRECATION_WARNINGS
+	FBlendSpacePreviewArgs Args;
+	Args.PreviewBlendSpace = InBlendSpace;
+	Args.PreviewPosition = InPosition;
+
+	return CreateBlendSpacePreviewWidget(Args);
+PRAGMA_ENABLE_DEPRECATION_WARNINGS
 }
 
 TSharedRef<class FWorkflowTabFactory> FPersonaModule::CreateAnimMontageSectionsTabFactory(const TSharedRef<class FWorkflowCentricApplication>& InHostingApp, const TSharedRef<IPersonaToolkit>& InPersonaToolkit, FSimpleMulticastDelegate& InOnSectionsChanged) const
@@ -444,7 +466,7 @@ TSharedRef<SWidget> FPersonaModule::CreateEditorWidgetForAnimDocument(const TSha
 		}
 		else if (UBlendSpace* BlendSpace = Cast<UBlendSpace>(InAnimAsset))
 		{
-			Result = SNew(SBlendSpaceEditor, InArgs.PreviewScene.Pin().ToSharedRef(), InArgs.OnPostUndo)
+			Result = SNew(SBlendSpaceEditor, InArgs.PreviewScene.Pin().ToSharedRef())
 				.BlendSpace(BlendSpace);
 
 			if (Cast<UAimOffsetBlendSpace>(InAnimAsset))
@@ -458,7 +480,7 @@ TSharedRef<SWidget> FPersonaModule::CreateEditorWidgetForAnimDocument(const TSha
 		}
 		else if (UBlendSpace1D* BlendSpace1D = Cast<UBlendSpace1D>(InAnimAsset))
 		{
-			Result = SNew(SBlendSpaceEditor1D, InArgs.PreviewScene.Pin().ToSharedRef(), InArgs.OnPostUndo)
+			Result = SNew(SBlendSpaceEditor1D, InArgs.PreviewScene.Pin().ToSharedRef())
 				.BlendSpace1D(BlendSpace1D);
 
 			if (Cast<UAimOffsetBlendSpace1D>(InAnimAsset))
@@ -1496,4 +1518,38 @@ void FPersonaModule::HandleNewAnimNotifyStateBlueprintCreated(UBlueprint* InBlue
 	}
 }
 
+TSharedRef<SWidget> FPersonaModule::CreateBlendSpaceEditWidget(UBlendSpaceBase* InBlendSpace, const FBlendSpaceEditorArgs& InArgs) const
+{
+	if (UBlendSpace* BlendSpace = Cast<UBlendSpace>(InBlendSpace))
+	{
+		return SNew(SBlendSpaceEditor)
+			.BlendSpace(BlendSpace)
+			.DisplayScrubBar(false)
+			.OnBlendSpaceSampleDoubleClicked(InArgs.OnBlendSpaceSampleDoubleClicked)
+			.OnBlendSpaceSampleAdded(InArgs.OnBlendSpaceSampleAdded)
+			.OnBlendSpaceSampleRemoved(InArgs.OnBlendSpaceSampleRemoved)
+			.OnGetBlendSpaceSampleName(InArgs.OnGetBlendSpaceSampleName)
+			.OnExtendSampleTooltip(InArgs.OnExtendSampleTooltip)
+			.OnSetPreviewPosition(InArgs.OnSetPreviewPosition)
+			.PreviewPosition(InArgs.PreviewPosition)
+			.StatusBarName(InArgs.StatusBarName);
+	}
+	else if (UBlendSpace1D* BlendSpace1D = Cast<UBlendSpace1D>(InBlendSpace))
+	{
+		return SNew(SBlendSpaceEditor1D)
+			.BlendSpace1D(BlendSpace1D)
+			.DisplayScrubBar(false)
+			.OnBlendSpaceSampleDoubleClicked(InArgs.OnBlendSpaceSampleDoubleClicked)
+			.OnBlendSpaceSampleAdded(InArgs.OnBlendSpaceSampleAdded)
+			.OnBlendSpaceSampleRemoved(InArgs.OnBlendSpaceSampleRemoved)
+			.OnGetBlendSpaceSampleName(InArgs.OnGetBlendSpaceSampleName)
+			.OnExtendSampleTooltip(InArgs.OnExtendSampleTooltip)
+			.OnSetPreviewPosition(InArgs.OnSetPreviewPosition)
+			.PreviewPosition(InArgs.PreviewPosition)
+			.StatusBarName(InArgs.StatusBarName);
+	}
+
+	return SNullWidget::NullWidget;
+}
+
 #undef LOCTEXT_NAMESPACE
diff --git a/Engine/Source/Editor/Persona/Private/SAnimCurveViewer.cpp b/Engine/Source/Editor/Persona/Private/SAnimCurveViewer.cpp
--- a/Engine/Source/Editor/Persona/Private/SAnimCurveViewer.cpp
+++ b/Engine/Source/Editor/Persona/Private/SAnimCurveViewer.cpp
@@ -532,7 +532,7 @@ FText SAnimCurveTypeList::GetAnimCurveType() const
 //////////////////////////////////////////////////////////////////////////
 // SAnimCurveViewer
 
-void SAnimCurveViewer::Construct(const FArguments& InArgs, const TSharedRef<class IEditableSkeleton>& InEditableSkeleton, const TSharedRef<IPersonaPreviewScene>& InPreviewScene, FSimpleMulticastDelegate& InOnPostUndo, FOnObjectsSelected InOnObjectsSelected)
+void SAnimCurveViewer::Construct(const FArguments& InArgs, const TSharedRef<class IEditableSkeleton>& InEditableSkeleton, const TSharedRef<IPersonaPreviewScene>& InPreviewScene, FOnObjectsSelected InOnObjectsSelected)
 {
 	OnObjectsSelected = InOnObjectsSelected;
 
@@ -547,7 +547,6 @@ void SAnimCurveViewer::Construct(const FArguments& InArgs, const TSharedRef<clas
 
 	InPreviewScene->RegisterOnPreviewMeshChanged(FOnPreviewMeshChanged::CreateSP(this, &SAnimCurveViewer::OnPreviewMeshChanged));
 	InPreviewScene->RegisterOnAnimChanged(FOnAnimChanged::CreateSP(this, &SAnimCurveViewer::OnPreviewAssetChanged));
-	InOnPostUndo.Add(FSimpleDelegate::CreateSP(this, &SAnimCurveViewer::OnPostUndo));
 
 	SmartNameChangedHandle = InEditableSkeleton->RegisterOnSmartNameChanged(FOnSmartNameChanged::FDelegate::CreateSP(this, &SAnimCurveViewer::HandleSmartNamesChange));
 
@@ -1017,7 +1016,7 @@ bool SAnimCurveViewer::GetAnimCurveOverride(FName& Name, float& Weight)
 	}
 }
 
-void SAnimCurveViewer::OnPostUndo()
+void SAnimCurveViewer::PostUndoRedo()
 {
 	RefreshCurveList();
 }
diff --git a/Engine/Source/Editor/Persona/Private/SAnimCurveViewer.h b/Engine/Source/Editor/Persona/Private/SAnimCurveViewer.h
--- a/Engine/Source/Editor/Persona/Private/SAnimCurveViewer.h
+++ b/Engine/Source/Editor/Persona/Private/SAnimCurveViewer.h
@@ -18,6 +18,7 @@
 #include "Animation/AnimInstance.h"
 #include "EditorObjectsTracker.h"
 #include "PersonaDelegates.h"
+#include "EditorUndoClient.h"
 
 class FUICommandList;
 class IEditableSkeleton;
@@ -134,7 +135,7 @@ private:
 //////////////////////////////////////////////////////////////////////////
 // SAnimCurveViewer
 
-class SAnimCurveViewer : public SCompoundWidget
+class SAnimCurveViewer : public SCompoundWidget, public FSelfRegisteringEditorUndoClient
 {
 public:
 	SLATE_BEGIN_ARGS( SAnimCurveViewer )
@@ -148,7 +149,7 @@ public:
 	* @param InArgs - Arguments passed from Slate
 	*
 	*/
-	void Construct( const FArguments& InArgs, const TSharedRef<class IEditableSkeleton>& InEditableSkeleton, const TSharedRef<class IPersonaPreviewScene>& InPreviewScene, FSimpleMulticastDelegate& InOnPostUndo, FOnObjectsSelected InOnObjectsSelected);
+	void Construct( const FArguments& InArgs, const TSharedRef<class IEditableSkeleton>& InEditableSkeleton, const TSharedRef<class IPersonaPreviewScene>& InPreviewScene, FOnObjectsSelected InOnObjectsSelected);
 
 	/**
 	* Destructor - resets the animation curve
@@ -235,10 +236,14 @@ public:
 	*/
 	FText& GetFilterText() { return FilterText; }
 
+	// FSelfRegisteringEditorUndoClient interface
+	virtual void PostUndo(bool bSuccess) override { PostUndoRedo(); }
+	virtual void PostRedo(bool bSuccess) override { PostUndoRedo(); }
+
 	/**
 	 * Refreshes the morph target list after an undo
 	 */
-	void OnPostUndo();
+	void PostUndoRedo();
 
 	virtual void Tick(const FGeometry& AllottedGeometry, const double InCurrentTime, const float InDeltaTime);
 
diff --git a/Engine/Source/Editor/Persona/Private/SAnimEditorBase.cpp b/Engine/Source/Editor/Persona/Private/SAnimEditorBase.cpp
--- a/Engine/Source/Editor/Persona/Private/SAnimEditorBase.cpp
+++ b/Engine/Source/Editor/Persona/Private/SAnimEditorBase.cpp
@@ -25,7 +25,7 @@ TSharedRef<SWidget> SAnimEditorBase::CreateDocumentAnchor()
 	return IDocumentation::Get()->CreateAnchor(TEXT("Engine/Animation/Sequences"));
 }
 
-void SAnimEditorBase::Construct(const FArguments& InArgs, const TSharedRef<class IPersonaPreviewScene>& InPreviewScene)
+void SAnimEditorBase::Construct(const FArguments& InArgs, const TSharedPtr<class IPersonaPreviewScene>& InPreviewScene)
 {
 	PreviewScenePtr = InPreviewScene;
 	OnObjectsSelected = InArgs._OnObjectsSelected;
@@ -80,14 +80,19 @@ void SAnimEditorBase::Construct(const FArguments& InArgs, const TSharedRef<class
 	}
 }
 
-TSharedRef<class SAnimationScrubPanel> SAnimEditorBase::ConstructAnimScrubPanel()
+TSharedRef<SWidget> SAnimEditorBase::ConstructAnimScrubPanel()
 {
-	return SAssignNew( AnimScrubPanel, SAnimationScrubPanel, PreviewScenePtr.Pin().ToSharedRef() )
-		.LockedSequence(Cast<UAnimSequenceBase>(GetEditorObject()))
-		.ViewInputMin(this, &SAnimEditorBase::GetViewMinInput)
-		.ViewInputMax(this, &SAnimEditorBase::GetViewMaxInput)
-		.OnSetInputViewRange(this, &SAnimEditorBase::SetInputViewRange)
-		.bAllowZoom(true);
+	if(PreviewScenePtr.IsValid())
+	{
+		return SAssignNew( AnimScrubPanel, SAnimationScrubPanel, PreviewScenePtr.Pin().ToSharedRef() )
+			.LockedSequence(Cast<UAnimSequenceBase>(GetEditorObject()))
+			.ViewInputMin(this, &SAnimEditorBase::GetViewMinInput)
+			.ViewInputMax(this, &SAnimEditorBase::GetViewMaxInput)
+			.OnSetInputViewRange(this, &SAnimEditorBase::SetInputViewRange)
+			.bAllowZoom(true);
+	}
+
+	return SNullWidget::NullWidget;
 }
 
 void SAnimEditorBase::AddReferencedObjects( FReferenceCollector& Collector )
diff --git a/Engine/Source/Editor/Persona/Private/SAnimEditorBase.h b/Engine/Source/Editor/Persona/Private/SAnimEditorBase.h
--- a/Engine/Source/Editor/Persona/Private/SAnimEditorBase.h
+++ b/Engine/Source/Editor/Persona/Private/SAnimEditorBase.h
@@ -49,7 +49,7 @@ public:
 
 	SLATE_END_ARGS()
 
-	void Construct(const FArguments& InArgs, const TSharedRef<class IPersonaPreviewScene>& InPreviewScene);
+	void Construct(const FArguments& InArgs, const TSharedPtr<class IPersonaPreviewScene>& InPreviewScene);
 
 	/** Accessors to the current viewed Min/Max input range of the editor */
 	float GetViewMinInput() const { return ViewMinInput; }
@@ -90,7 +90,7 @@ protected:
 	TSharedRef<IPersonaPreviewScene> GetPreviewScene() const { return PreviewScenePtr.Pin().ToSharedRef(); }
 
 	/** Allows derived classes to create different animation scrub panel */
-	virtual TSharedRef<class SAnimationScrubPanel> ConstructAnimScrubPanel();
+	virtual TSharedRef<SWidget> ConstructAnimScrubPanel();
 
 	/** Allows derived classes to init newly created editor objects */
 	virtual void InitDetailsViewEditorObject(class UEditorAnimBaseObj* EdObj) {};
diff --git a/Engine/Source/Editor/Persona/Private/SAnimationBlendSpace.cpp b/Engine/Source/Editor/Persona/Private/SAnimationBlendSpace.cpp
--- a/Engine/Source/Editor/Persona/Private/SAnimationBlendSpace.cpp
+++ b/Engine/Source/Editor/Persona/Private/SAnimationBlendSpace.cpp
@@ -8,12 +8,35 @@
 
 #define LOCTEXT_NAMESPACE "BlendSpaceEditor"
 
-void SBlendSpaceEditor::Construct(const FArguments& InArgs, const TSharedRef<class IPersonaPreviewScene>& InPreviewScene, FSimpleMulticastDelegate& OnPostUndo)
+void SBlendSpaceEditor::Construct(const FArguments& InArgs)
 {
 	SBlendSpaceEditorBase::Construct(SBlendSpaceEditorBase::FArguments()
-									.BlendSpace(InArgs._BlendSpace),
-									InPreviewScene,
-									OnPostUndo );
+		.BlendSpace(InArgs._BlendSpace)
+		.DisplayScrubBar(InArgs._DisplayScrubBar)
+		.OnBlendSpaceSampleDoubleClicked(InArgs._OnBlendSpaceSampleDoubleClicked)
+		.OnBlendSpaceSampleAdded(InArgs._OnBlendSpaceSampleAdded)
+		.OnBlendSpaceSampleRemoved(InArgs._OnBlendSpaceSampleRemoved)
+		.OnGetBlendSpaceSampleName(InArgs._OnGetBlendSpaceSampleName)
+		.OnExtendSampleTooltip(InArgs._OnExtendSampleTooltip)
+		.OnSetPreviewPosition(InArgs._OnSetPreviewPosition)
+		.PreviewPosition(InArgs._PreviewPosition)
+		.StatusBarName(InArgs._StatusBarName));
+}
+
+void SBlendSpaceEditor::Construct(const FArguments& InArgs, const TSharedRef<class IPersonaPreviewScene>& InPreviewScene)
+{
+	SBlendSpaceEditorBase::Construct(SBlendSpaceEditorBase::FArguments()
+		.BlendSpace(InArgs._BlendSpace)
+		.DisplayScrubBar(InArgs._DisplayScrubBar)
+		.OnBlendSpaceSampleDoubleClicked(InArgs._OnBlendSpaceSampleDoubleClicked)
+		.OnBlendSpaceSampleAdded(InArgs._OnBlendSpaceSampleAdded)
+		.OnBlendSpaceSampleRemoved(InArgs._OnBlendSpaceSampleRemoved)
+		.OnGetBlendSpaceSampleName(InArgs._OnGetBlendSpaceSampleName)
+		.OnExtendSampleTooltip(InArgs._OnExtendSampleTooltip)
+		.OnSetPreviewPosition(InArgs._OnSetPreviewPosition)
+		.PreviewPosition(InArgs._PreviewPosition)
+		.StatusBarName(InArgs._StatusBarName),
+		InPreviewScene);
 }
 
 void SBlendSpaceEditor::ResampleData()
diff --git a/Engine/Source/Editor/Persona/Private/SAnimationBlendSpace.h b/Engine/Source/Editor/Persona/Private/SAnimationBlendSpace.h
--- a/Engine/Source/Editor/Persona/Private/SAnimationBlendSpace.h
+++ b/Engine/Source/Editor/Persona/Private/SAnimationBlendSpace.h
@@ -11,12 +11,37 @@ class SBlendSpaceEditor : public SBlendSpaceEditorBase
 {
 public:
 	SLATE_BEGIN_ARGS(SBlendSpaceEditor)
-		: _BlendSpace(NULL)			
+		: _BlendSpace(nullptr)		
+		, _DisplayScrubBar(true)
+		, _StatusBarName(TEXT("AssetEditor.AnimationEditor.MainMenu"))
 		{}
-		SLATE_ARGUMENT(UBlendSpace*, BlendSpace)
+
+	SLATE_ARGUMENT(UBlendSpace*, BlendSpace)
+
+	SLATE_ARGUMENT(bool, DisplayScrubBar)
+
+	SLATE_EVENT(FOnBlendSpaceSampleDoubleClicked, OnBlendSpaceSampleDoubleClicked)
+
+	SLATE_EVENT(FOnBlendSpaceSampleAdded, OnBlendSpaceSampleAdded)
+
+	SLATE_EVENT(FOnBlendSpaceSampleRemoved, OnBlendSpaceSampleRemoved)
+
+	SLATE_EVENT(FOnGetBlendSpaceSampleName, OnGetBlendSpaceSampleName)
+
+	SLATE_EVENT(FOnExtendBlendSpaceSampleTooltip, OnExtendSampleTooltip)
+
+	SLATE_EVENT(FOnSetBlendSpacePreviewPosition, OnSetPreviewPosition)
+
+	SLATE_ATTRIBUTE(FVector, PreviewPosition)
+
+	SLATE_ARGUMENT(FName, StatusBarName)
+
 	SLATE_END_ARGS()
 
-	void Construct(const FArguments& InArgs, const TSharedRef<class IPersonaPreviewScene>& InPreviewScene, FSimpleMulticastDelegate& OnPostUndo);
+	void Construct(const FArguments& InArgs);
+
+	void Construct(const FArguments& InArgs, const TSharedRef<class IPersonaPreviewScene>& InPreviewScene);
+
 protected:
 	virtual void ResampleData() override;
 
diff --git a/Engine/Source/Editor/Persona/Private/SAnimationBlendSpace1D.cpp b/Engine/Source/Editor/Persona/Private/SAnimationBlendSpace1D.cpp
--- a/Engine/Source/Editor/Persona/Private/SAnimationBlendSpace1D.cpp
+++ b/Engine/Source/Editor/Persona/Private/SAnimationBlendSpace1D.cpp
@@ -9,12 +9,35 @@
 
 #define LOCTEXT_NAMESPACE "BlendSpace1DEditor"
 
-void SBlendSpaceEditor1D::Construct(const FArguments& InArgs, const TSharedRef<IPersonaPreviewScene>& InPreviewScene, FSimpleMulticastDelegate& OnPostUndo)
+void SBlendSpaceEditor1D::Construct(const FArguments& InArgs)
 {
 	SBlendSpaceEditorBase::Construct(SBlendSpaceEditorBase::FArguments()
-		.BlendSpace(InArgs._BlendSpace1D),
-		InPreviewScene,
-		OnPostUndo );
+		.BlendSpace(InArgs._BlendSpace1D)
+		.DisplayScrubBar(InArgs._DisplayScrubBar)
+		.OnBlendSpaceSampleDoubleClicked(InArgs._OnBlendSpaceSampleDoubleClicked)
+		.OnBlendSpaceSampleAdded(InArgs._OnBlendSpaceSampleAdded)
+		.OnBlendSpaceSampleRemoved(InArgs._OnBlendSpaceSampleRemoved)
+		.OnGetBlendSpaceSampleName(InArgs._OnGetBlendSpaceSampleName)
+		.OnExtendSampleTooltip(InArgs._OnExtendSampleTooltip)
+		.OnSetPreviewPosition(InArgs._OnSetPreviewPosition)
+		.PreviewPosition(InArgs._PreviewPosition)
+		.StatusBarName(InArgs._StatusBarName));
+}
+
+void SBlendSpaceEditor1D::Construct(const FArguments& InArgs, const TSharedRef<IPersonaPreviewScene>& InPreviewScene)
+{
+	SBlendSpaceEditorBase::Construct(SBlendSpaceEditorBase::FArguments()
+		.BlendSpace(InArgs._BlendSpace1D)
+		.DisplayScrubBar(InArgs._DisplayScrubBar)
+		.OnBlendSpaceSampleDoubleClicked(InArgs._OnBlendSpaceSampleDoubleClicked)
+		.OnBlendSpaceSampleAdded(InArgs._OnBlendSpaceSampleAdded)
+		.OnBlendSpaceSampleRemoved(InArgs._OnBlendSpaceSampleRemoved)
+		.OnGetBlendSpaceSampleName(InArgs._OnGetBlendSpaceSampleName)
+		.OnExtendSampleTooltip(InArgs._OnExtendSampleTooltip)
+		.OnSetPreviewPosition(InArgs._OnSetPreviewPosition)
+		.PreviewPosition(InArgs._PreviewPosition)
+		.StatusBarName(InArgs._StatusBarName),
+		InPreviewScene);
 }
 
 void SBlendSpaceEditor1D::ResampleData()
diff --git a/Engine/Source/Editor/Persona/Private/SAnimationBlendSpace1D.h b/Engine/Source/Editor/Persona/Private/SAnimationBlendSpace1D.h
--- a/Engine/Source/Editor/Persona/Private/SAnimationBlendSpace1D.h
+++ b/Engine/Source/Editor/Persona/Private/SAnimationBlendSpace1D.h
@@ -13,12 +13,37 @@ class SBlendSpaceEditor1D : public SBlendSpaceEditorBase
 {
 public:
 	SLATE_BEGIN_ARGS(SBlendSpaceEditor1D)
-		: _BlendSpace1D(NULL)
+		: _BlendSpace1D(nullptr)
+		, _DisplayScrubBar(true)
+		, _StatusBarName(TEXT("AssetEditor.AnimationEditor.MainMenu"))
 		{}
-		SLATE_ARGUMENT(UBlendSpace1D*, BlendSpace1D)
+
+	SLATE_ARGUMENT(UBlendSpace1D*, BlendSpace1D)
+
+	SLATE_ARGUMENT(bool, DisplayScrubBar)
+
+	SLATE_EVENT(FOnBlendSpaceSampleDoubleClicked, OnBlendSpaceSampleDoubleClicked)
+
+	SLATE_EVENT(FOnBlendSpaceSampleAdded, OnBlendSpaceSampleAdded)
+
+	SLATE_EVENT(FOnBlendSpaceSampleRemoved, OnBlendSpaceSampleRemoved)
+
+	SLATE_EVENT(FOnGetBlendSpaceSampleName, OnGetBlendSpaceSampleName)
+
+	SLATE_EVENT(FOnExtendBlendSpaceSampleTooltip, OnExtendSampleTooltip)
+
+	SLATE_EVENT(FOnSetBlendSpacePreviewPosition, OnSetPreviewPosition)
+
+	SLATE_ATTRIBUTE(FVector, PreviewPosition)
+
+	SLATE_ARGUMENT(FName, StatusBarName)
+
 	SLATE_END_ARGS()
 
-	void Construct(const FArguments& InArgs, const TSharedRef<IPersonaPreviewScene>& InPreviewScene, FSimpleMulticastDelegate& OnPostUndo);
+	void Construct(const FArguments& InArgs);
+
+	void Construct(const FArguments& InArgs, const TSharedRef<IPersonaPreviewScene>& InPreviewScene);
+
 protected:	
 	virtual void ResampleData() override;
 
diff --git a/Engine/Source/Editor/Persona/Private/SAnimationBlendSpaceBase.cpp b/Engine/Source/Editor/Persona/Private/SAnimationBlendSpaceBase.cpp
--- a/Engine/Source/Editor/Persona/Private/SAnimationBlendSpaceBase.cpp
+++ b/Engine/Source/Editor/Persona/Private/SAnimationBlendSpaceBase.cpp
@@ -27,18 +27,20 @@ SBlendSpaceEditorBase::~SBlendSpaceEditorBase()
 	FCoreUObjectDelegates::OnObjectPropertyChanged.Remove(OnPropertyChangedHandleDelegateHandle);
 }
 
-void SBlendSpaceEditorBase::Construct(const FArguments& InArgs, const TSharedRef<class IPersonaPreviewScene>& InPreviewScene, FSimpleMulticastDelegate& OnPostUndo)
+void SBlendSpaceEditorBase::Construct(const FArguments& InArgs)
 {
 	BlendSpace = InArgs._BlendSpace;
-	PreviewScenePtr = InPreviewScene;
-	OnPostUndo.Add(FSimpleDelegate::CreateSP( this, &SBlendSpaceEditorBase::PostUndo ) );
+
+	OnBlendSpaceSampleAdded = InArgs._OnBlendSpaceSampleAdded;
+	OnBlendSpaceSampleRemoved = InArgs._OnBlendSpaceSampleRemoved;
+	OnSetPreviewPosition = InArgs._OnSetPreviewPosition;
 
 	bShouldSetPreviewValue = false;
 
 	SAnimEditorBase::Construct(SAnimEditorBase::FArguments()
 		.DisplayAnimTimeline(false)
-		.DisplayAnimScrubBar(true),
-		InPreviewScene);
+		.DisplayAnimScrubBar(InArgs._DisplayScrubBar),
+		PreviewScenePtr.Pin());
 
 	NonScrollEditorPanels->AddSlot()
 	[
@@ -68,10 +70,15 @@ void SBlendSpaceEditorBase::Construct(const FArguments& InArgs, const TSharedRef
 							.Cursor(EMouseCursor::Crosshairs)
 							.BlendSpaceBase(BlendSpace)
 							.NotifyHook(this)
+							.Position(InArgs._PreviewPosition)
 							.OnSampleMoved(this, &SBlendSpaceEditorBase::OnSampleMoved)
 							.OnSampleRemoved(this, &SBlendSpaceEditorBase::OnSampleRemoved)
 							.OnSampleAdded(this, &SBlendSpaceEditorBase::OnSampleAdded)
 							.OnSampleAnimationChanged(this, &SBlendSpaceEditorBase::OnUpdateAnimation)
+							.OnSampleDoubleClicked(InArgs._OnBlendSpaceSampleDoubleClicked)
+							.OnExtendSampleTooltip(InArgs._OnExtendSampleTooltip)
+							.OnGetBlendSpaceSampleName(InArgs._OnGetBlendSpaceSampleName)
+							.StatusBarName(InArgs._StatusBarName)
 						]
 					]
 				]
@@ -83,6 +90,13 @@ void SBlendSpaceEditorBase::Construct(const FArguments& InArgs, const TSharedRef
 	OnPropertyChangedHandleDelegateHandle = FCoreUObjectDelegates::OnObjectPropertyChanged.Add(OnPropertyChangedHandle);
 }
 
+void SBlendSpaceEditorBase::Construct(const FArguments& InArgs, const TSharedRef<class IPersonaPreviewScene>& InPreviewScene)
+{
+	PreviewScenePtr = InPreviewScene;
+
+	Construct(InArgs);
+}
+
 void SBlendSpaceEditorBase::OnSampleMoved(const int32 SampleIndex, const FVector& NewValue, bool bIsInteractive, bool bSnap)
 {
 	bool bMoveSuccesful = true;
@@ -98,7 +112,7 @@ void SBlendSpaceEditorBase::OnSampleMoved(const int32 SampleIndex, const FVector
 			ResampleData();
 		}
 	}
-  }
+}
 
 void SBlendSpaceEditorBase::OnSampleRemoved(const int32 SampleIndex)
 {
@@ -110,6 +124,8 @@ void SBlendSpaceEditorBase::OnSampleRemoved(const int32 SampleIndex)
 	{
 		ResampleData();
 		BlendSpace->ValidateSampleData();
+
+		OnBlendSpaceSampleRemoved.ExecuteIfBound(SampleIndex);
 	}
 	BlendSpace->PostEditChange();
 }
@@ -119,11 +135,23 @@ void SBlendSpaceEditorBase::OnSampleAdded(UAnimSequence* Animation, const FVecto
 	FScopedTransaction ScopedTransaction(LOCTEXT("AddSample", "Adding Blend Grid Sample"));
 	BlendSpace->Modify();
 
-	const bool bAddSuccesful = BlendSpace->AddSample(Animation, Value);	
+	bool bAddSuccesful = false;
+
+	if(BlendSpace->IsAsset())
+	{
+		bAddSuccesful = BlendSpace->AddSample(Animation, Value);
+	}
+	else
+	{
+		bAddSuccesful = BlendSpace->AddSample(Value);
+	}
+
 	if (bAddSuccesful)
 	{
 		ResampleData();
 		BlendSpace->ValidateSampleData();
+
+		OnBlendSpaceSampleAdded.ExecuteIfBound(Animation, Value);
 	}
 	BlendSpace->PostEditChange();
 }
@@ -141,7 +169,7 @@ void SBlendSpaceEditorBase::OnUpdateAnimation(UAnimSequence* Animation, const FV
 	}
 }
 
-void SBlendSpaceEditorBase::PostUndo()
+void SBlendSpaceEditorBase::PostUndoRedo()
 {
 	// Validate and resample blend space data
 	BlendSpace->ValidateSampleData();
@@ -159,22 +187,30 @@ void SBlendSpaceEditorBase::PostUndo()
 
 void SBlendSpaceEditorBase::UpdatePreviewParameter() const
 {
-	class UDebugSkelMeshComponent* Component = GetPreviewScene()->GetPreviewMeshComponent();
-
-	if (Component != nullptr && Component->IsPreviewOn())
+	if(GetPreviewScene().IsValid())
 	{
-		if (Component->PreviewInstance->GetCurrentAsset() == BlendSpace)
+		class UDebugSkelMeshComponent* Component = GetPreviewScene()->GetPreviewMeshComponent();
+
+		if (Component != nullptr && Component->IsPreviewOn())
 		{
-			const FVector BlendInput = NewBlendSpaceGridWidget->GetBlendPreviewValue();
-			Component->PreviewInstance->SetBlendSpaceInput(BlendInput);
-			GetPreviewScene()->InvalidateViews();			
+			if (Component->PreviewInstance->GetCurrentAsset() == BlendSpace)
+			{
+				const FVector BlendInput = NewBlendSpaceGridWidget->GetBlendPreviewValue();
+				Component->PreviewInstance->SetBlendSpaceInput(BlendInput);
+				GetPreviewScene()->InvalidateViews();			
+			}
 		}
 	}
+	else if(OnSetPreviewPosition.IsBound())
+	{
+		const FVector BlendInput = NewBlendSpaceGridWidget->GetBlendPreviewValue();
+		OnSetPreviewPosition.Execute(BlendInput);
+	}
 }
 
-TSharedRef<class IPersonaPreviewScene> SBlendSpaceEditorBase::GetPreviewScene() const
+TSharedPtr<class IPersonaPreviewScene> SBlendSpaceEditorBase::GetPreviewScene() const
 {
-	return PreviewScenePtr.Pin().ToSharedRef();
+	return PreviewScenePtr.Pin();
 }
 
 void SBlendSpaceEditorBase::Tick(const FGeometry& AllottedGeometry, const double InCurrentTime, const float InDeltaTime)
diff --git a/Engine/Source/Editor/Persona/Private/SAnimationBlendSpaceBase.h b/Engine/Source/Editor/Persona/Private/SAnimationBlendSpaceBase.h
--- a/Engine/Source/Editor/Persona/Private/SAnimationBlendSpaceBase.h
+++ b/Engine/Source/Editor/Persona/Private/SAnimationBlendSpaceBase.h
@@ -12,21 +12,45 @@
 #include "Animation/BlendSpaceBase.h"
 #include "Widgets/Input/SCheckBox.h"
 #include "Misc/NotifyHook.h"
+#include "EditorUndoClient.h"
 
 class SBlendSpaceGridWidget;
 
-class SBlendSpaceEditorBase : public SAnimEditorBase, public FNotifyHook
+class SBlendSpaceEditorBase : public SAnimEditorBase, public FNotifyHook, public FSelfRegisteringEditorUndoClient
 {
 public:
 	SLATE_BEGIN_ARGS(SBlendSpaceEditorBase)
 		: _BlendSpace(NULL)
+		, _StatusBarName(TEXT("AssetEditor.AnimationEditor.MainMenu"))
 		{}
-			SLATE_ARGUMENT(UBlendSpaceBase*, BlendSpace)
-		SLATE_END_ARGS()
+		
+	SLATE_ARGUMENT(UBlendSpaceBase*, BlendSpace)
+
+	SLATE_ARGUMENT(bool, DisplayScrubBar)
+
+	SLATE_EVENT(FOnBlendSpaceSampleDoubleClicked, OnBlendSpaceSampleDoubleClicked)
+
+	SLATE_EVENT(FOnBlendSpaceSampleAdded, OnBlendSpaceSampleAdded)
+
+	SLATE_EVENT(FOnBlendSpaceSampleRemoved, OnBlendSpaceSampleRemoved)
+
+	SLATE_EVENT(FOnGetBlendSpaceSampleName, OnGetBlendSpaceSampleName) 
+
+	SLATE_EVENT(FOnExtendBlendSpaceSampleTooltip, OnExtendSampleTooltip)
+
+	SLATE_EVENT(FOnSetBlendSpacePreviewPosition, OnSetPreviewPosition)
+
+	SLATE_ATTRIBUTE(FVector, PreviewPosition)
+
+	SLATE_ARGUMENT(FName, StatusBarName)
+
+	SLATE_END_ARGS()
 
 	~SBlendSpaceEditorBase();
 
-	void Construct(const FArguments& InArgs, const TSharedRef<class IPersonaPreviewScene>& InPreviewScene, FSimpleMulticastDelegate& OnPostUndo);	
+	void Construct(const FArguments& InArgs);
+
+	void Construct(const FArguments& InArgs, const TSharedRef<class IPersonaPreviewScene>& InPreviewScene);
 	
 	// Begin SWidget overrides
 	virtual void Tick(const FGeometry& AllottedGeometry, const double InCurrentTime, const float InDeltaTime) override;
@@ -49,14 +73,18 @@ protected:
 
 	virtual void ResampleData() = 0;
 
+	// FSelfRegisteringEditorUndoClient interface
+	virtual void PostUndo(bool bSuccess) override { PostUndoRedo(); }
+	virtual void PostRedo(bool bSuccess) override { PostUndoRedo(); }
+
 	/** Delegate which is called when the Editor has performed an undo operation*/
-	void PostUndo();
+	void PostUndoRedo();
 
 	/** Updates Persona's preview window */
 	void UpdatePreviewParameter() const;
 	
 	/** Retrieves the preview scene shown by Persona */
-	TSharedRef<class IPersonaPreviewScene> GetPreviewScene() const;
+	TSharedPtr<class IPersonaPreviewScene> GetPreviewScene() const;
 
 	/** Global callback to anticipate on changes to the blend space */
 	void OnPropertyChanged(UObject* ObjectBeingModified, FPropertyChangedEvent& PropertyChangedEvent);
@@ -76,6 +104,15 @@ protected:
 	/** Handle to the registered OnPropertyChangedHandle delegate */
 	FDelegateHandle OnPropertyChangedHandleDelegateHandle;
 
+	/** Delegate called when a sample is added */
+	FOnBlendSpaceSampleAdded OnBlendSpaceSampleAdded;
+
+	/** Delegate called when a sample is removed */
+	FOnBlendSpaceSampleRemoved OnBlendSpaceSampleRemoved;
+
+	/** Delegate called to externally control the preview position */
+	FOnSetBlendSpacePreviewPosition OnSetPreviewPosition;
+
 	/** Flag to check whether or not the preview value should be (re-)set on the next tick */
 	bool bShouldSetPreviewValue;
 };
