Fix navigation up into BlendSpace graphs

diff --git a/Engine/Source/Editor/Kismet/Private/BlueprintEditor.cpp b/Engine/Source/Editor/Kismet/Private/BlueprintEditor.cpp
--- a/Engine/Source/Editor/Kismet/Private/BlueprintEditor.cpp
+++ b/Engine/Source/Editor/Kismet/Private/BlueprintEditor.cpp
@@ -3155,7 +3155,7 @@ void FBlueprintEditor::NavigateToParentGraph()
 	{
 		if (UEdGraph* ParentGraph = UEdGraph::GetOuterGraph(FocusedGraphEdPtr.Pin()->GetCurrentGraph()))
 		{
-			OpenDocument(ParentGraph, FDocumentTracker::NavigatingCurrentDocument);
+			JumpToHyperlink(ParentGraph, false);
 		}
 	}
 }
