Fixed SCS editor padding issues.

diff --git a/Engine/Source/Editor/EditorStyle/Private/StarshipStyle.cpp b/Engine/Source/Editor/EditorStyle/Private/StarshipStyle.cpp
--- a/Engine/Source/Editor/EditorStyle/Private/StarshipStyle.cpp
+++ b/Engine/Source/Editor/EditorStyle/Private/StarshipStyle.cpp
@@ -4240,6 +4240,7 @@ void FStarshipEditorStyle::FStyle::SetupGraphEditorStyles()
 			.SetShadowColorAndOpacity(FLinearColor(0, 0, 0, 0.9f)));
 
 		Set("SCSEditor.TreePanel", new FSlateNoResource());
+		Set("SCSEditor.Background", new FSlateRoundedBoxBrush(FStyleColors::Recessed, 4.0f));
 
 		//
 
diff --git a/Engine/Source/Editor/Kismet/Private/SSCSEditor.cpp b/Engine/Source/Editor/Kismet/Private/SSCSEditor.cpp
--- a/Engine/Source/Editor/Kismet/Private/SSCSEditor.cpp
+++ b/Engine/Source/Editor/Kismet/Private/SSCSEditor.cpp
@@ -1773,7 +1773,7 @@ void SSCS_RowWidget::Construct( const FArguments& InArgs, TSharedPtr<SSCSEditor>
 		.Style(bIsSeparator ?
 			&FEditorStyle::Get().GetWidgetStyle<FTableRowStyle>("TableView.NoHoverTableRow") :
 			&FEditorStyle::Get().GetWidgetStyle<FTableRowStyle>("SceneOutliner.TableViewRow")) //@todo create editor style for the SCS tree
-		.Padding(FMargin(0.f, 0.f, 0.f, 4.f))
+		.Padding(FMargin(0.f, 4.f, 0.f, 4.f))
 		.ShowSelection(!bIsSeparator)
 		.OnDragDetected(this, &SSCS_RowWidget::HandleOnDragDetected)
 		.OnDragEnter(this, &SSCS_RowWidget::HandleOnDragEnter)
@@ -1820,7 +1820,7 @@ TSharedRef<SWidget> SSCS_RowWidget::GenerateWidgetForColumn( const FName& Column
 		
 		TSharedRef<SToolTip> Tooltip = CreateToolTipWidget();
 
-		return	SNew(SHorizontalBox)
+		return SNew(SHorizontalBox)
 				.ToolTip(Tooltip)
 				+SHorizontalBox::Slot()
 					.AutoWidth()
@@ -1838,7 +1838,6 @@ TSharedRef<SWidget> SSCS_RowWidget::GenerateWidgetForColumn( const FName& Column
 					]
 				+SHorizontalBox::Slot()
 					.VAlign(VAlign_Center)
-					.Padding(2, 0, 0, 0)
 					[
 						InlineWidget.ToSharedRef()
 					];
@@ -1849,7 +1848,6 @@ TSharedRef<SWidget> SSCS_RowWidget::GenerateWidgetForColumn( const FName& Column
 			SNew(SHorizontalBox)
 			+SHorizontalBox::Slot()
 			.VAlign(VAlign_Center)
-			.Padding(2, 0, 0, 0)
 			[
 				SNew(STextBlock)
 				.Visibility(this, &SSCS_RowWidget::GetAssetVisibility)
@@ -1868,6 +1866,7 @@ TSharedRef<SWidget> SSCS_RowWidget::GenerateWidgetForColumn( const FName& Column
 				.ToolTip(MobilityTooltip)
 				.Visibility(EVisibility::Visible) // so we still get tooltip text for an empty SHorizontalBox
 				+ SHorizontalBox::Slot()
+				.VAlign(VAlign_Center)
 				.FillWidth(1.0f)
 				[
 					SNew(SImage)
@@ -3576,21 +3575,24 @@ TSharedRef<SWidget> SSCS_RowWidget_ActorRoot::GenerateWidgetForColumn(const FNam
 
 	NodePtr->SetRenameRequestedDelegate(FSCSEditorTreeNode::FOnRenameRequested::CreateSP(InlineEditableWidget.Get(), &SInlineEditableTextBlock::EnterEditingMode));
 
+	const bool IsRootActorNode = NodePtr->GetNodeType() == FSCSEditorTreeNode::ENodeType::RootActorNode;
+
 	return SNew(SHorizontalBox)
 		.ToolTip(CreateToolTipWidget())
 
 	+ SHorizontalBox::Slot()
+		.Padding(FMargin(IsRootActorNode ? 0.f : 4.f, 0.f, 0.f, 0.f))
 		.AutoWidth()
 		.VAlign(VAlign_Center)
 		[
 			SNew(SExpanderArrow, SharedThis(this))
-			.Visibility(NodePtr->GetNodeType() == FSCSEditorTreeNode::ENodeType::RootActorNode ? EVisibility::Collapsed : EVisibility::Visible)
+			.Visibility(IsRootActorNode ? EVisibility::Collapsed : EVisibility::Visible)
 		]
 
 	+ SHorizontalBox::Slot()
+		.Padding(FMargin(IsRootActorNode ? 4.f : 0.f, 0.f, 0.f, 0.f))
 		.AutoWidth()
 		.VAlign(VAlign_Center)
-		.Padding(FMargin(0.f, 0.f, 6.f, 0.f))
 		[
 			SNew(SImage)
 			.Image(GetIconBrush())
@@ -3600,15 +3602,15 @@ TSharedRef<SWidget> SSCS_RowWidget_ActorRoot::GenerateWidgetForColumn(const FNam
 		.AutoWidth()
 		.HAlign(HAlign_Left)
 		.VAlign(VAlign_Center)
-		.Padding(0.0f, 0.0f)
+		.Padding(6.f, 0.f, 0.f, 0.f)
 		[
 			InlineEditableWidget.ToSharedRef()
 		]
 
-	+SHorizontalBox::Slot()
+	+ SHorizontalBox::Slot()
 		.HAlign(HAlign_Left)
 		.VAlign(VAlign_Center)
-		.Padding(0.0f, 0.0f)
+		.Padding(4.f, 0.f, 0.f, 0.f)
 		[
 			SNew(STextBlock)
 			.Text(this, &SSCS_RowWidget_ActorRoot::GetActorContextText)
@@ -4041,29 +4043,22 @@ void SSCSEditor::Construct( const FArguments& InArgs )
 
 	Contents = SNew(SVerticalBox)
 	+ SVerticalBox::Slot()
-	.Padding(0.0f)
+	.AutoHeight()
+	.VAlign(VAlign_Top)
+	.Padding(4.f, 0, 4.f, 4.f)
 	[
-		SNew(SVerticalBox)
-
-		+ SVerticalBox::Slot()
-		.AutoHeight()
-		.VAlign(VAlign_Top)
-		.Padding(4.0)
-		[
-			SAssignNew(HeaderBox, SVerticalBox)
-		]
+		SAssignNew(HeaderBox, SVerticalBox)
+	]
 
-		+ SVerticalBox::Slot()
-		.Padding(4.0)
+	+ SVerticalBox::Slot()
+	[
+		SNew(SBorder)
+		.BorderImage(FAppStyle::Get().GetBrush("SCSEditor.Background"))
+		.Padding(4.f)
+		.AddMetaData<FTagMetaData>(FTagMetaData(TEXT("ComponentsPanel")))
+		.Visibility(this, &SSCSEditor::GetComponentsTreeVisibility)
 		[
-			SNew(SBorder)
-			.BorderImage(FAppStyle::Get().GetBrush("Brushes.Recessed"))
-			.Padding(0.0)
-			.AddMetaData<FTagMetaData>(FTagMetaData(TEXT("ComponentsPanel")))
-			.Visibility(this, &SSCSEditor::GetComponentsTreeVisibility)
-			[
-				SCSTreeWidget.ToSharedRef()
-			]
+			SCSTreeWidget.ToSharedRef()
 		]
 	];
 
diff --git a/Engine/Source/Editor/LevelEditor/Private/SActorDetails.cpp b/Engine/Source/Editor/LevelEditor/Private/SActorDetails.cpp
--- a/Engine/Source/Editor/LevelEditor/Private/SActorDetails.cpp
+++ b/Engine/Source/Editor/LevelEditor/Private/SActorDetails.cpp
@@ -131,26 +131,25 @@ void SActorDetails::Construct(const FArguments& InArgs, UTypedElementSelectionSe
 	GEditor->RegisterForUndo(this);
 
 	ComponentsBox = SNew(SBox)
-		.Visibility(this, &SActorDetails::GetComponentsBoxVisibility);
-
-	SCSEditor = SNew(SSCSEditor)
-		.EditorMode(EComponentEditorMode::ActorInstance)
-		.AllowEditing(this, &SActorDetails::GetAllowComponentTreeEditing)
-		.ActorContext(this, &SActorDetails::GetActorContext)
-		.OnSelectionUpdated(this, &SActorDetails::OnSCSEditorTreeViewSelectionChanged)
-		.OnItemDoubleClicked(this, &SActorDetails::OnSCSEditorTreeViewItemDoubleClicked);
-		
-	ComponentsBox->SetContent(SCSEditor.ToSharedRef());
+		.Padding(FMargin(2,0,2,0))
+		.Visibility(this, &SActorDetails::GetComponentsBoxVisibility)
+		[
+			SAssignNew(SCSEditor, SSCSEditor)
+			.EditorMode(EComponentEditorMode::ActorInstance)
+			.AllowEditing(this, &SActorDetails::GetAllowComponentTreeEditing)
+			.ActorContext(this, &SActorDetails::GetActorContext)
+			.OnSelectionUpdated(this, &SActorDetails::OnSCSEditorTreeViewSelectionChanged)
+			.OnItemDoubleClicked(this, &SActorDetails::OnSCSEditorTreeViewItemDoubleClicked)
+		];
 
 	TSharedRef<SWidget> ButtonBox = SCSEditor->GetToolButtonsBox().ToSharedRef();
 	DetailsView->SetNameAreaCustomContent( ButtonBox );
-	
 
 	ChildSlot
 	[
 		SNew(SVerticalBox)
 		+SVerticalBox::Slot()
-		.Padding(0.0f, 4.0f, 0.0f, 4.0f)
+		.Padding(10.f, 4.f, 0.f, 0.f)
 		.AutoHeight()
 		[
 			DetailsView->GetNameAreaWidget().ToSharedRef()
@@ -167,7 +166,7 @@ void SActorDetails::Construct(const FArguments& InArgs, UTypedElementSelectionSe
 				SNew( SVerticalBox )
 				+SVerticalBox::Slot()
 				.AutoHeight()
-				.Padding( FMargin( 0,0,0,1) )
+				.Padding( FMargin(0,0,0,1) )
 				[
 					SNew(SActorDetailsUneditableComponentWarning)
 					.Visibility(this, &SActorDetails::GetUCSComponentWarningVisibility)
@@ -176,7 +175,7 @@ void SActorDetails::Construct(const FArguments& InArgs, UTypedElementSelectionSe
 				]
 				+SVerticalBox::Slot()
 				.AutoHeight()
-				.Padding( FMargin( 0,0,0,1) )
+				.Padding( FMargin(0,0,0,1) )
 				[
 					SNew(SActorDetailsUneditableComponentWarning)
 					.Visibility(this, &SActorDetails::GetInheritedBlueprintComponentWarningVisibility)
@@ -185,7 +184,7 @@ void SActorDetails::Construct(const FArguments& InArgs, UTypedElementSelectionSe
 				]
 				+SVerticalBox::Slot()
 				.AutoHeight()
-				.Padding( FMargin( 0,0,0,1) )
+				.Padding( FMargin(0,0,0,1) )
 				[
 					SNew(SActorDetailsUneditableComponentWarning)
 					.Visibility(this, &SActorDetails::GetNativeComponentWarningVisibility)
diff --git a/Engine/Source/Editor/PropertyEditor/Private/SDetailNameArea.cpp b/Engine/Source/Editor/PropertyEditor/Private/SDetailNameArea.cpp
--- a/Engine/Source/Editor/PropertyEditor/Private/SDetailNameArea.cpp
+++ b/Engine/Source/Editor/PropertyEditor/Private/SDetailNameArea.cpp
@@ -119,7 +119,7 @@ TSharedRef< SWidget > SDetailNameArea::BuildObjectNameArea( const TArray< TWeakO
 		.AutoWidth()
 		.HAlign(HAlign_Left)
 		.VAlign(VAlign_Center)
-		.Padding(0,0,6,0)
+		.Padding(0)
 		[
 			SNew(SImage)
 			.Image(ActorIcon)
@@ -139,7 +139,7 @@ TSharedRef< SWidget > SDetailNameArea::BuildObjectNameArea( const TArray< TWeakO
 
 			ObjectNameArea->AddSlot()
 				.FillWidth(1.0)
-				.Padding(5, 0, 3, 0)
+				.Padding(8, 0, 3, 0)
 				[
 					SNew(SBox)
 					.WidthOverride(200.0f)
