Add some groundwork for Blueprint editor namespace "import" UX and scoping.

Experimental - disabled by default; features are currently exposed via CVar only.

#[fyi] Michael.Noland

diff --git a/Engine/Source/Editor/BlueprintGraph/Private/BlueprintActionFilter.cpp b/Engine/Source/Editor/BlueprintGraph/Private/BlueprintActionFilter.cpp
--- a/Engine/Source/Editor/BlueprintGraph/Private/BlueprintActionFilter.cpp
+++ b/Engine/Source/Editor/BlueprintGraph/Private/BlueprintActionFilter.cpp
@@ -11,6 +11,7 @@
 #include "EdGraph/EdGraphSchema.h"
 #include "EdGraph/EdGraph.h"
 #include "EdGraphSchema_K2.h"
+#include "BlueprintNamespaceHelper.h"
 #include "K2Node.h"
 #include "K2Node_Event.h"
 #include "K2Node_CallFunction.h"
@@ -286,6 +287,13 @@ namespace BlueprintActionFilterImpl
 	 */
 	static bool IsOutOfScopeLocalVariable(FBlueprintActionFilter const& Filter, FBlueprintActionInfo& BlueprintAction);
 
+	/**
+	 * Rejection test that checks to see if the action references a member variable or function that is
+	 * outside the scope of the current set of namespaces imported by the current Blueprint context(s).
+	 *
+	 */
+	static bool IsOutOfScopeMemberField(FBlueprintActionFilter const& Filter, FBlueprintActionInfo& BlueprintAction);
+
 	/**
 	 * 
 	 * 
@@ -1121,6 +1129,36 @@ static bool BlueprintActionFilterImpl::IsOutOfScopeLocalVariable(FBlueprintActio
 	return bIsFilteredOut;
 }
 
+//------------------------------------------------------------------------------
+static bool BlueprintActionFilterImpl::IsOutOfScopeMemberField(FBlueprintActionFilter const& Filter, FBlueprintActionInfo& BlueprintAction)
+{
+	bool bIsFilteredOut = false;
+	FFieldVariant MemberField = BlueprintAction.GetAssociatedMemberField();
+	if (MemberField.IsValid())
+	{
+		if (const UClass* OwnerClass = MemberField.GetOwnerClass())
+		{
+			// If the owning class was generated from a Blueprint, check to see if its namespace falls within the scope of the filter context.
+			if (const UBlueprint* OwnerBlueprint = UBlueprint::GetBlueprintFromClass(OwnerClass))
+			{
+				bIsFilteredOut = true;
+				for (const UBlueprint* Blueprint : Filter.Context.Blueprints)
+				{
+					// @todo - Cache this as a transient member on the BP itself and/or expose an API to query w/o incurring the cost of recreating this locally
+					FBlueprintNamespaceHelper NamespaceHelper(Blueprint);
+					if (NamespaceHelper.IsIncludedInNamespaceList(OwnerBlueprint->BlueprintNamespace))
+					{
+						// This action's member belongs to a class that falls within the scope of our current import set, so don't filter it out.
+						bIsFilteredOut = false;
+						break;
+					}
+				}
+			}
+		}
+	}
+	return bIsFilteredOut;
+}
+
 //------------------------------------------------------------------------------
 static bool BlueprintActionFilterImpl::IsLevelScriptActionValid(FBlueprintActionFilter const& Filter, FBlueprintActionInfo& BlueprintAction)
 {
@@ -1982,6 +2020,10 @@ FBlueprintActionFilter::FBlueprintActionFilter(uint32 Flags/*= 0x00*/)
 	AddRejectionTest(FRejectionTestDelegate::CreateStatic(IsNonTargetMember, !(Flags & BPFILTER_RejectGlobalFields)));
 	AddRejectionTest(FRejectionTestDelegate::CreateStatic(IsUnBoundBindingSpawner));
 	AddRejectionTest(FRejectionTestDelegate::CreateStatic(IsOutOfScopeLocalVariable));
+	if (Flags & BPFILTER_RejectOutOfScopeMembers)
+	{
+		AddRejectionTest(FRejectionTestDelegate::CreateStatic(IsOutOfScopeMemberField));
+	}
 	AddRejectionTest(FRejectionTestDelegate::CreateStatic(IsLevelScriptActionValid));
 
 	AddRejectionTest(FRejectionTestDelegate::CreateStatic(IsHiddenInNonEditorBlueprint));	
diff --git a/Engine/Source/Editor/BlueprintGraph/Public/BlueprintActionFilter.h b/Engine/Source/Editor/BlueprintGraph/Public/BlueprintActionFilter.h
--- a/Engine/Source/Editor/BlueprintGraph/Public/BlueprintActionFilter.h
+++ b/Engine/Source/Editor/BlueprintGraph/Public/BlueprintActionFilter.h
@@ -187,6 +187,13 @@ public:
 		 * in RejectedNodeTypes). 
 		 */
 		BPFILTER_PermitRejectionSubClasses	= (1<<3),
+
+		/**
+		 * Actions that reference member variables or functions that belong
+		 * to a struct that exists outside the scope of the current set of
+		 * imported namespaces will be filtered out.
+		 */
+		BPFILTER_RejectOutOfScopeMembers	= (1<<4),
 	};
 	FBlueprintActionFilter(uint32 const Flags = 0x00);
 	
diff --git a/Engine/Source/Editor/Kismet/Private/BlueprintActionMenuUtils.cpp b/Engine/Source/Editor/Kismet/Private/BlueprintActionMenuUtils.cpp
--- a/Engine/Source/Editor/Kismet/Private/BlueprintActionMenuUtils.cpp
+++ b/Engine/Source/Editor/Kismet/Private/BlueprintActionMenuUtils.cpp
@@ -30,6 +30,7 @@
 #include "IContentBrowserSingleton.h"
 #include "ContentBrowserModule.h"
 #include "ComponentAssetBroker.h"
+#include "BlueprintNamespaceHelper.h"
 
 #define LOCTEXT_NAMESPACE "BlueprintActionMenuUtils"
 
@@ -455,6 +456,11 @@ void FBlueprintActionMenuUtils::MakeContextMenu(FBlueprintActionContext const& C
 		FilterFlags |= FBlueprintActionFilter::BPFILTER_RejectGlobalFields;
 	}
 
+	if (bIsContextSensitive && FBlueprintNamespaceHelper::IsNamespaceImportScopingEnabled())
+	{
+		FilterFlags |= FBlueprintActionFilter::BPFILTER_RejectOutOfScopeMembers;
+	}
+
 	FBlueprintActionFilter MainMenuFilter(FilterFlags);
 	MainMenuFilter.Context = Context;
 	MainMenuFilter.Context.SelectedObjects.Empty();
diff --git a/Engine/Source/Editor/Kismet/Private/BlueprintDetailsCustomization.cpp b/Engine/Source/Editor/Kismet/Private/BlueprintDetailsCustomization.cpp
--- a/Engine/Source/Editor/Kismet/Private/BlueprintDetailsCustomization.cpp
+++ b/Engine/Source/Editor/Kismet/Private/BlueprintDetailsCustomization.cpp
@@ -58,6 +58,7 @@
 #include "NodeFactory.h"
 #include "Kismet2/Kismet2NameValidators.h"
 #include "Widgets/Layout/SWidgetSwitcher.h"
+#include "BlueprintNamespaceHelper.h"
 
 #include "Modules/ModuleManager.h"
 #include "ISequencerModule.h"
@@ -5148,80 +5149,79 @@ FReply FBlueprintGraphActionDetails::OnAddNewOutputClicked()
 }
 
 
+FBlueprintManagedListDetails::FBlueprintManagedListDetails(TWeakPtr<class FBlueprintGlobalOptionsDetails> InGlobalOptionsDetailsPtr)
+: GlobalOptionsDetailsPtr(InGlobalOptionsDetailsPtr)
+{
+}
 
-void FBlueprintInterfaceLayout::GenerateHeaderRowContent( FDetailWidgetRow& NodeRow )
+UBlueprint* FBlueprintManagedListDetails::GetBlueprintObjectChecked() const
 {
-	NodeRow
+	UBlueprint* BlueprintObject = nullptr;
+
+	TSharedPtr<FBlueprintGlobalOptionsDetails> PinnedGlobalOptionsDetailsPtr = GlobalOptionsDetailsPtr.Pin();
+	if (PinnedGlobalOptionsDetailsPtr.IsValid())
+	{
+		BlueprintObject = PinnedGlobalOptionsDetailsPtr->GetBlueprintObj();
+	}
+
+	check(BlueprintObject);
+	return BlueprintObject;
+}
+
+TSharedPtr<FBlueprintEditor> FBlueprintManagedListDetails::GetPinnedBlueprintEditorPtr() const
+{
+	TSharedPtr<FBlueprintGlobalOptionsDetails> PinnedGlobalOptionsDetailsPtr = GlobalOptionsDetailsPtr.Pin();
+	if (PinnedGlobalOptionsDetailsPtr.IsValid())
+	{
+		return PinnedGlobalOptionsDetailsPtr->GetBlueprintEditorPtr().Pin();
+	}
+
+	return TSharedPtr<FBlueprintEditor>();
+}
+
+void FBlueprintManagedListDetails::GenerateHeaderRowContent(FDetailWidgetRow& HeaderRow)
+{
+	HeaderRow
 	[
 		SNew(STextBlock)
-			.Text( bShowsInheritedInterfaces ?
-			LOCTEXT("BlueprintInheritedInterfaceTitle", "Inherited Interfaces") :
-			LOCTEXT("BlueprintImplementedInterfaceTitle", "Implemented Interfaces") )
-			.Font( IDetailLayoutBuilder::GetDetailFont() )
+		.Text(DisplayOptions.TitleText)
+		.Font(IDetailLayoutBuilder::GetDetailFont())
 	];
 }
 
 BEGIN_SLATE_FUNCTION_BUILD_OPTIMIZATION
-void FBlueprintInterfaceLayout::GenerateChildContent( IDetailChildrenBuilder& ChildrenBuilder )
+void FBlueprintManagedListDetails::GenerateChildContent(IDetailChildrenBuilder& ChildrenBuilder)
 {
-	UBlueprint* Blueprint = GlobalOptionsDetailsPtr.Pin()->GetBlueprintObj();
-	check(Blueprint);
-
-	TArray<FInterfaceName> Interfaces;
+	TArray<FManagedListItem> Items;
+	GetManagedListItems(Items);
 
-	if (!bShowsInheritedInterfaces)
+	if (Items.Num() > 0)
 	{
-		// Generate a list of interfaces already implemented
-		for (const FBPInterfaceDescription& ImplementedInterface : Blueprint->ImplementedInterfaces)
+		for (const FManagedListItem& Item : Items)
 		{
-			if (const TSubclassOf<UInterface> Interface = ImplementedInterface.Interface)
-			{
-				Interfaces.AddUnique(FInterfaceName(Interface->GetFName(), Interface->GetDisplayNameText()));
-			}
-		}
-	}
-	else
-	{
-		// Generate a list of interfaces implemented by classes this blueprint inherited from
-		UClass* BlueprintParent = Blueprint->ParentClass;
-		while (BlueprintParent)
-		{
-			for (TArray<FImplementedInterface>::TIterator It(BlueprintParent->Interfaces); It; ++It)
-			{
-				FImplementedInterface& CurrentInterface = *It;
-				if( CurrentInterface.Class )
-				{
-					Interfaces.Add(FInterfaceName(CurrentInterface.Class->GetFName(), CurrentInterface.Class->GetDisplayNameText()));
-				}
-			}
-			BlueprintParent = BlueprintParent->GetSuperClass();
-		}
-	}
-
-	for (int32 i = 0; i < Interfaces.Num(); ++i)
-	{
-		TSharedPtr<SHorizontalBox> Box;
-		ChildrenBuilder.AddCustomRow( LOCTEXT( "BlueprintInterfaceValue", "Interface Value" ) )
-		[
-			SAssignNew(Box, SHorizontalBox)
-			+SHorizontalBox::Slot()
+			TSharedPtr<SHorizontalBox> Box;
+			ChildrenBuilder.AddCustomRow(DisplayOptions.ItemRowFilterText)
 			[
-				SNew(STextBlock)
-					.Text(Interfaces[i].DisplayText)
-					.Font( IDetailLayoutBuilder::GetDetailFont() )
-			]
-		];
+				SAssignNew(Box, SHorizontalBox)
+				+ SHorizontalBox::Slot()
+				[
+					SNew(STextBlock)
+					.Text(Item.DisplayName)
+					.Font(IDetailLayoutBuilder::GetDetailFont())
+				]
+			];
 
-		// See if we need to add a button for opening this interface
-		if (!bShowsInheritedInterfaces)
-		{
-			UBlueprintGeneratedClass* Class = Cast<UBlueprintGeneratedClass>(*Blueprint->ImplementedInterfaces[i].Interface);
-			if (Class)
+			if (Item.AssetPtr.IsValid())
 			{
-				TWeakObjectPtr<UObject> Asset = Class->ClassGeneratedBy;
-		
-				const TSharedRef<SWidget> BrowseButton = PropertyCustomizationHelpers::MakeBrowseButton(FSimpleDelegate::CreateSP(this, &FBlueprintInterfaceLayout::OnBrowseToInterface, Asset));
-				BrowseButton->SetToolTipText( LOCTEXT("BlueprintInterfaceBrowseTooltip", "Opens this interface") );
+				TSharedRef<SWidget> BrowseButton = PropertyCustomizationHelpers::MakeBrowseButton(FSimpleDelegate::CreateLambda([this, AssetPtr = Item.AssetPtr]() -> void
+				{
+					if (AssetPtr.IsValid())
+					{
+						GEditor->GetEditorSubsystem<UAssetEditorSubsystem>()->OpenEditorForAsset(AssetPtr.Get());
+					}
+				}));
+
+				BrowseButton->SetToolTipText(DisplayOptions.BrowseButtonToolTipText);
 
 				Box->AddSlot()
 				.AutoWidth()
@@ -5230,62 +5230,253 @@ void FBlueprintInterfaceLayout::GenerateChildContent( IDetailChildrenBuilder& Ch
 					BrowseButton
 				];
 			}
-		}
 
-		if (!bShowsInheritedInterfaces)
-		{
-			Box->AddSlot()
-			.AutoWidth()
-			[
-				PropertyCustomizationHelpers::MakeClearButton(FSimpleDelegate::CreateSP(this, &FBlueprintInterfaceLayout::OnRemoveInterface, Interfaces[i]))
-			];
+			if (Item.bIsRemovable)
+			{
+				TSharedRef<SWidget> RemoveButton = PropertyCustomizationHelpers::MakeClearButton(FSimpleDelegate::CreateLambda([this, Item]() -> void
+				{
+					OnRemoveItem(Item);
+				}));
+
+				RemoveButton->SetToolTipText(DisplayOptions.RemoveButtonToolTipText);
+
+				Box->AddSlot()
+				.AutoWidth()
+				[
+					RemoveButton
+				];
+			}
 		}
 	}
-
-	// Add message if no interfaces are being used
-	if (Interfaces.Num() == 0)
+	else
 	{
-		ChildrenBuilder.AddCustomRow(LOCTEXT("BlueprintInterfaceValue", "Interface Value"))
-		[
+		ChildrenBuilder.AddCustomRow(DisplayOptions.ItemRowFilterText)
+			[
 			SNew(STextBlock)
-			.Text(LOCTEXT("NoBlueprintInterface", "No Interfaces"))
+			.Text(DisplayOptions.NoItemsLabelText)
 			.Font(IDetailLayoutBuilder::GetDetailFontItalic())
 		];
 	}
 
-	if (!bShowsInheritedInterfaces)
+	TSharedPtr<SWidget> AddItemRowWidget = MakeAddItemRowWidget();
+	if (AddItemRowWidget.IsValid())
 	{
-		ChildrenBuilder.AddCustomRow( LOCTEXT( "BlueprintAddInterface", "Add Interface" ) )
+		ChildrenBuilder.AddCustomRow(DisplayOptions.AddItemRowFilterText)
 		[
 			SNew(SBox)
 			.HAlign(HAlign_Right)
 			[
-				SAssignNew(AddInterfaceComboButton, SComboButton)
-				.ButtonContent()
-				[
-					SNew(STextBlock)
-						.Text(LOCTEXT("BlueprintAddInterfaceButton", "Add"))
-				]
-				.OnGetMenuContent(this, &FBlueprintInterfaceLayout::OnGetAddInterfaceMenuContent)
+				AddItemRowWidget.ToSharedRef()
 			]
 		];
 	}
 }
 END_SLATE_FUNCTION_BUILD_OPTIMIZATION
 
-void FBlueprintInterfaceLayout::OnBrowseToInterface(TWeakObjectPtr<UObject> Asset)
+void FBlueprintManagedListDetails::RegenerateChildContent()
+{
+	RegenerateChildrenDelegate.ExecuteIfBound();
+}
+
+void FBlueprintManagedListDetails::OnRefreshInDetailsView()
+{
+	TSharedPtr<FBlueprintEditor> BlueprintEditorPtr = GetPinnedBlueprintEditorPtr();
+	if (BlueprintEditorPtr.IsValid())
+	{
+		TSharedPtr<SKismetInspector> Inspector = BlueprintEditorPtr->GetInspector();
+		if (Inspector.IsValid())
+		{
+			// Show details for the Blueprint instance we're editing
+			Inspector->ShowDetailsForSingleObject(GetBlueprintObjectChecked());
+		}
+	}
+}
+
+
+FBlueprintImportsLayout::FBlueprintImportsLayout(TWeakPtr<class FBlueprintGlobalOptionsDetails> InGlobalOptionsDetails)
+	:FBlueprintManagedListDetails(InGlobalOptionsDetails)
+{
+	DisplayOptions.TitleText = LOCTEXT("BlueprintImportedNamespaceTitle", "Imported Namespaces");
+	DisplayOptions.NoItemsLabelText = LOCTEXT("NoBlueprintImports", "No Imports");
+	DisplayOptions.ItemRowFilterText = LOCTEXT("BlueprintImportsValue", "Imports Value");
+	DisplayOptions.AddItemRowFilterText = LOCTEXT("BlueprintAddImport", "Add Import");
+}
+
+TSharedPtr<SWidget> FBlueprintImportsLayout::MakeAddItemRowWidget()
+{
+	return SNew(SHorizontalBox)
+	+ SHorizontalBox::Slot()
+	.FillWidth(1.0f)
+	[
+		SAssignNew(ImportEntryTextBox, SEditableTextBox)
+		.MinDesiredWidth(200.0f)
+		.Font(IDetailLayoutBuilder::GetDetailFont())
+		.OnTextCommitted(this, &FBlueprintImportsLayout::HandleImportEntryTextCommitted)
+	]
+	+ SHorizontalBox::Slot()
+	.AutoWidth()
+	.Padding(4.0f, 0.0f)
+	[
+		SNew(SButton)
+		.Text(LOCTEXT("BlueprintAddImportButtonText", "Add"))
+		.IsEnabled(this, &FBlueprintImportsLayout::IsAddImportEntryButtonEnabled)
+		.OnClicked(this, &FBlueprintImportsLayout::OnAddImportEntryButtonClicked)
+	];
+}
+
+void FBlueprintImportsLayout::GetManagedListItems(TArray<FManagedListItem>& OutListItems) const
 {
-	if (Asset.IsValid())
+	UBlueprint* Blueprint = GetBlueprintObjectChecked();
+
+	for (const FString& ImportedNamespace : Blueprint->ImportedNamespaces)
 	{
-		GEditor->GetEditorSubsystem<UAssetEditorSubsystem>()->OpenEditorForAsset(Asset.Get());
+		FManagedListItem ItemDesc;
+		ItemDesc.ItemName = ImportedNamespace;
+		ItemDesc.DisplayName = FText::FromString(ImportedNamespace);
+		ItemDesc.bIsRemovable = true;
+
+		OutListItems.Add(MoveTemp(ItemDesc));
+	}
+}
+
+void FBlueprintImportsLayout::OnRemoveItem(const FManagedListItem& Item)
+{
+	UBlueprint* Blueprint = GetBlueprintObjectChecked();
+
+	Blueprint->ImportedNamespaces.Remove(Item.ItemName);
+
+	FBlueprintEditorUtils::MarkBlueprintAsModified(Blueprint);
+
+	RegenerateChildContent();
+
+	OnRefreshInDetailsView();
+}
+
+bool FBlueprintImportsLayout::IsAddImportEntryButtonEnabled() const
+{
+	return !ImportEntryTextBox->GetText().IsEmptyOrWhitespace();
+}
+
+FReply FBlueprintImportsLayout::OnAddImportEntryButtonClicked()
+{
+	HandleImportEntryTextCommitted(ImportEntryTextBox->GetText(), ETextCommit::OnEnter);
+	return FReply::Handled();
+}
+
+void FBlueprintImportsLayout::HandleImportEntryTextCommitted(const FText& NewLabel, ETextCommit::Type CommitType)
+{
+	bool bWasAdded = false;
+	const FString& Namespace = NewLabel.ToString();
+
+	// Add to the blueprint's list of imports.
+	if (!Namespace.IsEmpty())
+	{
+		UBlueprint* Blueprint = GetBlueprintObjectChecked();
+
+		if (FBlueprintEditorUtils::AddNamespaceToImportList(Blueprint, Namespace))
+		{
+			bWasAdded = true;
+		}
+
+		RegenerateChildContent();
+	}
+
+	OnRefreshInDetailsView();
+
+	// If added above, import the namespace into the current editor context.
+	if (bWasAdded)
+	{
+		TSharedPtr<FBlueprintEditor> BlueprintEditorPtr = GetPinnedBlueprintEditorPtr();
+		if (BlueprintEditorPtr.IsValid())
+		{
+			BlueprintEditorPtr->ImportNamespace(Namespace);
+		}
+	}
+}
+
+
+FBlueprintInterfaceLayout::FBlueprintInterfaceLayout(TWeakPtr<class FBlueprintGlobalOptionsDetails> InGlobalOptionsDetails, bool bInShowsInheritedInterfaces)
+	:FBlueprintManagedListDetails(InGlobalOptionsDetails)
+	,bShowsInheritedInterfaces(bInShowsInheritedInterfaces)
+{
+	DisplayOptions.TitleText = bShowsInheritedInterfaces ?
+		LOCTEXT("BlueprintInheritedInterfaceTitle", "Inherited Interfaces") :
+		LOCTEXT("BlueprintImplementedInterfaceTitle", "Implemented Interfaces");
+	DisplayOptions.NoItemsLabelText = LOCTEXT("NoBlueprintInterface", "No Interfaces");
+	DisplayOptions.ItemRowFilterText = LOCTEXT("BlueprintInterfaceValue", "Interface Value");
+	DisplayOptions.AddItemRowFilterText = LOCTEXT("BlueprintAddInterface", "Add Interface");
+	DisplayOptions.BrowseButtonToolTipText = LOCTEXT("BlueprintInterfaceBrowseTooltip", "Opens this interface");
+}
+
+TSharedPtr<SWidget> FBlueprintInterfaceLayout::MakeAddItemRowWidget()
+{
+	if (!bShowsInheritedInterfaces)
+	{
+		return SAssignNew(AddInterfaceComboButton, SComboButton)
+		.ButtonContent()
+		[
+			SNew(STextBlock)
+			.Text(LOCTEXT("BlueprintAddInterfaceButton", "Add"))
+		]
+		.OnGetMenuContent(this, &FBlueprintInterfaceLayout::OnGetAddInterfaceMenuContent);
 	}
+
+	return nullptr;
 }
 
-void FBlueprintInterfaceLayout::OnRemoveInterface(FInterfaceName InterfaceName)
+void FBlueprintInterfaceLayout::GetManagedListItems(TArray<FManagedListItem>& OutListItems) const
 {
-	UBlueprint* Blueprint = GlobalOptionsDetailsPtr.Pin()->GetBlueprintObj();
-	check(Blueprint);
+	const UBlueprint* Blueprint = GetBlueprintObjectChecked();
+
+	if (!bShowsInheritedInterfaces)
+	{
+		// Generate a list of interfaces already implemented
+		for (const FBPInterfaceDescription& ImplementedInterface : Blueprint->ImplementedInterfaces)
+		{
+			if (const TSubclassOf<UInterface> Interface = ImplementedInterface.Interface)
+			{
+				FManagedListItem ItemDesc;
+				ItemDesc.ItemName = Interface->GetFName().ToString();
+				ItemDesc.DisplayName = Interface->GetDisplayNameText();
+				ItemDesc.bIsRemovable = true;
 
+				// Allow browsing to Blueprint interface class assets.
+				if (UBlueprintGeneratedClass* Class = Cast<UBlueprintGeneratedClass>(*Interface))
+				{
+					ItemDesc.AssetPtr = Class->ClassGeneratedBy;
+				}
+
+				OutListItems.Add(MoveTemp(ItemDesc));
+			}
+		}
+	}
+	else
+	{
+		// Generate a list of interfaces implemented by classes this blueprint inherited from
+		UClass* BlueprintParent = Blueprint->ParentClass;
+		while (BlueprintParent)
+		{
+			for (TArray<FImplementedInterface>::TIterator It(BlueprintParent->Interfaces); It; ++It)
+			{
+				FImplementedInterface& CurrentInterface = *It;
+				if (CurrentInterface.Class)
+				{
+					FManagedListItem ItemDesc;
+					ItemDesc.ItemName = CurrentInterface.Class->GetFName().ToString();
+					ItemDesc.DisplayName = CurrentInterface.Class->GetDisplayNameText();
+					ItemDesc.bIsRemovable = false;
+
+					OutListItems.Add(MoveTemp(ItemDesc));
+				}
+			}
+
+			BlueprintParent = BlueprintParent->GetSuperClass();
+		}
+	}
+}
+
+void FBlueprintInterfaceLayout::OnRemoveItem(const FManagedListItem& Item)
+{
 	const EAppReturnType::Type DialogReturn = FMessageDialog::Open(EAppMsgType::YesNoCancel, NSLOCTEXT("UnrealEd", "TransferInterfaceFunctionsToBlueprint", "Would you like to transfer the interface functions to be part of your blueprint?"));
 
 	if (DialogReturn == EAppReturnType::Cancel)
@@ -5293,20 +5484,25 @@ void FBlueprintInterfaceLayout::OnRemoveInterface(FInterfaceName InterfaceName)
 		// We canceled!
 		return;
 	}
-	const FName InterfaceFName = InterfaceName.Name;
+	const FName InterfaceFName = FName(*Item.ItemName);
 
-	// Close all graphs that are about to be removed
-	TArray<UEdGraph*> Graphs;
-	FBlueprintEditorUtils::GetInterfaceGraphs(Blueprint, InterfaceFName, Graphs);
-	for( TArray<UEdGraph*>::TIterator GraphIt(Graphs); GraphIt; ++GraphIt )
+	UBlueprint* Blueprint = GetBlueprintObjectChecked();
+	TSharedPtr<FBlueprintEditor> BlueprintEditorPtr = GetPinnedBlueprintEditorPtr();
+	if (BlueprintEditorPtr.IsValid())
 	{
-		GlobalOptionsDetailsPtr.Pin()->GetBlueprintEditorPtr().Pin()->CloseDocumentTab(*GraphIt);
+		// Close all graphs that are about to be removed
+		TArray<UEdGraph*> Graphs;
+		FBlueprintEditorUtils::GetInterfaceGraphs(Blueprint, InterfaceFName, Graphs);
+		for (TArray<UEdGraph*>::TIterator GraphIt(Graphs); GraphIt; ++GraphIt)
+		{
+			BlueprintEditorPtr->CloseDocumentTab(*GraphIt);
+		}
 	}
 
 	// Do the work of actually removing the interface
 	FBlueprintEditorUtils::RemoveInterface(Blueprint, InterfaceFName, DialogReturn == EAppReturnType::Yes);
 
-	RegenerateChildrenDelegate.ExecuteIfBound();
+	RegenerateChildContent();
 
 	OnRefreshInDetailsView();
 }
@@ -5320,12 +5516,11 @@ void FBlueprintInterfaceLayout::OnClassPicked(UClass* PickedClass)
 
 	if (PickedClass)
 	{
-		UBlueprint* Blueprint = GlobalOptionsDetailsPtr.Pin()->GetBlueprintObj();
-		check(Blueprint);
+		UBlueprint* Blueprint = GetBlueprintObjectChecked();
 
 		FBlueprintEditorUtils::ImplementNewInterface(Blueprint, PickedClass->GetFName());
 
-		RegenerateChildrenDelegate.ExecuteIfBound();
+		RegenerateChildContent();
 	}
 
 	OnRefreshInDetailsView();
@@ -5333,7 +5528,7 @@ void FBlueprintInterfaceLayout::OnClassPicked(UClass* PickedClass)
 
 TSharedRef<SWidget> FBlueprintInterfaceLayout::OnGetAddInterfaceMenuContent()
 {
-	UBlueprint* Blueprint = GlobalOptionsDetailsPtr.Pin()->GetBlueprintObj();
+	UBlueprint* Blueprint = GetBlueprintObjectChecked();
 
 	TArray<UBlueprint*> Blueprints;
 	Blueprints.Add(Blueprint);
@@ -5352,16 +5547,6 @@ TSharedRef<SWidget> FBlueprintInterfaceLayout::OnGetAddInterfaceMenuContent()
 		];
 }
 
-void FBlueprintInterfaceLayout::OnRefreshInDetailsView()
-{
-	TSharedPtr<SKismetInspector> Inspector = GlobalOptionsDetailsPtr.Pin()->GetBlueprintEditorPtr().Pin()->GetInspector();
-	UBlueprint* Blueprint = GlobalOptionsDetailsPtr.Pin()->GetBlueprintObj();
-	check(Blueprint);
-
-	// Show details for the Blueprint instance we're editing
-	Inspector->ShowDetailsForSingleObject(Blueprint);
-}
-
 UBlueprint* FBlueprintGlobalOptionsDetails::GetBlueprintObj() const
 {
 	if(BlueprintEditorPtr.IsValid())
@@ -5503,6 +5688,16 @@ void FBlueprintGlobalOptionsDetails::CustomizeDetails(IDetailLayoutBuilder& Deta
 		const bool bIsLevelScriptBP = FBlueprintEditorUtils::IsLevelScriptBlueprint(Blueprint);
 		const bool bIsFunctionLibrary = Blueprint->BlueprintType == BPTYPE_FunctionLibrary;
 		const bool bSupportsInterfaces = !bIsInterfaceBP && !bIsMacroLibrary && !bIsFunctionLibrary;
+		const bool bSupportsNamespaces = FBlueprintNamespaceHelper::IsNamespaceImportEditorUXEnabled();
+
+		if(bSupportsNamespaces)
+		{
+			// Imported namespace details
+			IDetailCategoryBuilder& ImportsCategory = DetailLayout.EditCategory("Imports", LOCTEXT("BlueprintImportDetailsCategory", "Imports"));
+
+			TSharedRef<FBlueprintImportsLayout> ImportsLayout = MakeShareable(new FBlueprintImportsLayout(SharedThis(this)));
+			ImportsCategory.AddCustomBuilder(ImportsLayout);
+		}
 
 		if (bSupportsInterfaces)
 		{
diff --git a/Engine/Source/Editor/Kismet/Private/BlueprintDetailsCustomization.h b/Engine/Source/Editor/Kismet/Private/BlueprintDetailsCustomization.h
--- a/Engine/Source/Editor/Kismet/Private/BlueprintDetailsCustomization.h
+++ b/Engine/Source/Editor/Kismet/Private/BlueprintDetailsCustomization.h
@@ -663,48 +663,63 @@ private:
 	TArray<TSharedPtr<IDetailCustomization>> ExternalDetailCustomizations;
 };
 
-/** Blueprint Interface List Details */
-class FBlueprintInterfaceLayout : public IDetailCustomNodeBuilder, public TSharedFromThis<FBlueprintInterfaceLayout>
+/** Blueprint managed list details */
+class FBlueprintManagedListDetails : public IDetailCustomNodeBuilder
 {
 public:
-	FBlueprintInterfaceLayout(TWeakPtr<class FBlueprintGlobalOptionsDetails> InGlobalOptionsDetails, bool bInShowsInheritedInterfaces)
-		: GlobalOptionsDetailsPtr(InGlobalOptionsDetails)
-		, bShowsInheritedInterfaces(bInShowsInheritedInterfaces) {}
-
-	struct FInterfaceName
+	/** Constructor method */
+	FBlueprintManagedListDetails(TWeakPtr<class FBlueprintGlobalOptionsDetails> InGlobalOptionsDetailsPtr);
+	
+protected:
+	/** List item node type */
+	struct FManagedListItem
 	{
-		FName Name;
-		FText DisplayText;
+		FString ItemName;
+		FText DisplayName;
+		TWeakObjectPtr<> AssetPtr;
+		uint8 bIsRemovable : 1;
 
-		FInterfaceName() {}
-		FInterfaceName(FName InName, const FText& InDisplayText) 
-			: Name(InName), DisplayText(InDisplayText) {}
-
-		bool operator==(const FInterfaceName& Other) const
+		FManagedListItem()
+			:bIsRemovable(false)
 		{
-			return Name == Other.Name;
 		}
 	};
 
-private:
-	/** IDetailCustomNodeBuilder Interface*/
-	virtual void SetOnRebuildChildren( FSimpleDelegate InOnRegenerateChildren ) override {RegenerateChildrenDelegate = InOnRegenerateChildren;}
-	virtual void GenerateHeaderRowContent( FDetailWidgetRow& NodeRow ) override;
-	virtual void GenerateChildContent( IDetailChildrenBuilder& ChildrenBuilder ) override;
-	virtual void Tick( float DeltaTime ) override {}
-	virtual bool RequiresTick() const override { return false; }
+	/** Customizable display options */
+	struct FDisplayOptions
+	{
+		FText TitleText;
+		FText NoItemsLabelText;
+		FText ItemRowFilterText;
+		FText AddItemRowFilterText;
+		FText BrowseButtonToolTipText;
+		FText RemoveButtonToolTipText;
+	};
+
+	/** Mutable display options */
+	FDisplayOptions DisplayOptions;
+
+	/** Access to external resources */
+	UBlueprint* GetBlueprintObjectChecked() const;
+	TSharedPtr<FBlueprintEditor> GetPinnedBlueprintEditorPtr() const;
+
+	/** IDetailCustomNodeBuilder interface */
+	virtual void GenerateHeaderRowContent(FDetailWidgetRow& HeaderRow) override;
+	virtual void GenerateChildContent(IDetailChildrenBuilder& ChildrenBuilder) override;
 	virtual FName GetName() const override { return NAME_None; }
 	virtual bool InitiallyCollapsed() const override { return false; }
-	
-private:
-	/** Callbacks for details UI */
-	void OnBrowseToInterface(TWeakObjectPtr<UObject> Asset);
-	void OnRemoveInterface(FInterfaceName InterfaceName);
+	virtual void Tick(float DeltaTime) override {}
+	virtual bool RequiresTick() const override { return false; }
+	virtual void SetOnRebuildChildren(FSimpleDelegate InOnRebuildChildren) override { RegenerateChildrenDelegate = InOnRebuildChildren; }
+	/** END IDetailCustomNodeBuilder interface */
 
-	TSharedRef<SWidget> OnGetAddInterfaceMenuContent();
+	/** Overridable interface methods */
+	virtual TSharedPtr<SWidget> MakeAddItemRowWidget() { return nullptr; }
+	virtual void GetManagedListItems(TArray<FManagedListItem>& OutListItems) const {}
+	virtual void OnRemoveItem(const FManagedListItem& Item) {}
 
-	/** Callback function when an interface class is picked */
-	void OnClassPicked(UClass* PickedClass);
+	/** Helper function to regenerate the details customization */
+	void RegenerateChildContent();
 
 	/** Helper function to set the Blueprint back into the KismetInspector's details view */
 	void OnRefreshInDetailsView();
@@ -713,17 +728,57 @@ private:
 	/** The parent graph action details customization */
 	TWeakPtr<class FBlueprintGlobalOptionsDetails> GlobalOptionsDetailsPtr;
 
+	/** A delegate to regenerate this list of children */
+	FSimpleDelegate RegenerateChildrenDelegate;
+};
+
+/** Blueprint Imports List Details */
+class FBlueprintImportsLayout : public FBlueprintManagedListDetails, public TSharedFromThis<FBlueprintImportsLayout>
+{
+public:
+	FBlueprintImportsLayout(TWeakPtr<class FBlueprintGlobalOptionsDetails> InGlobalOptionsDetails);
+
+protected:
+	/** FBlueprintManagedListDetails interface*/
+	virtual TSharedPtr<SWidget> MakeAddItemRowWidget() override;
+	virtual void GetManagedListItems(TArray<FManagedListItem>& OutListItems) const override;
+	virtual void OnRemoveItem(const FManagedListItem& Item) override;
+	/** END FBlueprintManagedListDetails interface */
+
+	bool IsAddImportEntryButtonEnabled() const;
+	FReply OnAddImportEntryButtonClicked();
+	void HandleImportEntryTextCommitted(const FText& NewLabel, ETextCommit::Type CommitType);
+
+private:
+	TSharedPtr<SEditableTextBox> ImportEntryTextBox;
+};
+
+/** Blueprint Interface List Details */
+class FBlueprintInterfaceLayout : public FBlueprintManagedListDetails, public TSharedFromThis<FBlueprintInterfaceLayout>
+{
+public:
+	FBlueprintInterfaceLayout(TWeakPtr<class FBlueprintGlobalOptionsDetails> InGlobalOptionsDetails, bool bInShowsInheritedInterfaces);
+
+protected:
+	/** FBlueprintManagedListDetails interface */
+	virtual TSharedPtr<SWidget> MakeAddItemRowWidget() override;
+	virtual void GetManagedListItems(TArray<FManagedListItem>& OutListItems) const override;
+	virtual void OnRemoveItem(const FManagedListItem& Item) override;
+	/** END FBlueprintManagedListDetails interface */
+	
+private:
+	/** Callbacks for details UI */
+	TSharedRef<SWidget> OnGetAddInterfaceMenuContent();
+
+	/** Callback function when an interface class is picked */
+	void OnClassPicked(UClass* PickedClass);
+
+private:
 	/** Whether we show inherited interfaces versus implemented interfaces */
 	bool bShowsInheritedInterfaces;
 
-	/** List of unimplemented interfaces, for source for a list view */
-	TArray<TSharedPtr<FInterfaceName>> UnimplementedInterfaces;
-
 	/** The add interface combo button */
 	TSharedPtr<SComboButton> AddInterfaceComboButton;
-
-	/** A delegate to regenerate this list of children */
-	FSimpleDelegate RegenerateChildrenDelegate;
 };
 
 /** Details customization for Blueprint settings */
diff --git a/Engine/Source/Editor/Kismet/Private/BlueprintEditor.cpp b/Engine/Source/Editor/Kismet/Private/BlueprintEditor.cpp
--- a/Engine/Source/Editor/Kismet/Private/BlueprintEditor.cpp
+++ b/Engine/Source/Editor/Kismet/Private/BlueprintEditor.cpp
@@ -101,6 +101,7 @@
 #include "HAL/PlatformApplicationMisc.h"
 #include "Stats/StatsHierarchical.h"
 #include "BlueprintEditorLibrary.h"
+#include "BlueprintNamespaceHelper.h"
 
 #include "BlueprintEditorTabs.h"
 
@@ -1722,66 +1723,8 @@ void FBlueprintEditor::CommonInitialization(const TArray<UBlueprint*>& InitBluep
 	FKismetEditorUtilities::OnBlueprintUnloaded.AddSP(this, &FBlueprintEditor::OnBlueprintUnloaded);
 }
 
-struct FBlueprintNamespaceHelper
-{
-	TSet<FString> FullyQualifiedListOfNamespaces;
-
-	FBlueprintNamespaceHelper()
-	{
-		AddNamespaces(GetDefault<UBlueprintEditorSettings>()->NamespacesToAlwaysInclude);
-		AddNamespaces(GetDefault<UBlueprintEditorProjectSettings>()->NamespacesToAlwaysInclude);
-	}
-
-	
-	void AddNamespaces(const TArray<FString>& List)
-	{
-		for (const FString& Entry : List)
-		{
-			FullyQualifiedListOfNamespaces.Add(Entry);
-		}
-	}
-
-	void AddNamespace(const FString& Namespace)
-	{
-		if (!Namespace.IsEmpty())
-		{
-			FullyQualifiedListOfNamespaces.Add(Namespace);
-		}
-	}
-
-	bool IsIncludedInNamespaceList(const FString& TestNamespace) const
-	{
-		// Empty namespace == global namespace
-		if (TestNamespace.IsEmpty())
-		{
-			return true;
-		}
-
-		// Check recursively to see if X.Y.Z is present, and if not X.Y (which contains X.Y.Z), and so on until we run out of path segments
-		if (FullyQualifiedListOfNamespaces.Contains(TestNamespace))
-		{
-			return true;
-		}
-		else
-		{
-			int32 RightmostDotIndex;
-			if (TestNamespace.FindLastChar(TEXT('.'), /*out*/ RightmostDotIndex))
-			{
-				if (RightmostDotIndex > 0)
-				{
-					return IsIncludedInNamespaceList(TestNamespace.Left(RightmostDotIndex));
-				}
-			}
-		}
-
-		return false;
-	}
-};
-
 void FBlueprintEditor::LoadLibrariesFromAssetRegistry()
 {
-	FBlueprintNamespaceHelper NamespaceHelper;
-
 	if (EnableAutomaticLibraryAssetLoading == 0)
 	{
 		return;
@@ -1792,7 +1735,7 @@ void FBlueprintEditor::LoadLibrariesFromAssetRegistry()
 		const FString UserDeveloperPath = FPackageName::FilenameToLongPackageName( FPaths::GameUserDeveloperDir());
 		const FString DeveloperPath = FPackageName::FilenameToLongPackageName( FPaths::GameDevelopersDir() );
 
-		NamespaceHelper.AddNamespace(BP->BlueprintNamespace);
+		FBlueprintNamespaceHelper NamespaceHelper(BP);
 
 		// Interface blueprints don't show a node context menu anywhere so we can skip library loading
 		if (BP->BlueprintType != BPTYPE_Interface)
@@ -1898,6 +1841,12 @@ void FBlueprintEditor::LoadLibrariesFromAssetRegistry()
 	}
 }
 
+void FBlueprintEditor::ImportNamespace(const FString& InNamespace)
+{
+	// @todo - Make this more targeted - i.e. get/load only those assets tagged w/ the given namespace
+	LoadLibrariesFromAssetRegistry();
+}
+
 void FBlueprintEditor::RegisterTabSpawners(const TSharedRef<class FTabManager>& InTabManager)
 {
 	//@TODO: Can't we do this sooner?
diff --git a/Engine/Source/Editor/Kismet/Private/BlueprintNamespaceHelper.cpp b/Engine/Source/Editor/Kismet/Private/BlueprintNamespaceHelper.cpp
new file mode 100644
--- /dev/null
+++ b/Engine/Source/Editor/Kismet/Private/BlueprintNamespaceHelper.cpp
@@ -0,0 +1,88 @@
+// Copyright Epic Games, Inc. All Rights Reserved.
+
+#include "BlueprintNamespaceHelper.h"
+#include "Engine/Blueprint.h"
+#include "BlueprintEditorSettings.h"
+#include "Settings/EditorProjectSettings.h"
+
+FBlueprintNamespaceHelper::FBlueprintNamespaceHelper(const UBlueprint* InBlueprint)
+{
+	// Skip namespace adds if the feature is disabled.
+	if (!IsNamespaceImportScopingEnabled())
+	{
+		return;
+	}
+
+	// Default namespace paths implicitly imported by every Blueprint.
+	AddNamespaces(GetDefault<UBlueprintEditorSettings>()->NamespacesToAlwaysInclude);
+	AddNamespaces(GetDefault<UBlueprintEditorProjectSettings>()->NamespacesToAlwaysInclude);
+
+	if (InBlueprint)
+	{
+		AddNamespace(InBlueprint->BlueprintNamespace);
+		AddNamespaces(InBlueprint->ImportedNamespaces);
+
+		// @todo - Add parent class namespaces as well?
+	}
+}
+
+bool FBlueprintNamespaceHelper::IsIncludedInNamespaceList(const FString& TestNamespace) const
+{
+	// Always include if the feature is disabled.
+	if (!IsNamespaceImportScopingEnabled())
+	{
+		return true;
+	}
+
+	// Empty namespace == global namespace
+	if (TestNamespace.IsEmpty())
+	{
+		return true;
+	}
+
+	// Check recursively to see if X.Y.Z is present, and if not X.Y (which contains X.Y.Z), and so on until we run out of path segments
+	if (FullyQualifiedListOfNamespaces.Contains(TestNamespace))
+	{
+		return true;
+	}
+	else
+	{
+		int32 RightmostDotIndex;
+		if (TestNamespace.FindLastChar(TEXT('.'), /*out*/ RightmostDotIndex))
+		{
+			if (RightmostDotIndex > 0)
+			{
+				return IsIncludedInNamespaceList(TestNamespace.Left(RightmostDotIndex));
+			}
+		}
+	}
+
+	return false;
+}
+
+static TAutoConsoleVariable<bool> CVarBPEnableNamespaceImportScoping(
+	TEXT("BP.EnableNamespaceImportScoping"),
+	false,
+	TEXT("Enable Blueprint namespace import scoping features (experimental)."));
+
+bool FBlueprintNamespaceHelper::IsNamespaceImportScopingEnabled()
+{
+	// @todo - Replace with editor/project setting at some point?
+	return CVarBPEnableNamespaceImportScoping.GetValueOnAnyThread();
+}
+
+static TAutoConsoleVariable<bool> CVarBPEnableNamespaceImportEditorUX(
+	TEXT("BP.EnableNamespaceImportEditorUX"),
+	false,
+	TEXT("Enables the namespace importing UX in the Blueprint editor (experimental)."));
+
+bool FBlueprintNamespaceHelper::IsNamespaceImportEditorUXEnabled()
+{
+	if (!IsNamespaceImportScopingEnabled())
+	{
+		return false;
+	}
+
+	// @todo - Replace with editor/project setting at some point?
+	return CVarBPEnableNamespaceImportEditorUX.GetValueOnAnyThread();
+}
\ No newline at end of file
diff --git a/Engine/Source/Editor/Kismet/Public/BlueprintEditor.h b/Engine/Source/Editor/Kismet/Public/BlueprintEditor.h
--- a/Engine/Source/Editor/Kismet/Public/BlueprintEditor.h
+++ b/Engine/Source/Editor/Kismet/Public/BlueprintEditor.h
@@ -646,6 +646,9 @@ public:
 	/** Gets the default schema for this editor */
 	TSubclassOf<UEdGraphSchema> GetDefaultSchema() const { return GetDefaultSchemaClass(); }
 
+	/** Imports the given namespace into the editor. This may trigger a load event for additional macro and/or function library assets if not already loaded. */
+	void ImportNamespace(const FString& InNamespace);
+
 protected:
 	UE_DEPRECATED(4.26, "Please do any validation inside the UBlueprint class during compilation, extra errors during compiling only supplied by the designer can lead to design time only errors being reported and being missed during cooks/content validation.")
 	virtual void AppendExtraCompilerResults(TSharedPtr<class IMessageLogListing> ResultsListing) {}
diff --git a/Engine/Source/Editor/Kismet/Public/BlueprintNamespaceHelper.h b/Engine/Source/Editor/Kismet/Public/BlueprintNamespaceHelper.h
new file mode 100644
--- /dev/null
+++ b/Engine/Source/Editor/Kismet/Public/BlueprintNamespaceHelper.h
@@ -0,0 +1,45 @@
+// Copyright Epic Games, Inc. All Rights Reserved.
+
+#pragma once
+
+#include "CoreMinimal.h"
+
+// Utility class for extracting namespace info from a specified Blueprint context.
+class KISMET_API FBlueprintNamespaceHelper
+{
+public:
+	explicit FBlueprintNamespaceHelper(const class UBlueprint* InBlueprint);
+
+	void AddNamespace(const FString& Namespace)
+	{
+		if (!Namespace.IsEmpty())
+		{
+			FullyQualifiedListOfNamespaces.Add(Namespace);
+		}
+	}
+
+	template<typename ContainerType>
+	void AddNamespaces(const ContainerType& NamespaceList)
+	{
+		if (NamespaceList.Num() > 0)
+		{
+			FullyQualifiedListOfNamespaces.Reserve(FullyQualifiedListOfNamespaces.Num() + NamespaceList.Num());
+			for (const FString& Namespace : NamespaceList)
+			{
+				AddNamespace(Namespace);
+			}
+		}
+	}
+
+	bool IsIncludedInNamespaceList(const FString& TestNamespace) const;
+
+	/** Whether Blueprint namespace-based scoping features are enabled. */
+	static bool IsNamespaceImportScopingEnabled();
+
+	/** Whether the UX for importing namespaces in the Blueprint editor is enabled. */
+	static bool IsNamespaceImportEditorUXEnabled();
+
+private:
+	// Complete list of all fully-qualified namespace path identifiers for the associated Blueprint.
+	TSet<FString> FullyQualifiedListOfNamespaces;
+};
\ No newline at end of file
diff --git a/Engine/Source/Editor/UnrealEd/Private/Kismet2/BlueprintEditorUtils.cpp b/Engine/Source/Editor/UnrealEd/Private/Kismet2/BlueprintEditorUtils.cpp
--- a/Engine/Source/Editor/UnrealEd/Private/Kismet2/BlueprintEditorUtils.cpp
+++ b/Engine/Source/Editor/UnrealEd/Private/Kismet2/BlueprintEditorUtils.cpp
@@ -7585,6 +7585,32 @@ FName FBlueprintEditorUtils::FindUniqueCustomEventName(const UBlueprint* Bluepri
 	return FindUniqueKismetName(Blueprint, LOCTEXT("DefaultCustomEventName", "CustomEvent").ToString());
 }
 
+//////////////////////////////////////////////////////////////////////////
+// Scoping
+
+bool FBlueprintEditorUtils::AddNamespaceToImportList(UBlueprint* Blueprint, const FString& Namespace)
+{
+	check(!Namespace.IsEmpty());
+
+	if (Blueprint->ImportedNamespaces.Contains(Namespace))
+	{
+		ShowNotification(
+			FText::Format(
+				LOCTEXT("NamespaceAlreadyImportedFmt", "Blueprint '{0}' already imports the namespace '{1}'"),
+				FText::FromString(Blueprint->GetName()),
+				FText::FromString(Namespace)
+			),
+			EMessageSeverity::Warning
+		);
+		return false;
+	}
+
+	Blueprint->ImportedNamespaces.Add(Namespace);
+	FBlueprintEditorUtils::MarkBlueprintAsModified(Blueprint);
+
+	return true;
+}
+
 //////////////////////////////////////////////////////////////////////////
 // Timeline
 
diff --git a/Engine/Source/Editor/UnrealEd/Public/Kismet2/BlueprintEditorUtils.h b/Engine/Source/Editor/UnrealEd/Public/Kismet2/BlueprintEditorUtils.h
--- a/Engine/Source/Editor/UnrealEd/Public/Kismet2/BlueprintEditorUtils.h
+++ b/Engine/Source/Editor/UnrealEd/Public/Kismet2/BlueprintEditorUtils.h
@@ -1495,6 +1495,12 @@ public:
 	/** Finds a unique and valid name for a custom event */
 	static FName FindUniqueCustomEventName(const UBlueprint* Blueprint);
 
+	//////////////////////////////////////////////////////////////////////////
+	// Scoping
+
+	/** Add a new namespace to the blueprint's import list. Will return false if already imported. */
+	static bool AddNamespaceToImportList(UBlueprint* Blueprint, const FString& Namespace);
+
 	//////////////////////////////////////////////////////////////////////////
 	// Timeline
 
diff --git a/Engine/Source/Runtime/Engine/Classes/Engine/Blueprint.h b/Engine/Source/Runtime/Engine/Classes/Engine/Blueprint.h
--- a/Engine/Source/Runtime/Engine/Classes/Engine/Blueprint.h
+++ b/Engine/Source/Runtime/Engine/Classes/Engine/Blueprint.h
@@ -607,6 +607,10 @@ PRAGMA_ENABLE_DEPRECATION_WARNINGS
 	UPROPERTY()
 	TArray<FName> CategorySorting;
 
+	/** Namespaces imported by this blueprint */
+	UPROPERTY(AssetRegistrySearchable)
+	TSet<FString> ImportedNamespaces;
+
 	/** Array of info about the interfaces we implement in this blueprint */
 	UPROPERTY(AssetRegistrySearchable)
 	TArray<struct FBPInterfaceDescription> ImplementedInterfaces;
