Need to respect the drag-off location when dragging a pin and  promoting to a variable.

diff --git a/Engine/Source/Editor/Kismet/Private/BlueprintEditor.cpp b/Engine/Source/Editor/Kismet/Private/BlueprintEditor.cpp
--- a/Engine/Source/Editor/Kismet/Private/BlueprintEditor.cpp
+++ b/Engine/Source/Editor/Kismet/Private/BlueprintEditor.cpp
@@ -4007,7 +4007,7 @@ void FBlueprintEditor::DumpMessagesToCompilerLog(const TArray<TSharedRef<FTokeni
 	}
 }
 
-void FBlueprintEditor::DoPromoteToVariable( UBlueprint* InBlueprint, UEdGraphPin* InTargetPin, bool bInToMemberVariable )
+void FBlueprintEditor::DoPromoteToVariable( UBlueprint* InBlueprint, UEdGraphPin* InTargetPin, bool bInToMemberVariable, const FVector2D* InOptionalLocation /* = nullptr */)
 {
 	FName PinName = InTargetPin->PinName;
 	UEdGraphNode* PinNode = InTargetPin->GetOwningNode();
@@ -4153,8 +4153,16 @@ void FBlueprintEditor::DoPromoteToVariable( UBlueprint* InBlueprint, UEdGraphPin
 
 		// Set position of new node to be close to node we clicked on
 		FVector2D NewNodePos;
-		NewNodePos.X = (InTargetPin->Direction == EGPD_Input) ? PinNode->NodePosX - 200 : PinNode->NodePosX + 400;
-		NewNodePos.Y = PinNode->NodePosY;
+
+		if (InOptionalLocation)
+		{
+			NewNodePos = *InOptionalLocation;
+		}
+		else
+		{
+			NewNodePos.X = (InTargetPin->Direction == EGPD_Input) ? PinNode->NodePosX - 200 : PinNode->NodePosX + 400;
+			NewNodePos.Y = PinNode->NodePosY;
+		}
 
 		NodeInfo.PerformAction(GraphObj, InTargetPin, NewNodePos, false);
 
diff --git a/Engine/Source/Editor/Kismet/Private/SBlueprintActionMenu.cpp b/Engine/Source/Editor/Kismet/Private/SBlueprintActionMenu.cpp
--- a/Engine/Source/Editor/Kismet/Private/SBlueprintActionMenu.cpp
+++ b/Engine/Source/Editor/Kismet/Private/SBlueprintActionMenu.cpp
@@ -48,7 +48,7 @@ struct FBlueprintAction_PromoteVariable : public FEdGraphSchemaAction
 			UBlueprint* Blueprint = FBlueprintEditorUtils::FindBlueprintForGraph(ParentGraph);
 			if( ( MyBlueprintEditor.IsValid() == true ) && ( Blueprint != NULL ) )
 			{
-				MyBlueprintEditor.Pin()->DoPromoteToVariable( Blueprint, FromPin, bToMemberVariable );
+				MyBlueprintEditor.Pin()->DoPromoteToVariable( Blueprint, FromPin, bToMemberVariable, &Location );
 			}
 		}
 		return NULL;		
diff --git a/Engine/Source/Editor/Kismet/Public/BlueprintEditor.h b/Engine/Source/Editor/Kismet/Public/BlueprintEditor.h
--- a/Engine/Source/Editor/Kismet/Public/BlueprintEditor.h
+++ b/Engine/Source/Editor/Kismet/Public/BlueprintEditor.h
@@ -554,8 +554,9 @@ public:
 	 * @param	InBlueprint				The blueprint in which to create the variable.
 	 * @param	InTargetPin				The pin on which to base the variable.
 	 * @param	bInToMemberVariable		TRUE if attempting to create a member variable, FALSE if the variable should be local
+	 * @param   InOptionalLocation		Where the new node should be placed. If null, a fixed offset from the parent node will be used.
 	 */
-	void DoPromoteToVariable(UBlueprint* InBlueprint, UEdGraphPin* InTargetPin, bool bInToMemberVariable);
+	void DoPromoteToVariable(UBlueprint* InBlueprint, UEdGraphPin* InTargetPin, bool bInToMemberVariable, const FVector2D* InOptionalLocation = nullptr);
 
 	/** Called when node is spawned by keymap */
 	void OnNodeSpawnedByKeymap();
