Fix a Blueprint graph context menu options arrow clipping issue that made it difficult to reveal the drop-down menu.

Change summary:
- Override button style content padding for this instance to match the previous (UE4) combo button style default.

#rnx

diff --git a/Engine/Source/Editor/Kismet/Private/SBlueprintActionMenu.cpp b/Engine/Source/Editor/Kismet/Private/SBlueprintActionMenu.cpp
--- a/Engine/Source/Editor/Kismet/Private/SBlueprintActionMenu.cpp
+++ b/Engine/Source/Editor/Kismet/Private/SBlueprintActionMenu.cpp
@@ -254,6 +254,7 @@ void SBlueprintActionMenu::Construct( const FArguments& InArgs, TSharedPtr<FBlue
 		.MenuPlacement(MenuPlacement_MenuRight)
 		.HasDownArrow(false)
 		.ButtonStyle(FEditorStyle::Get(), "BlueprintEditor.ContextMenu.TargetsButton")
+		.ContentPadding(FMargin(5))
 		.MenuContent()
 		[
 			SAssignNew(ContextTargetSubMenu, SBlueprintContextTargetMenu, MenuContext)
