Fix for a load time regression when opening the first BP asset for editing.

Change summary:
- Decoupled namespace checks on BP editor load from the recently-added namespace feature-enabling flag since the code path was already being utilized in some projects as a load optimization.

diff --git a/Engine/Source/Editor/Kismet/Private/BlueprintNamespaceHelper.cpp b/Engine/Source/Editor/Kismet/Private/BlueprintNamespaceHelper.cpp
--- a/Engine/Source/Editor/Kismet/Private/BlueprintNamespaceHelper.cpp
+++ b/Engine/Source/Editor/Kismet/Private/BlueprintNamespaceHelper.cpp
@@ -7,12 +7,6 @@
 
 FBlueprintNamespaceHelper::FBlueprintNamespaceHelper(const UBlueprint* InBlueprint)
 {
-	// Skip namespace adds if the feature is disabled.
-	if (!IsNamespaceImportScopingEnabled())
-	{
-		return;
-	}
-
 	// Default namespace paths implicitly imported by every Blueprint.
 	AddNamespaces(GetDefault<UBlueprintEditorSettings>()->NamespacesToAlwaysInclude);
 	AddNamespaces(GetDefault<UBlueprintEditorProjectSettings>()->NamespacesToAlwaysInclude);
@@ -28,12 +22,6 @@ FBlueprintNamespaceHelper::FBlueprintNamespaceHelper(const UBlueprint* InBluepri
 
 bool FBlueprintNamespaceHelper::IsIncludedInNamespaceList(const FString& TestNamespace) const
 {
-	// Always include if the feature is disabled.
-	if (!IsNamespaceImportScopingEnabled())
-	{
-		return true;
-	}
-
 	// Empty namespace == global namespace
 	if (TestNamespace.IsEmpty())
 	{
