Move namespace feature flags into settings (transient for now); add a CVar sink function to bind current flags to settings.

#rnx

diff --git a/Engine/Source/Editor/BlueprintGraph/Private/BlueprintEditorSettings.cpp b/Engine/Source/Editor/BlueprintGraph/Private/BlueprintEditorSettings.cpp
--- a/Engine/Source/Editor/BlueprintGraph/Private/BlueprintEditorSettings.cpp
+++ b/Engine/Source/Editor/BlueprintGraph/Private/BlueprintEditorSettings.cpp
@@ -33,6 +33,7 @@ UBlueprintEditorSettings::UBlueprintEditorSettings(const FObjectInitializer& Obj
 	, bDoubleClickNavigatesToParent(false)
 	, bEnableTypePromotion(true)
 	// Experimental
+	, bEnableNamespaceImportingFeatures(false)
 	, bFavorPureCastNodes(false)
 	// Compiler Settings
 	, SaveOnCompile(SoC_Never)
@@ -115,4 +116,4 @@ void UBlueprintEditorSettings::PostEditChangeProperty(FPropertyChangedEvent& Pro
 	Super::PostEditChangeProperty(PropertyChangedEvent);
 
 	OnSettingsChange.Broadcast(this, PropertyChangedEvent.ChangeType);
-}
+}
\ No newline at end of file
diff --git a/Engine/Source/Editor/BlueprintGraph/Public/BlueprintEditorSettings.h b/Engine/Source/Editor/BlueprintGraph/Public/BlueprintEditorSettings.h
--- a/Engine/Source/Editor/BlueprintGraph/Public/BlueprintEditorSettings.h
+++ b/Engine/Source/Editor/BlueprintGraph/Public/BlueprintEditorSettings.h
@@ -101,6 +101,11 @@ public:
 
 // Experimental
 public:
+	/** Whether to enable namespace importing features in the Blueprint editor */
+	// @todo_namespaces - Make this a non-transient editable config property.
+	UPROPERTY(Transient/*, Category = Experimental*/)
+	bool bEnableNamespaceImportingFeatures;
+
 	// The list of namespaces to always expose in any Blueprint (local per-user)
 	UPROPERTY(EditAnywhere, config, Category = Experimental)
 	TArray<FString> NamespacesToAlwaysInclude;
diff --git a/Engine/Source/Editor/Kismet/Private/BlueprintActionMenuUtils.cpp b/Engine/Source/Editor/Kismet/Private/BlueprintActionMenuUtils.cpp
--- a/Engine/Source/Editor/Kismet/Private/BlueprintActionMenuUtils.cpp
+++ b/Engine/Source/Editor/Kismet/Private/BlueprintActionMenuUtils.cpp
@@ -31,6 +31,7 @@
 #include "ContentBrowserModule.h"
 #include "ComponentAssetBroker.h"
 #include "BlueprintNamespaceHelper.h"
+#include "Settings/EditorExperimentalSettings.h"
 
 #define LOCTEXT_NAMESPACE "BlueprintActionMenuUtils"
 
@@ -456,7 +457,7 @@ void FBlueprintActionMenuUtils::MakeContextMenu(FBlueprintActionContext const& C
 		FilterFlags |= FBlueprintActionFilter::BPFILTER_RejectGlobalFields;
 	}
 
-	if (bIsContextSensitive && FBlueprintNamespaceHelper::IsNamespaceImportScopingEnabled())
+	if (bIsContextSensitive && GetDefault<UEditorExperimentalSettings>()->bEnableNamespaceFilteringFeatures)
 	{
 		FilterFlags |= FBlueprintActionFilter::BPFILTER_RejectOutOfScopeMembers;
 	}
diff --git a/Engine/Source/Editor/Kismet/Private/BlueprintDetailsCustomization.cpp b/Engine/Source/Editor/Kismet/Private/BlueprintDetailsCustomization.cpp
--- a/Engine/Source/Editor/Kismet/Private/BlueprintDetailsCustomization.cpp
+++ b/Engine/Source/Editor/Kismet/Private/BlueprintDetailsCustomization.cpp
@@ -5688,7 +5688,7 @@ void FBlueprintGlobalOptionsDetails::CustomizeDetails(IDetailLayoutBuilder& Deta
 		const bool bIsLevelScriptBP = FBlueprintEditorUtils::IsLevelScriptBlueprint(Blueprint);
 		const bool bIsFunctionLibrary = Blueprint->BlueprintType == BPTYPE_FunctionLibrary;
 		const bool bSupportsInterfaces = !bIsInterfaceBP && !bIsMacroLibrary && !bIsFunctionLibrary;
-		const bool bSupportsNamespaces = FBlueprintNamespaceHelper::IsNamespaceImportEditorUXEnabled();
+		const bool bSupportsNamespaces = GetDefault<UBlueprintEditorSettings>()->bEnableNamespaceImportingFeatures;
 
 		if(bSupportsNamespaces)
 		{
diff --git a/Engine/Source/Editor/Kismet/Private/BlueprintNamespaceHelper.cpp b/Engine/Source/Editor/Kismet/Private/BlueprintNamespaceHelper.cpp
--- a/Engine/Source/Editor/Kismet/Private/BlueprintNamespaceHelper.cpp
+++ b/Engine/Source/Editor/Kismet/Private/BlueprintNamespaceHelper.cpp
@@ -2,8 +2,11 @@
 
 #include "BlueprintNamespaceHelper.h"
 #include "Engine/Blueprint.h"
+#include "BlueprintEditor.h"
 #include "BlueprintEditorSettings.h"
 #include "Settings/EditorProjectSettings.h"
+#include "Settings/EditorExperimentalSettings.h"
+#include "Toolkits/ToolkitManager.h"
 
 FBlueprintNamespaceHelper::FBlueprintNamespaceHelper(const UBlueprint* InBlueprint)
 {
@@ -48,29 +51,57 @@ bool FBlueprintNamespaceHelper::IsIncludedInNamespaceList(const FString& TestNam
 	return false;
 }
 
-static TAutoConsoleVariable<bool> CVarBPEnableNamespaceImportScoping(
-	TEXT("BP.EnableNamespaceImportScoping"),
-	false,
-	TEXT("Enable Blueprint namespace import scoping features (experimental)."));
+// ---
+// @todo_namespaces - Remove CVar flags/sink below after converting to editable 'config' properties
+// ---
 
-bool FBlueprintNamespaceHelper::IsNamespaceImportScopingEnabled()
-{
-	// @todo - Replace with editor/project setting at some point?
-	return CVarBPEnableNamespaceImportScoping.GetValueOnAnyThread();
-}
+static TAutoConsoleVariable<bool> CVarEnableNamespaceFilteringFeatures(
+	TEXT("Editor.EnableNamespaceFilteringFeatures"),
+	false,
+	TEXT("Enables namespace filtering features in the editor UI (experimental).")
+);
 
-static TAutoConsoleVariable<bool> CVarBPEnableNamespaceImportEditorUX(
-	TEXT("BP.EnableNamespaceImportEditorUX"),
+static TAutoConsoleVariable<bool> CVarBPEnableNamespaceImportingFeatures(
+	TEXT("BP.EnableNamespaceImportingFeatures"),
 	false,
-	TEXT("Enables the namespace importing UX in the Blueprint editor (experimental)."));
+	TEXT("Enables namespace importing features in the Blueprint editor (experimental)."));
 
-bool FBlueprintNamespaceHelper::IsNamespaceImportEditorUXEnabled()
+static void UpdateNamespaceFeatureSettingsCVarSinkFunction()
 {
-	if (!IsNamespaceImportScopingEnabled())
+	// Global editor UI settings.
+	UEditorExperimentalSettings* EditorSettingsPtr = GetMutableDefault<UEditorExperimentalSettings>();
+	EditorSettingsPtr->bEnableNamespaceFilteringFeatures = CVarEnableNamespaceFilteringFeatures.GetValueOnGameThread();
+
+	// Blueprint editor settings.
+	UBlueprintEditorSettings* BlueprintEditorSettingsPtr = GetMutableDefault<UBlueprintEditorSettings>();
+	BlueprintEditorSettingsPtr->bEnableNamespaceImportingFeatures = CVarBPEnableNamespaceImportingFeatures.GetValueOnGameThread();
+
+	// Refresh all relevant open Blueprint editor UI elements.
+	// @todo_namespaces - Move this into PostEditChangeProperty() on the appropriate settings object(s).
+	if (GEditor)
 	{
-		return false;
+		if (UAssetEditorSubsystem* AssetEditorSubsystem = GEditor->GetEditorSubsystem<UAssetEditorSubsystem>())
+		{
+			TArray<UObject*> EditedAssets = AssetEditorSubsystem->GetAllEditedAssets();
+			for (UObject* Asset : EditedAssets)
+			{
+				if (Asset && Asset->IsA<UBlueprint>())
+				{
+					TSharedPtr<IToolkit> AssetEditorPtr = FToolkitManager::Get().FindEditorForAsset(Asset);
+					if (AssetEditorPtr.IsValid() && AssetEditorPtr->IsBlueprintEditor())
+					{
+						TSharedPtr<IBlueprintEditor> BlueprintEditorPtr = StaticCastSharedPtr<IBlueprintEditor>(AssetEditorPtr);
+						BlueprintEditorPtr->RefreshMyBlueprint();
+						BlueprintEditorPtr->RefreshInspector();
+					}
+				}
+			}
+		}
 	}
+}
+
+static FAutoConsoleVariableSink CVarUpdateNamespaceFeatureSettingsSink(
+	FConsoleCommandDelegate::CreateStatic(&UpdateNamespaceFeatureSettingsCVarSinkFunction)
+);
 
-	// @todo - Replace with editor/project setting at some point?
-	return CVarBPEnableNamespaceImportEditorUX.GetValueOnAnyThread();
-}
\ No newline at end of file
+// ---
\ No newline at end of file
diff --git a/Engine/Source/Editor/Kismet/Public/BlueprintNamespaceHelper.h b/Engine/Source/Editor/Kismet/Public/BlueprintNamespaceHelper.h
--- a/Engine/Source/Editor/Kismet/Public/BlueprintNamespaceHelper.h
+++ b/Engine/Source/Editor/Kismet/Public/BlueprintNamespaceHelper.h
@@ -33,12 +33,6 @@ public:
 
 	bool IsIncludedInNamespaceList(const FString& TestNamespace) const;
 
-	/** Whether Blueprint namespace-based scoping features are enabled. */
-	static bool IsNamespaceImportScopingEnabled();
-
-	/** Whether the UX for importing namespaces in the Blueprint editor is enabled. */
-	static bool IsNamespaceImportEditorUXEnabled();
-
 private:
 	// Complete list of all fully-qualified namespace path identifiers for the associated Blueprint.
 	TSet<FString> FullyQualifiedListOfNamespaces;
diff --git a/Engine/Source/Editor/UnrealEd/Classes/Settings/EditorExperimentalSettings.h b/Engine/Source/Editor/UnrealEd/Classes/Settings/EditorExperimentalSettings.h
--- a/Engine/Source/Editor/UnrealEd/Classes/Settings/EditorExperimentalSettings.h
+++ b/Engine/Source/Editor/UnrealEd/Classes/Settings/EditorExperimentalSettings.h
@@ -137,6 +137,11 @@ public:
 	UPROPERTY(EditAnywhere, config, Category = "Content Browser", meta = (DisplayName = "Enable support for long paths (> 260 characters)"))
 	bool bEnableLongPathsSupport;
 
+	/** Whether to enable namespace filtering features in the editor UI */
+	// @todo_namespaces - Make this a non-transient editable config property.
+	UPROPERTY(Transient/*, Category = UserInterface*/)
+	bool bEnableNamespaceFilteringFeatures;
+
 	/**
 	 * Returns an event delegate that is executed when a setting has changed.
 	 *
diff --git a/Engine/Source/Editor/UnrealEd/Private/Settings/SettingsClasses.cpp b/Engine/Source/Editor/UnrealEd/Private/Settings/SettingsClasses.cpp
--- a/Engine/Source/Editor/UnrealEd/Private/Settings/SettingsClasses.cpp
+++ b/Engine/Source/Editor/UnrealEd/Private/Settings/SettingsClasses.cpp
@@ -157,6 +157,7 @@ UEditorExperimentalSettings::UEditorExperimentalSettings( const FObjectInitializ
 	, bEnableLocalizationDashboard(true)
 	, bUseOpenCLForConvexHullDecomp(false)
 	, bAllowPotentiallyUnsafePropertyEditing(false)
+	, bEnableNamespaceFilteringFeatures(false)
 {
 }
 
