Fix padding between icon and text in SCSEditor rows

diff --git a/Engine/Source/Editor/Kismet/Private/SSCSEditor.cpp b/Engine/Source/Editor/Kismet/Private/SSCSEditor.cpp
--- a/Engine/Source/Editor/Kismet/Private/SSCSEditor.cpp
+++ b/Engine/Source/Editor/Kismet/Private/SSCSEditor.cpp
@@ -1838,6 +1838,7 @@ TSharedRef<SWidget> SSCS_RowWidget::GenerateWidgetForColumn( const FName& Column
 					]
 				+SHorizontalBox::Slot()
 					.VAlign(VAlign_Center)
+					.Padding(4.0f, 0.0f)
 					[
 						InlineWidget.ToSharedRef()
 					];
