Fix a regression that caused async shader compiles to be indirectly disabled at startup time.

#rnx

diff --git a/Engine/Source/Editor/Kismet/Private/BlueprintNamespaceHelper.cpp b/Engine/Source/Editor/Kismet/Private/BlueprintNamespaceHelper.cpp
--- a/Engine/Source/Editor/Kismet/Private/BlueprintNamespaceHelper.cpp
+++ b/Engine/Source/Editor/Kismet/Private/BlueprintNamespaceHelper.cpp
@@ -68,6 +68,8 @@ static TAutoConsoleVariable<bool> CVarBPEnableNamespaceImportingFeatures(
 
 static void UpdateNamespaceFeatureSettingsCVarSinkFunction()
 {
+#if 0	// need to avoid doing this at startup time
+
 	// Global editor UI settings.
 	UEditorExperimentalSettings* EditorSettingsPtr = GetMutableDefault<UEditorExperimentalSettings>();
 	EditorSettingsPtr->bEnableNamespaceFilteringFeatures = CVarEnableNamespaceFilteringFeatures.GetValueOnGameThread();
@@ -98,6 +100,8 @@ static void UpdateNamespaceFeatureSettingsCVarSinkFunction()
 			}
 		}
 	}
+
+#endif
 }
 
 static FAutoConsoleVariableSink CVarUpdateNamespaceFeatureSettingsSink(
