Correction to SSCEditor to use WeakPtr since the child list already have a SharedPtr.

diff --git a/Engine/Source/Editor/Kismet/Private/SSCSEditor.cpp b/Engine/Source/Editor/Kismet/Private/SSCSEditor.cpp
--- a/Engine/Source/Editor/Kismet/Private/SSCSEditor.cpp
+++ b/Engine/Source/Editor/Kismet/Private/SSCSEditor.cpp
@@ -1667,10 +1667,12 @@ void FSCSEditorTreeNodeRootActor::RemoveChild(FSCSEditorTreeNodePtrType InChildN
 		return NodePtr->IsSceneComponent();
 	});
 
-	if (IndexOfFirstSceneComponent == -1 && SceneComponentSeparatorNodePtr.IsValid())
 	{
-		FSCSEditorTreeNodeActorBase::RemoveChild(SceneComponentSeparatorNodePtr);
-		SceneComponentSeparatorNodePtr.Reset();
+		TSharedPtr<FSCSEditorTreeNodeSeparator> SceneComponentSeparatorNodeSharedPtr = SceneComponentSeparatorNodePtr.Pin();
+		if (IndexOfFirstSceneComponent == -1 && SceneComponentSeparatorNodeSharedPtr.IsValid())
+		{
+			FSCSEditorTreeNodeActorBase::RemoveChild(SceneComponentSeparatorNodeSharedPtr);
+		}
 	}
 
 	const int32 IndexOfFirstNonSceneComponent = ComponentNodePtrs.IndexOfByPredicate([](const FSCSEditorTreeNodePtrType& NodePtr)
@@ -1678,10 +1680,13 @@ void FSCSEditorTreeNodeRootActor::RemoveChild(FSCSEditorTreeNodePtrType InChildN
 		return !NodePtr->IsSceneComponent();
 	});
 
-	if (IndexOfFirstNonSceneComponent == -1 && NonSceneComponentSeparatorNodePtr.IsValid())
 	{
-		FSCSEditorTreeNodeActorBase::RemoveChild(NonSceneComponentSeparatorNodePtr);
-		NonSceneComponentSeparatorNodePtr.Reset();
+		TSharedPtr<FSCSEditorTreeNodeSeparator> NonSceneComponentSeparatorNodeSharedPtr = NonSceneComponentSeparatorNodePtr.Pin();
+
+		if (IndexOfFirstNonSceneComponent == -1 && NonSceneComponentSeparatorNodeSharedPtr.IsValid())
+		{
+			FSCSEditorTreeNodeActorBase::RemoveChild(NonSceneComponentSeparatorNodeSharedPtr);
+		}
 	}
 
 	FSCSEditorTreeNodeActorBase::RemoveChild(InChildNodePtr);
diff --git a/Engine/Source/Editor/Kismet/Public/SSCSEditor.h b/Engine/Source/Editor/Kismet/Public/SSCSEditor.h
--- a/Engine/Source/Editor/Kismet/Public/SSCSEditor.h
+++ b/Engine/Source/Editor/Kismet/Public/SSCSEditor.h
@@ -672,8 +672,8 @@ private:
 	bool bAllowRename;
 	const UClass* CachedFilterType;
 	TArray<FString> CachedFilterTerms;
-	TSharedPtr<class FSCSEditorTreeNodeSeparator> SceneComponentSeparatorNodePtr;
-	TSharedPtr<class FSCSEditorTreeNodeSeparator> NonSceneComponentSeparatorNodePtr;
+	TWeakPtr<class FSCSEditorTreeNodeSeparator> SceneComponentSeparatorNodePtr;
+	TWeakPtr<class FSCSEditorTreeNodeSeparator> NonSceneComponentSeparatorNodePtr;
 };
 
 class KISMET_API FSCSEditorTreeNodeChildActor : public FSCSEditorTreeNodeActorBase
