Remove no longer used toolbar tab from BP macro editor.

diff --git a/Engine/Source/Editor/Kismet/Private/BlueprintEditorModes.cpp b/Engine/Source/Editor/Kismet/Private/BlueprintEditorModes.cpp
--- a/Engine/Source/Editor/Kismet/Private/BlueprintEditorModes.cpp
+++ b/Engine/Source/Editor/Kismet/Private/BlueprintEditorModes.cpp
@@ -432,7 +432,7 @@ FBlueprintMacroApplicationMode::FBlueprintMacroApplicationMode(TSharedPtr<class
 	BlueprintMacroTabFactories.RegisterFactory(MakeShareable(new FFindResultsSummoner(InBlueprintEditor)));
 	BlueprintMacroTabFactories.RegisterFactory(MakeShareable(new FSelectionDetailsSummoner(InBlueprintEditor)));
 
-	TabLayout = FTabManager::NewLayout( "Standalone_BlueprintMacro_Layout_v3" )
+	TabLayout = FTabManager::NewLayout( "Standalone_BlueprintMacro_Layout_v4" )
 		->AddArea
 		(
 			FTabManager::NewPrimaryArea() ->SetOrientation(Orient_Vertical)
@@ -442,14 +442,6 @@ FBlueprintMacroApplicationMode::FBlueprintMacroApplicationMode(TSharedPtr<class
 				->Split
 				(
 					FTabManager::NewSplitter()->SetOrientation( Orient_Vertical )
-					->SetSizeCoefficient(0.70f)
-					->Split
-					(
-						FTabManager::NewStack()
-						->SetSizeCoefficient(0.186721f)
-						->SetHideTabWell(true)
-						->AddTab(InBlueprintEditor->GetToolbarTabId(), ETabState::OpenedTab)
-					)
 					->Split
 					(
 						FTabManager::NewStack()
