More groundwork for namespaces.

Change summary:
- Consolidated transient setting flags into UBlueprintEditorSettings; no longer dependent on UEditorExperimentalSettings. These continue to mirror the CVars for now. Also restored the sink function so that changing options will refresh the UI.
- Compiling a Blueprint will now inject the namespace into the class metadata if it is set on the UBlueprint asset.
- Added IsImportedType()/IsImportedObject() helper methods to FBlueprintNamepaceHelper.
- Added a CVar to provide a way to test the inclusion of parent class namespaces when checking scope; disabled by default for now to preserve current behavior.

diff --git a/Engine/Source/Editor/BlueprintGraph/Classes/EdGraphSchema_K2.h b/Engine/Source/Editor/BlueprintGraph/Classes/EdGraphSchema_K2.h
--- a/Engine/Source/Editor/BlueprintGraph/Classes/EdGraphSchema_K2.h
+++ b/Engine/Source/Editor/BlueprintGraph/Classes/EdGraphSchema_K2.h
@@ -246,6 +246,9 @@ public:
 	/** Metadata that should be used with UPARAM to specify whether a TSubclassOf argument allows abstract classes */
 	static const FName MD_AllowAbstractClasses;
 
+	/** Namespace into which a type can be optionally defined; if empty or not set, the type will belong to the global namespace (default). */
+	static const FName MD_Namespace;
+
 private:
 	// This class should never be instantiated
 	FBlueprintMetadata() {}
diff --git a/Engine/Source/Editor/BlueprintGraph/Private/BlueprintEditorSettings.cpp b/Engine/Source/Editor/BlueprintGraph/Private/BlueprintEditorSettings.cpp
--- a/Engine/Source/Editor/BlueprintGraph/Private/BlueprintEditorSettings.cpp
+++ b/Engine/Source/Editor/BlueprintGraph/Private/BlueprintEditorSettings.cpp
@@ -33,6 +33,7 @@ UBlueprintEditorSettings::UBlueprintEditorSettings(const FObjectInitializer& Obj
 	, bDoubleClickNavigatesToParent(false)
 	, bEnableTypePromotion(true)
 	// Experimental
+	, bEnableNamespaceFilteringFeatures(false)
 	, bEnableNamespaceImportingFeatures(false)
 	, bFavorPureCastNodes(false)
 	// Compiler Settings
diff --git a/Engine/Source/Editor/BlueprintGraph/Private/EdGraphSchema_K2.cpp b/Engine/Source/Editor/BlueprintGraph/Private/EdGraphSchema_K2.cpp
--- a/Engine/Source/Editor/BlueprintGraph/Private/EdGraphSchema_K2.cpp
+++ b/Engine/Source/Editor/BlueprintGraph/Private/EdGraphSchema_K2.cpp
@@ -198,6 +198,8 @@ const FName FBlueprintMetadata::MD_AnimBlueprintFunction(TEXT("AnimBlueprintFunc
 
 const FName FBlueprintMetadata::MD_AllowAbstractClasses(TEXT("AllowAbstract"));
 
+const FName FBlueprintMetadata::MD_Namespace(TEXT("Namespace"));
+
 //////////////////////////////////////////////////////////////////////////
 
 #define LOCTEXT_NAMESPACE "KismetSchema"
diff --git a/Engine/Source/Editor/BlueprintGraph/Public/BlueprintEditorSettings.h b/Engine/Source/Editor/BlueprintGraph/Public/BlueprintEditorSettings.h
--- a/Engine/Source/Editor/BlueprintGraph/Public/BlueprintEditorSettings.h
+++ b/Engine/Source/Editor/BlueprintGraph/Public/BlueprintEditorSettings.h
@@ -101,6 +101,11 @@ public:
 
 // Experimental
 public:
+	/** Whether to enable namespace filtering features in the Blueprint editor */
+	// @todo_namespaces - Make this a non-transient editable config property.
+	UPROPERTY(Transient/*, Category = Experimental*/)
+	bool bEnableNamespaceFilteringFeatures;
+
 	/** Whether to enable namespace importing features in the Blueprint editor */
 	// @todo_namespaces - Make this a non-transient editable config property.
 	UPROPERTY(Transient/*, Category = Experimental*/)
diff --git a/Engine/Source/Editor/Kismet/Private/BlueprintActionMenuUtils.cpp b/Engine/Source/Editor/Kismet/Private/BlueprintActionMenuUtils.cpp
--- a/Engine/Source/Editor/Kismet/Private/BlueprintActionMenuUtils.cpp
+++ b/Engine/Source/Editor/Kismet/Private/BlueprintActionMenuUtils.cpp
@@ -31,7 +31,6 @@
 #include "ContentBrowserModule.h"
 #include "ComponentAssetBroker.h"
 #include "BlueprintNamespaceHelper.h"
-#include "Settings/EditorExperimentalSettings.h"
 
 #define LOCTEXT_NAMESPACE "BlueprintActionMenuUtils"
 
@@ -457,7 +456,7 @@ void FBlueprintActionMenuUtils::MakeContextMenu(FBlueprintActionContext const& C
 		FilterFlags |= FBlueprintActionFilter::BPFILTER_RejectGlobalFields;
 	}
 
-	if (bIsContextSensitive && GetDefault<UEditorExperimentalSettings>()->bEnableNamespaceFilteringFeatures)
+	if (bIsContextSensitive && GetDefault<UBlueprintEditorSettings>()->bEnableNamespaceFilteringFeatures)
 	{
 		FilterFlags |= FBlueprintActionFilter::BPFILTER_RejectOutOfScopeMembers;
 	}
diff --git a/Engine/Source/Editor/Kismet/Private/BlueprintEditor.cpp b/Engine/Source/Editor/Kismet/Private/BlueprintEditor.cpp
--- a/Engine/Source/Editor/Kismet/Private/BlueprintEditor.cpp
+++ b/Engine/Source/Editor/Kismet/Private/BlueprintEditor.cpp
@@ -1735,7 +1735,7 @@ void FBlueprintEditor::LoadLibrariesFromAssetRegistry()
 		const FString UserDeveloperPath = FPackageName::FilenameToLongPackageName( FPaths::GameUserDeveloperDir());
 		const FString DeveloperPath = FPackageName::FilenameToLongPackageName( FPaths::GameDevelopersDir() );
 
-		FBlueprintNamespaceHelper NamespaceHelper(BP);
+		TSharedRef<FBlueprintNamespaceHelper> NamespaceHelper = GetOrCreateNamespaceHelperForBlueprint(BP);
 
 		// Interface blueprints don't show a node context menu anywhere so we can skip library loading
 		if (BP->BlueprintType != BPTYPE_Interface)
@@ -1782,7 +1782,7 @@ void FBlueprintEditor::LoadLibrariesFromAssetRegistry()
 
 					// See if this passes the namespace check
 					const FString AssetBPNamespace = AssetEntry.GetTagValueRef<FString>(BPNamespaceName);
-					bAllowLoadBP = bAllowLoadBP && NamespaceHelper.IsIncludedInNamespaceList(AssetBPNamespace);
+					bAllowLoadBP = bAllowLoadBP && NamespaceHelper->IsIncludedInNamespaceList(AssetBPNamespace);
 
 					// For blueprints inside developers folder, only allow the ones inside current user's developers folder.
 					if (bAllowLoadBP)
@@ -1843,8 +1843,43 @@ void FBlueprintEditor::LoadLibrariesFromAssetRegistry()
 
 void FBlueprintEditor::ImportNamespace(const FString& InNamespace)
 {
-	// @todo - Make this more targeted - i.e. get/load only those assets tagged w/ the given namespace
-	LoadLibrariesFromAssetRegistry();
+	bool bShouldReloadLibraries = false;
+	TTuple<TWeakObjectPtr<const UBlueprint, FWeakObjectPtr>, TUniquePtr<FBlueprintNamespaceHelper, TDefaultDelete<FBlueprintNamespaceHelper>>> TestTuple;
+
+	// Add the namespace to any cached helper objects.
+	const TArray<UObject*>& EditingObjs = GetEditingObjects();
+	for (const UObject* EditingObj : EditingObjs)
+	{
+		if (const UBlueprint* BlueprintObj = Cast<const UBlueprint>(EditingObj))
+		{
+			TSharedRef<FBlueprintNamespaceHelper> NamespaceHelper = GetOrCreateNamespaceHelperForBlueprint(BlueprintObj);
+			if (!NamespaceHelper->IsIncludedInNamespaceList(InNamespace))
+			{
+				NamespaceHelper->AddNamespace(InNamespace);
+
+				// Load additional libraries that may now be in scope.
+				bShouldReloadLibraries = true;
+			}
+		}
+	}
+
+	if (bShouldReloadLibraries)
+	{
+		// @todo_namespaces - Make this more targeted - i.e. get/load only those assets tagged w/ the given namespace
+		LoadLibrariesFromAssetRegistry();
+	}
+}
+
+TSharedRef<FBlueprintNamespaceHelper> FBlueprintEditor::GetOrCreateNamespaceHelperForBlueprint(const UBlueprint* InBlueprint)
+{
+	TWeakObjectPtr<const UBlueprint> Key = TWeakObjectPtr<const UBlueprint>(InBlueprint);
+	if (TSharedRef<FBlueprintNamespaceHelper>* ValuePtr = CachedNamespaceHelpers.Find(Key))
+	{
+		return *ValuePtr;
+	}
+
+	TSharedPtr<FBlueprintNamespaceHelper> NewValue = MakeShared<FBlueprintNamespaceHelper>(InBlueprint);
+	return CachedNamespaceHelpers.Add(Key, NewValue.ToSharedRef());
 }
 
 void FBlueprintEditor::RegisterTabSpawners(const TSharedRef<class FTabManager>& InTabManager)
diff --git a/Engine/Source/Editor/Kismet/Private/BlueprintNamespaceHelper.cpp b/Engine/Source/Editor/Kismet/Private/BlueprintNamespaceHelper.cpp
--- a/Engine/Source/Editor/Kismet/Private/BlueprintNamespaceHelper.cpp
+++ b/Engine/Source/Editor/Kismet/Private/BlueprintNamespaceHelper.cpp
@@ -5,8 +5,86 @@
 #include "BlueprintEditor.h"
 #include "BlueprintEditorSettings.h"
 #include "Settings/EditorProjectSettings.h"
-#include "Settings/EditorExperimentalSettings.h"
 #include "Toolkits/ToolkitManager.h"
+#include "EdGraph/EdGraphPin.h"
+#include "EdGraphSchema_K2.h"
+#include "AssetRegistryModule.h"
+
+// ---
+// @todo_namespaces - Remove CVar flags/sink below after converting to editable 'config' properties
+// ---
+
+static TAutoConsoleVariable<bool> CVarBPEnableNamespaceFilteringFeatures(
+	TEXT("BP.EnableNamespaceFilteringFeatures"),
+	false,
+	TEXT("Enables namespace filtering features in the Blueprint editor (experimental).")
+);
+
+static TAutoConsoleVariable<bool> CVarBPEnableNamespaceImportingFeatures(
+	TEXT("BP.EnableNamespaceImportingFeatures"),
+	false,
+	TEXT("Enables namespace importing features in the Blueprint editor (experimental)."));
+
+static TAutoConsoleVariable<bool> CVarBPImportParentClassNamespaces(
+	TEXT("BP.ImportParentClassNamespaces"),
+	false,
+	TEXT("Enables import of parent class namespaces when opening a Blueprint for editing."));
+
+static void UpdateNamespaceFeatureSettingsCVarSinkFunction()
+{
+	// Note: Do NOT try to access settings objects below during the initial editor load! They rely on the config being loaded, which may not have occurred yet.
+	if (GIsInitialLoad || !GEditor)
+	{
+		return;
+	}
+
+	auto CheckAndUpdateSettingValueLambda = [](bool& CurValue, const bool NewValue) -> bool
+	{
+		if (CurValue != NewValue)
+		{
+			CurValue = NewValue;
+			return true;
+		}
+
+		return false;
+	};
+
+	bool bWasUpdated = false;
+
+	// Blueprint editor settings.
+	UBlueprintEditorSettings* BlueprintEditorSettingsPtr = GetMutableDefault<UBlueprintEditorSettings>();
+	bWasUpdated |= CheckAndUpdateSettingValueLambda(BlueprintEditorSettingsPtr->bEnableNamespaceFilteringFeatures, CVarBPEnableNamespaceFilteringFeatures.GetValueOnGameThread());
+	bWasUpdated |= CheckAndUpdateSettingValueLambda(BlueprintEditorSettingsPtr->bEnableNamespaceImportingFeatures, CVarBPEnableNamespaceImportingFeatures.GetValueOnGameThread());
+
+	if (bWasUpdated)
+	{
+		// Refresh all relevant open Blueprint editor UI elements.
+		// @todo_namespaces - Move this into PostEditChangeProperty() on the appropriate settings object(s).
+		if (UAssetEditorSubsystem* AssetEditorSubsystem = GEditor->GetEditorSubsystem<UAssetEditorSubsystem>())
+		{
+			TArray<UObject*> EditedAssets = AssetEditorSubsystem->GetAllEditedAssets();
+			for (UObject* Asset : EditedAssets)
+			{
+				if (Asset && Asset->IsA<UBlueprint>())
+				{
+					TSharedPtr<IToolkit> AssetEditorPtr = FToolkitManager::Get().FindEditorForAsset(Asset);
+					if (AssetEditorPtr.IsValid() && AssetEditorPtr->IsBlueprintEditor())
+					{
+						TSharedPtr<IBlueprintEditor> BlueprintEditorPtr = StaticCastSharedPtr<IBlueprintEditor>(AssetEditorPtr);
+						BlueprintEditorPtr->RefreshMyBlueprint();
+						BlueprintEditorPtr->RefreshInspector();
+					}
+				}
+			}
+		}
+	}
+}
+
+static FAutoConsoleVariableSink CVarUpdateNamespaceFeatureSettingsSink(
+	FConsoleCommandDelegate::CreateStatic(&UpdateNamespaceFeatureSettingsCVarSinkFunction)
+);
+
+// ---
 
 FBlueprintNamespaceHelper::FBlueprintNamespaceHelper(const UBlueprint* InBlueprint)
 {
@@ -19,7 +97,25 @@ FBlueprintNamespaceHelper::FBlueprintNamespaceHelper(const UBlueprint* InBluepri
 		AddNamespace(InBlueprint->BlueprintNamespace);
 		AddNamespaces(InBlueprint->ImportedNamespaces);
 
-		// @todo - Add parent class namespaces as well?
+		const bool bAddParentClassNamespaces = CVarBPImportParentClassNamespaces.GetValueOnGameThread();
+		if(bAddParentClassNamespaces)
+		{
+			const UClass* ParentClass = InBlueprint->ParentClass;
+			while (ParentClass)
+			{
+				if (const UBlueprint* ParentClassBlueprint = UBlueprint::GetBlueprintFromClass(ParentClass))
+				{
+					AddNamespace(ParentClassBlueprint->BlueprintNamespace);
+					AddNamespaces(ParentClassBlueprint->ImportedNamespaces);
+				}
+				else if (const FString* ParentClassNamespace = ParentClass->FindMetaData(FBlueprintMetadata::MD_Namespace))
+				{
+					AddNamespace(*ParentClassNamespace);
+				}
+
+				ParentClass = ParentClass->GetSuperClass();
+			}
+		}
 	}
 }
 
@@ -51,61 +147,56 @@ bool FBlueprintNamespaceHelper::IsIncludedInNamespaceList(const FString& TestNam
 	return false;
 }
 
-// ---
-// @todo_namespaces - Remove CVar flags/sink below after converting to editable 'config' properties
-// ---
-
-static TAutoConsoleVariable<bool> CVarEnableNamespaceFilteringFeatures(
-	TEXT("Editor.EnableNamespaceFilteringFeatures"),
-	false,
-	TEXT("Enables namespace filtering features in the editor UI (experimental).")
-);
-
-static TAutoConsoleVariable<bool> CVarBPEnableNamespaceImportingFeatures(
-	TEXT("BP.EnableNamespaceImportingFeatures"),
-	false,
-	TEXT("Enables namespace importing features in the Blueprint editor (experimental)."));
-
-static void UpdateNamespaceFeatureSettingsCVarSinkFunction()
+bool FBlueprintNamespaceHelper::IsImportedType(const UField* InType) const
 {
-#if 0	// need to avoid doing this at startup time
+	if (InType)
+	{
+		if (const FString* TypeNamespace = InType->FindMetaData(FBlueprintMetadata::MD_Namespace))
+		{
+			return IsIncludedInNamespaceList(*TypeNamespace);
+		}
+	}
 
-	// Global editor UI settings.
-	UEditorExperimentalSettings* EditorSettingsPtr = GetMutableDefault<UEditorExperimentalSettings>();
-	EditorSettingsPtr->bEnableNamespaceFilteringFeatures = CVarEnableNamespaceFilteringFeatures.GetValueOnGameThread();
+	return true;
+}
 
-	// Blueprint editor settings.
-	UBlueprintEditorSettings* BlueprintEditorSettingsPtr = GetMutableDefault<UBlueprintEditorSettings>();
-	BlueprintEditorSettingsPtr->bEnableNamespaceImportingFeatures = CVarBPEnableNamespaceImportingFeatures.GetValueOnGameThread();
+bool FBlueprintNamespaceHelper::IsImportedObject(const UObject* InObject) const
+{
+	if (const UField* Type = Cast<UField>(InObject))
+	{
+		return IsImportedType(Type);
+	}
+	else
+	{
+		return IsImportedType(InObject->GetClass());
+	}
+}
 
-	// Refresh all relevant open Blueprint editor UI elements.
-	// @todo_namespaces - Move this into PostEditChangeProperty() on the appropriate settings object(s).
-	if (GEditor)
+bool FBlueprintNamespaceHelper::IsImportedObject(const FSoftObjectPath& InObjectPath) const
+{
+	if (const UObject* Object = InObjectPath.ResolveObject())
 	{
-		if (UAssetEditorSubsystem* AssetEditorSubsystem = GEditor->GetEditorSubsystem<UAssetEditorSubsystem>())
+		return IsImportedObject(Object);
+	}
+	
+	const FAssetRegistryModule& AssetRegistryModule = FModuleManager::LoadModuleChecked<FAssetRegistryModule>(TEXT("AssetRegistry"));
+	const FAssetData AssetData = AssetRegistryModule.Get().GetAssetByObjectPath(*InObjectPath.ToString());
+	if (AssetData.IsValid())
+	{
+		if (const UClass* AssetClass = FindObject<UClass>(ANY_PACKAGE, *AssetData.AssetClass.ToString()))
 		{
-			TArray<UObject*> EditedAssets = AssetEditorSubsystem->GetAllEditedAssets();
-			for (UObject* Asset : EditedAssets)
+			if (AssetClass->IsChildOf<UBlueprint>())
 			{
-				if (Asset && Asset->IsA<UBlueprint>())
+				FString OutNamespaceString;
+				if (AssetData.GetTagValue<FString>(GET_MEMBER_NAME_STRING_CHECKED(UBlueprint, BlueprintNamespace), OutNamespaceString))
 				{
-					TSharedPtr<IToolkit> AssetEditorPtr = FToolkitManager::Get().FindEditorForAsset(Asset);
-					if (AssetEditorPtr.IsValid() && AssetEditorPtr->IsBlueprintEditor())
-					{
-						TSharedPtr<IBlueprintEditor> BlueprintEditorPtr = StaticCastSharedPtr<IBlueprintEditor>(AssetEditorPtr);
-						BlueprintEditorPtr->RefreshMyBlueprint();
-						BlueprintEditorPtr->RefreshInspector();
-					}
+					return IsIncludedInNamespaceList(OutNamespaceString);
 				}
 			}
+
+			// @todo_namespaces - Add cases for unloaded UDS/UDE assets once they have a searchable namespace member property.
 		}
 	}
 
-#endif
-}
-
-static FAutoConsoleVariableSink CVarUpdateNamespaceFeatureSettingsSink(
-	FConsoleCommandDelegate::CreateStatic(&UpdateNamespaceFeatureSettingsCVarSinkFunction)
-);
-
-// ---
\ No newline at end of file
+	return true;
+}
\ No newline at end of file
diff --git a/Engine/Source/Editor/Kismet/Public/BlueprintEditor.h b/Engine/Source/Editor/Kismet/Public/BlueprintEditor.h
--- a/Engine/Source/Editor/Kismet/Public/BlueprintEditor.h
+++ b/Engine/Source/Editor/Kismet/Public/BlueprintEditor.h
@@ -51,6 +51,7 @@ struct FToolMenuContext;
 class UK2Node_FunctionEntry;
 class UK2Node_Event;
 class UToolMenu;
+class FBlueprintNamespaceHelper;
 
 /* Enums to use when grouping the blueprint members in the list panel. The order here will determine the order in the list */
 namespace NodeSectionID
@@ -650,6 +651,9 @@ public:
 	/** Imports the given namespace into the editor. This may trigger a load event for additional macro and/or function library assets if not already loaded. */
 	void ImportNamespace(const FString& InNamespace);
 
+	/** Returns an instanced namespace helper utility object that corresponds to the given Blueprint. */
+	TSharedRef<FBlueprintNamespaceHelper> GetOrCreateNamespaceHelperForBlueprint(const UBlueprint* InBlueprint);
+
 protected:
 	UE_DEPRECATED(4.26, "Please do any validation inside the UBlueprint class during compilation, extra errors during compiling only supplied by the designer can lead to design time only errors being reported and being missed during cooks/content validation.")
 	virtual void AppendExtraCompilerResults(TSharedPtr<class IMessageLogListing> ResultsListing) {}
@@ -1275,6 +1279,9 @@ protected:
 	/** The toolbar builder class */
 	TSharedPtr<class FBlueprintEditorToolbar> Toolbar;
 
+	/** Cached set of instanced namespace helper objects */
+	TMap<TWeakObjectPtr<const UBlueprint>, TSharedRef<FBlueprintNamespaceHelper>> CachedNamespaceHelpers;
+
 	FOnSetPinVisibility OnSetPinVisibility;
 
 	/** Has someone requested a deferred update of the saved document state? */
diff --git a/Engine/Source/Editor/Kismet/Public/BlueprintNamespaceHelper.h b/Engine/Source/Editor/Kismet/Public/BlueprintNamespaceHelper.h
--- a/Engine/Source/Editor/Kismet/Public/BlueprintNamespaceHelper.h
+++ b/Engine/Source/Editor/Kismet/Public/BlueprintNamespaceHelper.h
@@ -33,6 +33,10 @@ public:
 
 	bool IsIncludedInNamespaceList(const FString& TestNamespace) const;
 
+	bool IsImportedType(const UField* InType) const;
+	bool IsImportedObject(const UObject* InObject) const;
+	bool IsImportedObject(const FSoftObjectPath& InObjectPath) const;
+
 private:
 	// Complete list of all fully-qualified namespace path identifiers for the associated Blueprint.
 	TSet<FString> FullyQualifiedListOfNamespaces;
diff --git a/Engine/Source/Editor/KismetCompiler/Private/KismetCompiler.cpp b/Engine/Source/Editor/KismetCompiler/Private/KismetCompiler.cpp
--- a/Engine/Source/Editor/KismetCompiler/Private/KismetCompiler.cpp
+++ b/Engine/Source/Editor/KismetCompiler/Private/KismetCompiler.cpp
@@ -2463,6 +2463,16 @@ void FKismetCompilerContext::FinishCompilingClass(UClass* Class)
 			Class->RemoveMetaData(NAME_DisplayName);
 		}
 
+		// Add the namespace (if set) to class metadata
+		if (!Blueprint->BlueprintNamespace.IsEmpty())
+		{
+			Class->SetMetaData(FBlueprintMetadata::MD_Namespace, *Blueprint->BlueprintNamespace);
+		}
+		else
+		{
+			Class->RemoveMetaData(FBlueprintMetadata::MD_Namespace);
+		}
+
 		// Copy the category info from the parent class
 #if WITH_EDITORONLY_DATA
 		
diff --git a/Engine/Source/Editor/UnrealEd/Classes/Settings/EditorExperimentalSettings.h b/Engine/Source/Editor/UnrealEd/Classes/Settings/EditorExperimentalSettings.h
--- a/Engine/Source/Editor/UnrealEd/Classes/Settings/EditorExperimentalSettings.h
+++ b/Engine/Source/Editor/UnrealEd/Classes/Settings/EditorExperimentalSettings.h
@@ -141,11 +141,6 @@ public:
 	UPROPERTY(EditAnywhere, config, Category = "Content Browser", meta = (DisplayName = "Enable support for long paths (> 260 characters)"))
 	bool bEnableLongPathsSupport;
 
-	/** Whether to enable namespace filtering features in the editor UI */
-	// @todo_namespaces - Make this a non-transient editable config property.
-	UPROPERTY(Transient/*, Category = UserInterface*/)
-	bool bEnableNamespaceFilteringFeatures;
-
 	/**
 	 * Returns an event delegate that is executed when a setting has changed.
 	 *
diff --git a/Engine/Source/Editor/UnrealEd/Private/Settings/SettingsClasses.cpp b/Engine/Source/Editor/UnrealEd/Private/Settings/SettingsClasses.cpp
--- a/Engine/Source/Editor/UnrealEd/Private/Settings/SettingsClasses.cpp
+++ b/Engine/Source/Editor/UnrealEd/Private/Settings/SettingsClasses.cpp
@@ -158,7 +158,6 @@ UEditorExperimentalSettings::UEditorExperimentalSettings( const FObjectInitializ
 	, bEnableLocalizationDashboard(true)
 	, bUseOpenCLForConvexHullDecomp(false)
 	, bAllowPotentiallyUnsafePropertyEditing(false)
-	, bEnableNamespaceFilteringFeatures(false)
 {
 }
 
