Move BPVariableDragDropAction.h to a public header so that the subobject editor can use it

#rnx

diff --git a/Engine/Source/Editor/Kismet/Private/BPVariableDragDropAction.h b/Engine/Source/Editor/Kismet/Private/BPVariableDragDropAction.h
deleted file mode 100644
diff --git a/Engine/Source/Editor/Kismet/Private/MyBlueprintItemDragDropAction.h b/Engine/Source/Editor/Kismet/Private/MyBlueprintItemDragDropAction.h
deleted file mode 100644
diff --git a/Engine/Source/Editor/Kismet/Public/BPVariableDragDropAction.h b/Engine/Source/Editor/Kismet/Public/BPVariableDragDropAction.h
new file mode 100644
--- /dev/null
+++ b/Engine/Source/Editor/Kismet/Public/BPVariableDragDropAction.h
@@ -0,0 +1,112 @@
+// Copyright Epic Games, Inc. All Rights Reserved.
+
+#pragma once
+
+#include "CoreMinimal.h"
+#include "UObject/Class.h"
+#include "UObject/UnrealType.h"
+#include "Input/DragAndDrop.h"
+#include "Input/Reply.h"
+#include "Engine/Blueprint.h"
+#include "BlueprintEditor.h"
+#include "GraphEditorDragDropAction.h"
+#include "MyBlueprintItemDragDropAction.h"
+
+class UEdGraph;
+
+/** DragDropAction class for dropping a Variable onto a graph */
+class KISMET_API FKismetVariableDragDropAction : public FMyBlueprintItemDragDropAction
+{
+public:
+	DRAG_DROP_OPERATOR_TYPE(FKismetVariableDragDropAction, FMyBlueprintItemDragDropAction)
+
+	// FGraphEditorDragDropAction interface
+	virtual void HoverTargetChanged() override;
+	virtual FReply DroppedOnPin(FVector2D ScreenPosition, FVector2D GraphPosition) override;
+	virtual FReply DroppedOnNode(FVector2D ScreenPosition, FVector2D GraphPosition) override;
+	virtual FReply DroppedOnPanel(const TSharedRef< class SWidget >& Panel, FVector2D ScreenPosition, FVector2D GraphPosition, UEdGraph& Graph) override;
+	// End of FGraphEditorDragDropAction
+
+	static TSharedRef<FKismetVariableDragDropAction> New(TSharedPtr<FEdGraphSchemaAction> InAction, FName InVariableName, UStruct* InVariableSource, FNodeCreationAnalytic AnalyticCallback)
+	{
+		TSharedRef<FKismetVariableDragDropAction> Operation = MakeShareable(new FKismetVariableDragDropAction);
+		Operation->VariableName = InVariableName;
+		Operation->VariableSource = InVariableSource;
+		Operation->AnalyticCallback = AnalyticCallback;
+		Operation->SourceAction = InAction;
+		Operation->Construct();
+		return Operation;
+	}
+
+	FProperty* GetVariableProperty()
+	{
+		if (VariableSource.IsValid() && VariableName != NAME_None)
+		{
+			return FindFProperty<FProperty>(VariableSource.Get(), VariableName);
+		}
+		return nullptr;
+	}
+
+protected:
+	 /** Construct a FKismetVariableDragDropAction */
+	FKismetVariableDragDropAction();
+
+	/** Structure for required node construction parameters */
+	struct FNodeConstructionParams
+	{
+		FVector2D GraphPosition;
+		UEdGraph* Graph;
+		FName VariableName;
+		TWeakObjectPtr<UStruct> VariableSource;
+	};
+
+	// FGraphSchemaActionDragDropAction interface
+	virtual void GetDefaultStatusSymbol(const FSlateBrush*& PrimaryBrushOut, FSlateColor& IconColorOut, FSlateBrush const*& SecondaryBrushOut, FSlateColor& SecondaryColorOut) const override;
+	// End of FGraphSchemaActionDragDropAction interface
+
+	// FMyBlueprintItemDragDropAction interface
+	virtual UBlueprint* GetSourceBlueprint() const override;
+	// End of FMyBlueprintItemDragDropAction interface
+
+	/** Called when user selects to create a Getter for the variable */
+	static void MakeGetter(FNodeConstructionParams InParams);
+	/** Called when user selects to create a Setter for the variable */
+	static void MakeSetter(FNodeConstructionParams InParams);
+	/** Called too check if we can execute a setter on a given property */
+	static bool CanExecuteMakeSetter(FNodeConstructionParams InParams, FProperty* InVariableProperty);
+
+	/**
+	 * Test new variable type against existing links for node and get any links that will break
+	 *
+	 * @param	Node						The node with existing links
+	 * @param	NewVariableProperty			the property for the new variable type 
+	 * @param	OutBroken						All of the links which are NOT compatible with the new type
+	 */
+	void GetLinksThatWillBreak(UEdGraphNode* Node, FProperty* NewVariableProperty, TArray<class UEdGraphPin*>& OutBroken);
+
+	/** Indicates if replacing the variable node, with the new property will require any links to be broken*/
+	bool WillBreakLinks( UEdGraphNode* Node, FProperty* NewVariableProperty ) 
+	{
+		TArray<class UEdGraphPin*> BadLinks;
+		GetLinksThatWillBreak(Node,NewVariableProperty,BadLinks);
+		return BadLinks.Num() > 0;
+	}
+
+	/**
+	 * Checks if the property can be dropped in a graph
+	 *
+	 * @param InVariableProperty		The variable property to check with
+	 * @param InGraph					The graph to check against placing the variable
+	 */
+	bool CanVariableBeDropped(const FProperty* InVariableProperty, const UEdGraph& InGraph) const;
+
+
+	/** Returns the local variable's scope, if any */
+	UStruct* GetLocalVariableScope() const;
+
+protected:
+	/** Name of variable being dragged */
+	FName VariableName;
+	/** Scope this variable belongs to */
+	TWeakObjectPtr<UStruct> VariableSource;
+};
diff --git a/Engine/Source/Editor/Kismet/Public/MyBlueprintItemDragDropAction.h b/Engine/Source/Editor/Kismet/Public/MyBlueprintItemDragDropAction.h
new file mode 100644
--- /dev/null
+++ b/Engine/Source/Editor/Kismet/Public/MyBlueprintItemDragDropAction.h
@@ -0,0 +1,60 @@
+// Copyright Epic Games, Inc. All Rights Reserved.
+
+#pragma once
+
+#include "CoreMinimal.h"
+#include "UObject/Class.h"
+#include "UObject/UnrealType.h"
+#include "Input/DragAndDrop.h"
+#include "Input/Reply.h"
+#include "Engine/Blueprint.h"
+#include "BlueprintEditor.h"
+#include "GraphEditorDragDropAction.h"
+
+class UEdGraph;
+struct FEdGraphSchemaAction;
+
+/** DragDropAction class for drag and dropping an item from the My Blueprints tree (e.g., variable or function) */
+class KISMET_API FMyBlueprintItemDragDropAction : public FGraphSchemaActionDragDropAction
+{
+public:
+	DRAG_DROP_OPERATOR_TYPE(FMyBlueprintItemDragDropAction, FGraphSchemaActionDragDropAction)
+
+	// FGraphEditorDragDropAction interface
+	virtual FReply DroppedOnAction(TSharedRef<FEdGraphSchemaAction> Action) override;
+	virtual FReply DroppedOnCategory(FText Category) override;
+	virtual void HoverTargetChanged() override;
+	// End of FGraphEditorDragDropAction
+
+	/** Set if operation is modified by alt */
+	void SetAltDrag(bool InIsAltDrag) {	bAltDrag = InIsAltDrag; }
+
+	/** Set if operation is modified by the ctrl key */
+	void SetCtrlDrag(bool InIsCtrlDrag) { bControlDrag = InIsCtrlDrag; }
+
+protected:
+	 /** Constructor */
+	FMyBlueprintItemDragDropAction();
+
+	virtual UBlueprint* GetSourceBlueprint() const
+	{
+		return nullptr;
+	}
+
+	/** Helper method to see if we're dragging in the same blueprint */
+	bool IsFromBlueprint(UBlueprint* InBlueprint) const
+	{
+		return GetSourceBlueprint() == InBlueprint;
+	}
+
+	void SetFeedbackMessageError(const FText& Message);
+	void SetFeedbackMessageOK(const FText& Message);
+
+protected:
+	/** Was ctrl held down at start of drag */
+	bool bControlDrag;
+	/** Was alt held down at the start of drag */
+	bool bAltDrag;
+	/** Analytic delegate to track node creation */
+	FNodeCreationAnalytic AnalyticCallback;
+};
