Some editor fixes for BPGC as asset of BP pacakge (cooked and uncooked):
- Edit action disabled on cooked packages
- Save action disabled on cooked packages
- Find parent blueprint in content browser works even if the parent is cooked (doesn't have access to the UBlueprint that generated it)
- Edit parent blueprint button in BP editor isn't available if the parent is cooked

diff --git a/Engine/Plugins/Editor/ContentBrowser/ContentBrowserAssetDataSource/Source/ContentBrowserAssetDataSource/Private/ContentBrowserAssetDataCore.cpp b/Engine/Plugins/Editor/ContentBrowser/ContentBrowserAssetDataSource/Source/ContentBrowserAssetDataSource/Private/ContentBrowserAssetDataCore.cpp
--- a/Engine/Plugins/Editor/ContentBrowser/ContentBrowserAssetDataSource/Source/ContentBrowserAssetDataSource/Private/ContentBrowserAssetDataCore.cpp
+++ b/Engine/Plugins/Editor/ContentBrowser/ContentBrowserAssetDataSource/Source/ContentBrowserAssetDataSource/Private/ContentBrowserAssetDataCore.cpp
@@ -202,7 +202,18 @@ bool CanEditItem(IAssetTools* InAssetTools, const UContentBrowserDataSource* InO
 
 bool CanEditAssetFileItem(IAssetTools* InAssetTools, const FContentBrowserAssetFileItemDataPayload& InAssetPayload, FText* OutErrorMsg)
 {
-	return CanModifyAssetFileItem(InAssetTools, InAssetPayload, OutErrorMsg);
+	if (!CanModifyAssetFileItem(InAssetTools, InAssetPayload, OutErrorMsg))
+	{
+		return false;
+	}
+
+	if (InAssetPayload.GetAssetData().PackageFlags & PKG_FilterEditorOnly)
+	{
+		SetOptionalErrorMessage(OutErrorMsg, LOCTEXT("Error_CannotEditCookedPackages", "Cannot edit cooked packages"));
+		return false;
+	}
+
+	return true;
 }
 
 bool CanPreviewItem(IAssetTools* InAssetTools, const UContentBrowserDataSource* InOwnerDataSource, const FContentBrowserItemData& InItem, FText* OutErrorMsg)
@@ -444,6 +455,12 @@ bool CanSaveAssetFileItem(IAssetTools* InAssetTools, const FContentBrowserAssetF
 		return false;
 	}
 
+	if (InAssetPayload.GetAssetData().PackageFlags & PKG_FilterEditorOnly)
+	{
+		SetOptionalErrorMessage(OutErrorMsg, LOCTEXT("Error_CannotSaveCookedPackages", "Cannot save cooked packages"));
+		return false;
+	}
+
 	if (EnumHasAnyFlags(InSaveFlags, EContentBrowserItemSaveFlags::SaveOnlyIfLoaded))
 	{
 		// Can't save a package that hasn't been loaded
diff --git a/Engine/Source/Editor/Kismet/Private/BlueprintEditor.cpp b/Engine/Source/Editor/Kismet/Private/BlueprintEditor.cpp
--- a/Engine/Source/Editor/Kismet/Private/BlueprintEditor.cpp
+++ b/Engine/Source/Editor/Kismet/Private/BlueprintEditor.cpp
@@ -2123,13 +2123,13 @@ void FBlueprintEditor::PostRegenerateMenusAndToolbars()
 			[
 				SNew(SButton)
 				.VAlign(VAlign_Center)
-				.ButtonStyle( FEditorStyle::Get(), "HoverHintOnly" )
-				.OnClicked( this, &FBlueprintEditor::OnFindParentClassInContentBrowserClicked )
-				.IsEnabled( this, &FBlueprintEditor::IsParentClassABlueprint )
-				.Visibility( this, &FBlueprintEditor::ParentClassButtonsVisibility )
-				.ToolTipText( LOCTEXT("FindParentInCBToolTip", "Find parent in Content Browser") )
+				.ButtonStyle(FEditorStyle::Get(), "HoverHintOnly")
+				.OnClicked(this, &FBlueprintEditor::OnFindParentClassInContentBrowserClicked)
+				.IsEnabled(this, &FBlueprintEditor::IsParentClassABlueprint)
+				.Visibility(this, &FBlueprintEditor::GetFindParentClassVisibility)
+				.ToolTipText(LOCTEXT("FindParentInCBToolTip", "Find parent in Content Browser"))
 				.ContentPadding(4.0f)
-				.ForegroundColor( FSlateColor::UseForeground() )
+				.ForegroundColor(FSlateColor::UseForeground())
 				[
 					SNew(SImage)
 					.Image(FEditorStyle::GetBrush("Icons.Search"))
@@ -2140,13 +2140,13 @@ void FBlueprintEditor::PostRegenerateMenusAndToolbars()
 			[
 				SNew(SButton)
 				.VAlign(VAlign_Center)
-				.ButtonStyle( FEditorStyle::Get(), "HoverHintOnly" )
-				.OnClicked( this, &FBlueprintEditor::OnEditParentClassClicked )
-				.IsEnabled( this, &FBlueprintEditor::IsParentClassABlueprint )
-				.Visibility( this, &FBlueprintEditor::ParentClassButtonsVisibility )
-				.ToolTipText( LOCTEXT("EditParentClassToolTip", "Open parent in editor") )
+				.ButtonStyle(FEditorStyle::Get(), "HoverHintOnly")
+				.OnClicked(this, &FBlueprintEditor::OnEditParentClassClicked)
+				.IsEnabled(this, &FBlueprintEditor::IsParentClassAnEditableBlueprint)
+				.Visibility(this, &FBlueprintEditor::GetEditParentClassVisibility)
+				.ToolTipText(LOCTEXT("EditParentClassToolTip", "Open parent in editor"))
 				.ContentPadding(4.0f)
-				.ForegroundColor( FSlateColor::UseForeground() )
+				.ForegroundColor(FSlateColor::UseForeground())
 				[
 					SNew(SImage)
 					.Image(FEditorStyle::GetBrush("Icons.Edit"))
@@ -2182,35 +2182,31 @@ FText FBlueprintEditor::GetParentClassNameText() const
 	return (ParentClass != NULL) ? ParentClass->GetDisplayNameText() : LOCTEXT("BlueprintEditor_NoParentClass", "None");
 }
 
-bool FBlueprintEditor::IsParentClassOfObjectABlueprint( const UBlueprint* Blueprint ) const
+bool FBlueprintEditor::IsParentClassOfObjectABlueprint(const UBlueprint* Blueprint) const
 {
-	if ( Blueprint != NULL )
-	{
-		UObject* ParentClass = Blueprint->ParentClass;
-		if ( ParentClass != NULL )
-		{
-			if ( ParentClass->IsA( UBlueprintGeneratedClass::StaticClass() ) )
-			{
-				return true;
-			}
-		}
-	}
-
-	return false;
+	return FBlueprintEditorUtils::IsParentClassABlueprint(Blueprint);
 }
 
 bool FBlueprintEditor::IsParentClassABlueprint() const
 {
-	const UBlueprint* Blueprint = GetBlueprintObj();
+	return FBlueprintEditorUtils::IsParentClassABlueprint(GetBlueprintObj());
+}
 
-	return IsParentClassOfObjectABlueprint( Blueprint );
+bool FBlueprintEditor::IsParentClassAnEditableBlueprint() const
+{
+	return FBlueprintEditorUtils::IsParentClassAnEditableBlueprint(GetBlueprintObj());
 }
 
-EVisibility FBlueprintEditor::ParentClassButtonsVisibility() const
+EVisibility FBlueprintEditor::GetFindParentClassVisibility() const
 {
 	return IsParentClassABlueprint() ? EVisibility::Visible : EVisibility::Collapsed;
 }
 
+EVisibility FBlueprintEditor::GetEditParentClassVisibility() const
+{
+	return IsParentClassAnEditableBlueprint() ? EVisibility::Visible : EVisibility::Collapsed;
+}
+
 bool FBlueprintEditor::IsParentClassNative() const
 {
 	const UBlueprint* Blueprint = GetBlueprintObj();
@@ -2263,20 +2259,24 @@ FText FBlueprintEditor::GetTextForNativeParentClassHeaderLink() const
 FReply FBlueprintEditor::OnFindParentClassInContentBrowserClicked()
 {
 	UBlueprint* Blueprint = GetBlueprintObj();
-	if ( Blueprint != NULL )
+	if (Blueprint)
 	{
 		UObject* ParentClass = Blueprint->ParentClass;
-		if ( ParentClass != NULL )
+		if (ParentClass)
 		{
 			UBlueprintGeneratedClass* ParentBlueprintGeneratedClass = Cast<UBlueprintGeneratedClass>( ParentClass );
-			if ( ParentBlueprintGeneratedClass != NULL )
+			if (ParentBlueprintGeneratedClass)
 			{
-				if ( ParentBlueprintGeneratedClass->ClassGeneratedBy != NULL )
+				TArray<UObject*> ParentObjectList;
+				if (ParentBlueprintGeneratedClass->ClassGeneratedBy)
 				{
-					TArray< UObject* > ParentObjectList;
-					ParentObjectList.Add( ParentBlueprintGeneratedClass->ClassGeneratedBy );
-					GEditor->SyncBrowserToObjects( ParentObjectList );
+					ParentObjectList.Add(ParentBlueprintGeneratedClass->ClassGeneratedBy);
 				}
+				else
+				{
+					ParentObjectList.Add(ParentBlueprintGeneratedClass);
+				}
+				GEditor->SyncBrowserToObjects(ParentObjectList);
 			}
 		}
 	}
@@ -2287,15 +2287,15 @@ FReply FBlueprintEditor::OnFindParentClassInContentBrowserClicked()
 FReply FBlueprintEditor::OnEditParentClassClicked()
 {
 	UBlueprint* Blueprint = GetBlueprintObj();
-	if ( Blueprint != NULL )
+	if (Blueprint)
 	{
 		UObject* ParentClass = Blueprint->ParentClass;
-		if ( ParentClass != NULL )
+		if (ParentClass)
 		{
 			UBlueprintGeneratedClass* ParentBlueprintGeneratedClass = Cast<UBlueprintGeneratedClass>( ParentClass );
-			if ( ParentBlueprintGeneratedClass != NULL )
+			if (ParentBlueprintGeneratedClass)
 			{
-				GEditor->GetEditorSubsystem<UAssetEditorSubsystem>()->OpenEditorForAsset( ParentBlueprintGeneratedClass->ClassGeneratedBy );
+				GEditor->GetEditorSubsystem<UAssetEditorSubsystem>()->OpenEditorForAsset(ParentBlueprintGeneratedClass->ClassGeneratedBy);
 			}
 		}
 	}
diff --git a/Engine/Source/Editor/Kismet/Public/BlueprintEditor.h b/Engine/Source/Editor/Kismet/Public/BlueprintEditor.h
--- a/Engine/Source/Editor/Kismet/Public/BlueprintEditor.h
+++ b/Engine/Source/Editor/Kismet/Public/BlueprintEditor.h
@@ -361,11 +361,15 @@ public:
 	virtual bool IsCompilingEnabled() const;
 
 	/** Returns true if the parent class is also a Blueprint */
+	UE_DEPRECATED(4.27, "Please use FBlueprintEditorUtils::IsParentClassABlueprint instead")
 	bool IsParentClassOfObjectABlueprint(const UBlueprint* Blueprint) const;
 
 	/** Returns true if the parent class of the Blueprint being edited is also a Blueprint */
 	bool IsParentClassABlueprint() const;
 
+	/** Returns true if the parent class of the Blueprint being edited is an editable Blueprint */
+	bool IsParentClassAnEditableBlueprint() const;
+
 	/** Returns true if the parent class of the Blueprint being edited is native */
 	bool IsParentClassNative() const;
 
@@ -682,8 +686,11 @@ protected:
 	void NavigateToChildGraph();
 	bool CanNavigateToChildGraph() const;
 
-	/** Determines visibility of the parent class manipulation buttons on the menu bar overlay */
-	EVisibility ParentClassButtonsVisibility() const;
+	/** Determines visibility of the find parent class in content browser button on the menu bar overlay */
+	EVisibility GetFindParentClassVisibility() const;
+
+	/** Determines visibility of the edit parent class button on the menu bar overlay */
+	EVisibility GetEditParentClassVisibility() const;
 
 	/** Recreates the overlay on the menu bar */
 	virtual void PostRegenerateMenusAndToolbars() override;
diff --git a/Engine/Source/Editor/UnrealEd/Private/Kismet2/BlueprintEditorUtils.cpp b/Engine/Source/Editor/UnrealEd/Private/Kismet2/BlueprintEditorUtils.cpp
--- a/Engine/Source/Editor/UnrealEd/Private/Kismet2/BlueprintEditorUtils.cpp
+++ b/Engine/Source/Editor/UnrealEd/Private/Kismet2/BlueprintEditorUtils.cpp
@@ -3297,6 +3297,41 @@ bool FBlueprintEditorUtils::IsLevelScriptBlueprint(const UBlueprint* Blueprint)
 	return (Blueprint && Blueprint->BlueprintType == BPTYPE_LevelScript);
 }
 
+bool FBlueprintEditorUtils::IsParentClassABlueprint(const UBlueprint* Blueprint)
+{
+	if (Blueprint)
+	{
+		UObject* ParentClass = Blueprint->ParentClass;
+		if (ParentClass)
+		{
+			if (ParentClass->IsA(UBlueprintGeneratedClass::StaticClass()))
+			{
+				return true;
+			}
+		}
+	}
+
+	return false;
+}
+
+bool FBlueprintEditorUtils::IsParentClassAnEditableBlueprint(const UBlueprint* Blueprint)
+{
+	if (Blueprint)
+	{
+		UObject* ParentClass = Blueprint->ParentClass;
+		if (ParentClass)
+		{
+			UBlueprintGeneratedClass* ParentBPGC = Cast<UBlueprintGeneratedClass>(ParentClass);
+			if (ParentBPGC && ParentBPGC->ClassGeneratedBy)
+			{
+				return true;
+			}
+		}
+	}
+
+	return false;
+}
+
 bool FBlueprintEditorUtils::IsAnonymousBlueprintClass(const UClass* Class)
 {
 	return (Class && Class->GetOutermost()->ContainsMap());
diff --git a/Engine/Source/Editor/UnrealEd/Public/Kismet2/BlueprintEditorUtils.h b/Engine/Source/Editor/UnrealEd/Public/Kismet2/BlueprintEditorUtils.h
--- a/Engine/Source/Editor/UnrealEd/Public/Kismet2/BlueprintEditorUtils.h
+++ b/Engine/Source/Editor/UnrealEd/Public/Kismet2/BlueprintEditorUtils.h
@@ -597,6 +597,12 @@ public:
 	 */
 	static bool IsLevelScriptBlueprint(const UBlueprint* Blueprint);
 
+	/** Returns whether the parent class of the specified blueprint is also a blueprint */
+	static bool IsParentClassABlueprint(const UBlueprint* Blueprint);
+
+	/** Returns whether the parent class of the specified blueprint is an editable blueprint */
+	static bool IsParentClassAnEditableBlueprint(const UBlueprint* Blueprint);
+
 	/**
 	 * Whether or not this class represents a class generated by an anonymous actor class stored in a level 
 	 *
