Backout problematic CL 14160476 which broke default value propegation when changing variable types and caused cooking errors.

#rnx

diff --git a/Engine/Source/Editor/Kismet/Private/BlueprintEditor.cpp b/Engine/Source/Editor/Kismet/Private/BlueprintEditor.cpp
--- a/Engine/Source/Editor/Kismet/Private/BlueprintEditor.cpp
+++ b/Engine/Source/Editor/Kismet/Private/BlueprintEditor.cpp
@@ -8884,20 +8884,6 @@ void FBlueprintEditor::OnFinishedChangingProperties(const FPropertyChangedEvent&
 	FName PropertyName = PropertyChangedEvent.GetPropertyName();
 	FName MemberPropertyName = PropertyChangedEvent.MemberProperty ? PropertyChangedEvent.MemberProperty->GetFName() : PropertyName;
 
-	if (MemberPropertyName.IsValid())
-	{
-		UBlueprint* Blueprint = GetBlueprintObj();
-		if (Blueprint)
-		{
-			int32 VarIndex = FBlueprintEditorUtils::FindNewVariableIndex(Blueprint, MemberPropertyName);
-			if (VarIndex != INDEX_NONE)
-			{
-				Blueprint->Modify();
-				Blueprint->NewVariables[VarIndex].DefaultValue.Empty();
-			}
-		}
-	}
-
 	//@TODO: This code does not belong here (might not even be necessary anymore as they seem to have PostEditChangeProperty impls now)!
 	if ((PropertyName == GET_MEMBER_NAME_CHECKED(UK2Node_Switch, bHasDefaultPin)) ||
 		(PropertyName == GET_MEMBER_NAME_CHECKED(UK2Node_SwitchInteger, StartIndex)) ||
diff --git a/Engine/Source/Editor/Kismet/Private/SMyBlueprint.cpp b/Engine/Source/Editor/Kismet/Private/SMyBlueprint.cpp
--- a/Engine/Source/Editor/Kismet/Private/SMyBlueprint.cpp
+++ b/Engine/Source/Editor/Kismet/Private/SMyBlueprint.cpp
@@ -3205,7 +3205,7 @@ void SMyBlueprint::OnCopy()
 			{
 				// Grab the address of where the property is actually stored (UObject* base, plus the offset defined in the property)
 				void* OldPropertyAddr = TargetProperty->ContainerPtrToValuePtr<void>(GeneratedCDO);
-				if (OldPropertyAddr && Description.DefaultValue.IsEmpty())
+				if (OldPropertyAddr)
 				{
 					TargetProperty->ExportTextItem(Description.DefaultValue, OldPropertyAddr, OldPropertyAddr, nullptr, PPF_SerializedAsImportText);
 				}
diff --git a/Engine/Source/Editor/KismetCompiler/Private/KismetCompiler.cpp b/Engine/Source/Editor/KismetCompiler/Private/KismetCompiler.cpp
--- a/Engine/Source/Editor/KismetCompiler/Private/KismetCompiler.cpp
+++ b/Engine/Source/Editor/KismetCompiler/Private/KismetCompiler.cpp
@@ -758,6 +758,12 @@ void FKismetCompilerContext::CreateClassVariablesFromBlueprint()
 			if(!Variable.DefaultValue.IsEmpty())
 			{
 				SetPropertyDefaultValue(NewProperty, Variable.DefaultValue);
+
+				// We're copying the value to the real CDO, so clear the version stored in the blueprint editor data
+				if (CompileOptions.CompileType == EKismetCompileType::Full)
+				{
+					Variable.DefaultValue.Empty();
+				}
 			}
 
 			if (NewProperty->HasAnyPropertyFlags(CPF_Net))
