Weakend a check about broken (null) subgraphs to an ensure with a message of what the containing graph is to get more information

diff --git a/Engine/Source/Editor/Kismet/Private/SMyBlueprint.cpp b/Engine/Source/Editor/Kismet/Private/SMyBlueprint.cpp
--- a/Engine/Source/Editor/Kismet/Private/SMyBlueprint.cpp
+++ b/Engine/Source/Editor/Kismet/Private/SMyBlueprint.cpp
@@ -1063,7 +1063,11 @@ void SMyBlueprint::GetChildGraphs(UEdGraph* InEdGraph, int32 const SectionId, FG
 	// Grab children graphs
 	for (UEdGraph* Graph : InEdGraph->SubGraphs)
 	{
-		check(Graph);
+		if (Graph == nullptr)
+		{
+			ensureMsgf(Graph != nullptr, TEXT("A subgraph of %s was null"), *GetPathNameSafe(InEdGraph));
+			continue;
+		}
 
 		FGraphDisplayInfo ChildGraphDisplayInfo;
 		if (const UEdGraphSchema* ChildSchema = Graph->GetSchema())
