Correctly display tooltips for properties with native MD_DisplayName set

#rnx

diff --git a/Engine/Source/Editor/Kismet/Private/SBlueprintPalette.cpp b/Engine/Source/Editor/Kismet/Private/SBlueprintPalette.cpp
--- a/Engine/Source/Editor/Kismet/Private/SBlueprintPalette.cpp
+++ b/Engine/Source/Editor/Kismet/Private/SBlueprintPalette.cpp
@@ -1689,9 +1689,18 @@ FText SBlueprintPaletteItem::GetToolTipText() const
 			}
 			else
 			{
-				FString Result = GetVarTooltip(Blueprint, VarClass, VarAction->GetVariableName());
-				// Only use the variable tooltip if it has been filled out.
-				ToolTipText = FText::FromString( !Result.IsEmpty() ? Result : GetVarType(VarClass, VarAction->GetVariableName(), true, true) );
+				// Use the native display name metadata if we can
+				const FProperty* Property = FindFProperty<FProperty>(VarClass, VarAction->GetVariableName());
+				if (Property && Property->IsNative())
+				{
+					ToolTipText = Property->GetDisplayNameText();
+				}
+				else
+				{	
+					FString Result = GetVarTooltip(Blueprint, VarClass, VarAction->GetVariableName());
+					// Only use the variable tooltip if it has been filled out.
+					ToolTipText = FText::FromString( !Result.IsEmpty() ? Result : GetVarType(VarClass, VarAction->GetVariableName(), true, true) );	
+				}
 			}
 		}
 		else if (PaletteAction->GetTypeId() == FEdGraphSchemaAction_K2LocalVar::StaticGetTypeId())
