Fixed missing icon in struct editor

diff --git a/Engine/Source/Editor/EditorStyle/Private/StarshipStyle.cpp b/Engine/Source/Editor/EditorStyle/Private/StarshipStyle.cpp
--- a/Engine/Source/Editor/EditorStyle/Private/StarshipStyle.cpp
+++ b/Engine/Source/Editor/EditorStyle/Private/StarshipStyle.cpp
@@ -5188,10 +5188,11 @@ void FStarshipEditorStyle::FStyle::SetupPersonaStyle()
 		Set("Blueprint.CompileStatus.Background", new IMAGE_BRUSH_SVG("Starship/Blueprints/CompileStatus_Background", Icon20x20));
 
 		// @todo - Icon Replacement - trying out tinting compile backgrounds
+/*
 		Set("Blueprint.CompileStatus.Background.Unknown", new IMAGE_BRUSH_SVG("Starship/Blueprints/CompileStatus_Background", Icon20x20, FStyleColors::AccentYellow));
 		Set("Blueprint.CompileStatus.Background.Warning", new IMAGE_BRUSH_SVG("Starship/Blueprints/CompileStatus_Background", Icon20x20, FStyleColors::Warning));
 		Set("Blueprint.CompileStatus.Background.Good", new IMAGE_BRUSH_SVG("Starship/Blueprints/CompileStatus_Background", Icon20x20, FStyleColors::AccentGreen));
-		Set("Blueprint.CompileStatus.Background.Error", new IMAGE_BRUSH_SVG("Starship/Blueprints/CompileStatus_Background", Icon20x20, FStyleColors::Error));
+		Set("Blueprint.CompileStatus.Background.Error", new IMAGE_BRUSH_SVG("Starship/Blueprints/CompileStatus_Background", Icon20x20, FStyleColors::Error));*/
 
 		Set("Blueprint.CompileStatus.Overlay.Unknown", new IMAGE_BRUSH_SVG("Starship/Blueprints/CompileStatus_Unknown_Badge", Icon20x20, FStyleColors::AccentYellow));
 		Set("Blueprint.CompileStatus.Overlay.Warning", new IMAGE_BRUSH_SVG("Starship/Blueprints/CompileStatus_Warning_Badge", Icon20x20, FStyleColors::Warning));
diff --git a/Engine/Source/Editor/Kismet/Private/UserDefinedStructureEditor.cpp b/Engine/Source/Editor/Kismet/Private/UserDefinedStructureEditor.cpp
--- a/Engine/Source/Editor/Kismet/Private/UserDefinedStructureEditor.cpp
+++ b/Engine/Source/Editor/Kismet/Private/UserDefinedStructureEditor.cpp
@@ -29,6 +29,7 @@
 #include "SPinTypeSelector.h"
 #include "Kismet2/StructureEditorUtils.h"
 #include "Widgets/Docking/SDockTab.h"
+#include "Widgets/Images/SLayeredImage.h"
 
 #define LOCTEXT_NAMESPACE "StructureEditor"
 
@@ -469,16 +470,16 @@ public:
 				switch(Struct->Status.GetValue())
 				{
 				case EUserDefinedStructureStatus::UDSS_Error:
-					return FEditorStyle::GetBrush("Kismet.Status.Error.Small");
+					return FEditorStyle::GetBrush("Blueprint.CompileStatus.Overlay.Error");
 				case EUserDefinedStructureStatus::UDSS_UpToDate:
-					return FEditorStyle::GetBrush("Kismet.Status.Good.Small");
+					return FEditorStyle::GetBrush("Blueprint.CompileStatus.Overlay.Good");
 				default:
-					return FEditorStyle::GetBrush("Kismet.Status.Unknown.Small");
+					return FEditorStyle::GetBrush("Blueprint.CompileStatus.Overlay.Unknown");
 				}
 				
 			}
 		}
-		return NULL;
+		return nullptr;
 	}
 
 	FText GetStatusTooltip() const
@@ -1067,6 +1068,7 @@ void FUserDefinedStructureLayout::GenerateChildContent( IDetailChildrenBuilder&
 	const float NameWidth = 80.0f;
 	const float ContentWidth = 130.0f;
 
+	TSharedPtr<SLayeredImage> StatusImage;
 	ChildrenBuilder.AddCustomRow(FText::GetEmpty())
 	[
 		SNew(SHorizontalBox)
@@ -1075,9 +1077,10 @@ void FUserDefinedStructureLayout::GenerateChildContent( IDetailChildrenBuilder&
 		.HAlign(HAlign_Left)
 		.VAlign(VAlign_Center)
 		[
-			SNew(SImage)
-			.Image(this, &FUserDefinedStructureLayout::OnGetStructureStatus)
+			SAssignNew(StatusImage, SLayeredImage)
+			.Image(FAppStyle::Get().GetBrush("Blueprint.CompileStatus.Background"))
 			.ToolTipText(this, &FUserDefinedStructureLayout::GetStatusTooltip)
+			.DesiredSizeOverride(FVector2D(16,16))
 		]
 		+ SHorizontalBox::Slot()
 		.AutoWidth()
@@ -1120,6 +1123,8 @@ void FUserDefinedStructureLayout::GenerateChildContent( IDetailChildrenBuilder&
 		]
 	];
 
+	StatusImage->AddLayer(MakeAttributeSP(this, &FUserDefinedStructureLayout::OnGetStructureStatus));
+
 	auto StructureDetailsSP = StructureDetails.Pin();
 	if(StructureDetailsSP.IsValid())
 	{
