Fix static analysis warning about temp TObjectPtr in SSCSEditor.cpp

diff --git a/Engine/Source/Editor/Kismet/Private/SSCSEditor.cpp b/Engine/Source/Editor/Kismet/Private/SSCSEditor.cpp
--- a/Engine/Source/Editor/Kismet/Private/SSCSEditor.cpp
+++ b/Engine/Source/Editor/Kismet/Private/SSCSEditor.cpp
@@ -4776,7 +4776,7 @@ FSCSEditorTreeNodePtrType SSCSEditor::GetNodeFromActorComponent(const UActorComp
 
 					for(int32 StackIndex = ParentBPStack.Num() - 1; StackIndex >= 0; --StackIndex)
 					{
-						USimpleConstructionScript* ParentSCS = ParentBPStack[StackIndex] ? ParentBPStack[StackIndex]->SimpleConstructionScript : nullptr;
+						USimpleConstructionScript* ParentSCS = ParentBPStack[StackIndex] ? ParentBPStack[StackIndex]->SimpleConstructionScript.Get() : nullptr;
 						if (ParentSCS)
 						{
 							// Attempt to locate an SCS node with a variable name that matches the name of the given component
@@ -4976,7 +4976,7 @@ void SSCSEditor::BuildSubTreeForActorNode(FSCSEditorActorNodePtrType InActorNode
 		// Add the full SCS tree node hierarchy (including SCS nodes inherited from parent blueprints)
 		for (int32 StackIndex = ParentBPStack.Num() - 1; StackIndex >= 0; --StackIndex)
 		{
-			USimpleConstructionScript* ParentSCS = ParentBPStack[StackIndex] ? ParentBPStack[StackIndex]->SimpleConstructionScript : nullptr;
+			USimpleConstructionScript* ParentSCS = ParentBPStack[StackIndex] ? ParentBPStack[StackIndex]->SimpleConstructionScript.Get() : nullptr;
 			if (ParentSCS)
 			{
 				for (USCS_Node* SCS_Node : ParentSCS->GetRootNodes())
