Fixes for PageDown/PageUp navigation in Blueprint editor:
- Keep focus on the blueprint widget after going into a child blueprint , so PageUp key will work with no clicking
- reimplement child blueprint menu as an actual menu so that keyboard navigation will work (bonus now there is search, and it looks nicer)

diff --git a/Engine/Source/Editor/Kismet/Private/BlueprintEditor.cpp b/Engine/Source/Editor/Kismet/Private/BlueprintEditor.cpp
--- a/Engine/Source/Editor/Kismet/Private/BlueprintEditor.cpp
+++ b/Engine/Source/Editor/Kismet/Private/BlueprintEditor.cpp
@@ -219,68 +219,6 @@ TSharedRef<SDockTab> FSelectionDetailsSummoner::SpawnTab(const FWorkflowTabSpawn
 	return Tab;
 }
 
-/////////////////////////////////////////////////////
-// SChildGraphPicker
-
-class SChildGraphPicker : public SCompoundWidget
-{
-public:
-	SLATE_BEGIN_ARGS(SChildGraphPicker) {}
-	SLATE_END_ARGS()
-public:
-	void Construct(const FArguments& InArgs, UEdGraph* ParentGraph)
-	{
-		this->ChildSlot
-		[
-			SNew(SBorder)
-			.BorderImage(FEditorStyle::GetBrush("Menu.Background"))
-			.Padding(5)
-			.ToolTipText(LOCTEXT("ChildGraphPickerTooltip", "Pick the graph to enter"))
-			[
-				SNew(SVerticalBox)
-				+SVerticalBox::Slot()
-				.AutoHeight()
-				[
-					SNew(STextBlock)
-					.Text(LOCTEXT("ChildGraphPickerDesc", "Navigate to graph"))
-					.Font( FEditorStyle::GetFontStyle(TEXT("Kismet.GraphPicker.Title.Font")) )
-				]
-				+SVerticalBox::Slot()
-				.AutoHeight()
-				[
-					SNew(SListView<UEdGraph*>)
-					.ItemHeight(20)
-					.ListItemsSource(&ToRawPtrTArrayUnsafe(ParentGraph->SubGraphs))
-					.OnGenerateRow( this, &SChildGraphPicker::GenerateMenuItemRow )
-				]
-			]
-		];	
-	}
-private:
-	/** Generate a row for the InItem in the combo box's list (passed in as OwnerTable). Do this by calling the user-specified OnGenerateWidget */
-	TSharedRef<ITableRow> GenerateMenuItemRow(UEdGraph* InItem, const TSharedRef<STableViewBase>& OwnerTable)
-	{
-		check(InItem != NULL);
-
-		const FText DisplayName = FLocalKismetCallbacks::GetGraphDisplayName(InItem);
-
-		return SNew( STableRow<UEdGraph*>, OwnerTable )
-			[
-				SNew(SHyperlink)
-					.Text(DisplayName)
-					.Style(FEditorStyle::Get(), "HoverOnlyHyperlink")
-					.OnNavigate(this, &SChildGraphPicker::NavigateToGraph, InItem)
-			]
-		;
-	}
-
-	void NavigateToGraph(UEdGraph* ChildGraph)
-	{
-		FKismetEditorUtilities::BringKismetToFocusAttentionOnObject(ChildGraph);
-
-		FSlateApplication::Get().DismissAllMenus();
-	}
-};
 
 /////////////////////////////////////////////////////
 // FBlueprintEditor
@@ -3158,10 +3096,32 @@ void FBlueprintEditor::NavigateToChildGraph()
 		if (CurrentGraph->SubGraphs.Num() > 1)
 		{
 			// Display a child jump list
+			FMenuBuilder MenuBuilder(true, nullptr);
+			MenuBuilder.BeginSection("NavigateToGraph", LOCTEXT("ChildGraphPickerDesc", "Navigate to graph"));
+
+			TArray<TObjectPtr<UEdGraph>> SortedSubGraphs = CurrentGraph->SubGraphs;
+			SortedSubGraphs.Sort([](UEdGraph& A, UEdGraph& B) { return FLocalKismetCallbacks::GetGraphDisplayName(&A).CompareToCaseIgnored(FLocalKismetCallbacks::GetGraphDisplayName(&B)) < 0; });
+
+			for (TObjectPtr<UEdGraph> ChildGraph : SortedSubGraphs)
+			{
+				MenuBuilder.AddMenuEntry(
+					FLocalKismetCallbacks::GetGraphDisplayName(ChildGraph),
+					LOCTEXT("ChildGraphPickerTooltip", "Pick the graph to enter"),
+					FSlateIcon(),
+					FUIAction(
+						FExecuteAction::CreateLambda([this,ChildGraph]() {
+							 OpenDocument(ChildGraph, FDocumentTracker::NavigatingCurrentDocument);
+							 // focus blueprint widget so that keyboard inputs (navigate to child/parent) will work before clicking
+							 FSlateApplication::Get().SetKeyboardFocus(GetMyBlueprintWidget());
+							  }),
+						FCanExecuteAction()));
+			}
+			MenuBuilder.EndSection();
+
 			FSlateApplication::Get().PushMenu( 
 				GetToolkitHost()->GetParentWidget(),
 				FWidgetPath(),
-				SNew(SChildGraphPicker, CurrentGraph),
+				MenuBuilder.MakeWidget(),
 				FSlateApplication::Get().GetCursorPos(), // summon location
 				FPopupTransitionEffect( FPopupTransitionEffect::TypeInPopup )
 			);
