Change UE4 to UE

diff --git a/Engine/Source/Editor/Kismet/Private/DiffUtils.cpp b/Engine/Source/Editor/Kismet/Private/DiffUtils.cpp
--- a/Engine/Source/Editor/Kismet/Private/DiffUtils.cpp
+++ b/Engine/Source/Editor/Kismet/Private/DiffUtils.cpp
@@ -14,7 +14,7 @@
 #include "ObjectEditorUtils.h"
 #include "Components/ActorComponent.h"
 
-namespace UE4DiffUtils_Private
+namespace UEDiffUtils_Private
 {
 	FProperty* Resolve( const UStruct* Class, FName PropertyName )
 	{
@@ -56,7 +56,7 @@ FResolvedProperty FPropertySoftPath::Resolve(const UStruct* Struct, const void*
 	for (int32 i = 0; i < PropertyChain.Num(); ++i)
 	{
 		CurrentBlock = NextBlock;
-		const FProperty* NextProperty = UE4DiffUtils_Private::Resolve(NextClass, PropertyChain[i].PropertyName);
+		const FProperty* NextProperty = UEDiffUtils_Private::Resolve(NextClass, PropertyChain[i].PropertyName);
 		if (NextProperty)
 		{
 			Property = NextProperty;
@@ -129,7 +129,7 @@ FPropertyPath FPropertySoftPath::ResolvePath(const UObject* Object) const
 	for( int32 I = 0; I < PropertyChain.Num(); ++I )
 	{
 		FName PropertyIdentifier = PropertyChain[I].PropertyName;
-		FProperty* ResolvedProperty = UE4DiffUtils_Private::Resolve(ContainerStruct, PropertyIdentifier);
+		FProperty* ResolvedProperty = UEDiffUtils_Private::Resolve(ContainerStruct, PropertyIdentifier);
 
 		FPropertyInfo Info(ResolvedProperty, INDEX_NONE);
 		Ret.AddProperty(Info);
@@ -287,8 +287,8 @@ const UObject* DiffUtils::GetCDO(const UBlueprint* ForBlueprint)
 
 void DiffUtils::CompareUnrelatedStructs(const UStruct* StructA, const void* A, const UStruct* StructB, const void* B, TArray<FSingleObjectDiffEntry>& OutDifferingProperties)
 {
-	FPropertySoftPathSet PropertiesInA = UE4DiffUtils_Private::GetPropertyNameSet(StructA);
-	FPropertySoftPathSet PropertiesInB = UE4DiffUtils_Private::GetPropertyNameSet(StructB);
+	FPropertySoftPathSet PropertiesInA = UEDiffUtils_Private::GetPropertyNameSet(StructA);
+	FPropertySoftPathSet PropertiesInB = UEDiffUtils_Private::GetPropertyNameSet(StructB);
 
 	// any properties in A that aren't in B are differing:
 	auto AddedToA = PropertiesInA.Difference(PropertiesInB).Array();
