Add the blueprint debugger to the "Debug" tab in blueprints

diff --git a/Engine/Source/Editor/Kismet/Private/BlueprintDebugger.cpp b/Engine/Source/Editor/Kismet/Private/BlueprintDebugger.cpp
--- a/Engine/Source/Editor/Kismet/Private/BlueprintDebugger.cpp
+++ b/Engine/Source/Editor/Kismet/Private/BlueprintDebugger.cpp
@@ -14,6 +14,7 @@
 #include "Widgets/Docking/SDockTab.h"
 #include "WorkspaceMenuStructure.h"
 #include "WorkspaceMenuStructureModule.h"
+#include "BlueprintEditorTabs.h"
 
 #define LOCTEXT_NAMESPACE "BlueprintDebugger"
 
@@ -64,14 +65,12 @@ private:
 	FBlueprintDebuggerImpl& operator=(FBlueprintDebuggerImpl&&);
 };
 
-const FName DebuggerAppName = FName(TEXT("DebuggerApp"));
-
 FBlueprintDebuggerImpl::FBlueprintDebuggerImpl()
 {
 	const IWorkspaceMenuStructure& MenuStructure = WorkspaceMenu::GetMenuStructure();
 
 	FBlueprintDebuggerCommands::Register();
-	FGlobalTabmanager::Get()->RegisterNomadTabSpawner(DebuggerAppName, FOnSpawnTab::CreateRaw(this, &FBlueprintDebuggerImpl::CreateBluprintDebuggerTab))
+	FGlobalTabmanager::Get()->RegisterNomadTabSpawner(FBlueprintEditorTabs::BlueprintDebuggerID, FOnSpawnTab::CreateRaw(this, &FBlueprintDebuggerImpl::CreateBluprintDebuggerTab))
 		.SetDisplayName(NSLOCTEXT("BlueprintDebugger", "TabTitle", "Blueprint Debugger"))
 		.SetTooltipText(NSLOCTEXT("BlueprintDebugger", "TooltipText", "Open the Blueprint Debugger tab."))
 		.SetGroup(MenuStructure.GetDeveloperToolsDebugCategory())
@@ -82,14 +81,14 @@ FBlueprintDebuggerImpl::~FBlueprintDebuggerImpl()
 {
 	if (FSlateApplication::IsInitialized())
 	{
-		FGlobalTabmanager::Get()->UnregisterNomadTabSpawner(DebuggerAppName);
+		FGlobalTabmanager::Get()->UnregisterNomadTabSpawner(FBlueprintEditorTabs::BlueprintDebuggerID);
 	}
 
 	const IWorkspaceMenuStructure& MenuStructure = WorkspaceMenu::GetMenuStructure();
 
 	if (DebuggingToolsTabManager.IsValid())
 	{
-		FGlobalTabmanager::Get()->UnregisterTabSpawner(DebuggerAppName);
+		FGlobalTabmanager::Get()->UnregisterTabSpawner(FBlueprintEditorTabs::BlueprintDebuggerID);
 		BlueprintDebuggerLayout = TSharedPtr<FTabManager::FLayout>();
 		DebuggingToolsTabManager = TSharedPtr<FTabManager>();
 	}
diff --git a/Engine/Source/Editor/Kismet/Private/BlueprintEditor.cpp b/Engine/Source/Editor/Kismet/Private/BlueprintEditor.cpp
--- a/Engine/Source/Editor/Kismet/Private/BlueprintEditor.cpp
+++ b/Engine/Source/Editor/Kismet/Private/BlueprintEditor.cpp
@@ -2694,6 +2694,11 @@ void FBlueprintEditor::CreateDefaultCommands()
 		FCanExecuteAction::CreateSP(this, &FBlueprintEditor::HasAnyWatches)
 		);
 
+	ToolkitCommands->MapAction( FBlueprintEditorCommands::Get().OpenBlueprintDebugger,
+	    FExecuteAction::CreateSP(this, &FBlueprintEditor::OpenBlueprintDebugger),
+	    FCanExecuteAction::CreateSP(this, &FBlueprintEditor::CanOpenBlueprintDebugger)
+    );
+
 	// New document actions
 	ToolkitCommands->MapAction( FBlueprintEditorCommands::Get().AddNewVariable,
 		FExecuteAction::CreateSP(this, &FBlueprintEditor::OnAddNewVariable),
@@ -3664,6 +3669,17 @@ void FBlueprintEditor::ClearAllWatches()
 	FDebuggingActionCallbacks::ClearWatches(GetBlueprintObj());
 }
 
+void FBlueprintEditor::OpenBlueprintDebugger()
+{
+	FGlobalTabmanager::Get()->TryInvokeTab(FBlueprintEditorTabs::BlueprintDebuggerID);
+}
+
+bool FBlueprintEditor::CanOpenBlueprintDebugger() const
+{
+	// The BP debugger can always be spawned because it will get updated on PIE
+	return true;
+}
+
 bool FBlueprintEditor::HasAnyBreakpoints() const
 {
 	return GetBlueprintObj() && GetBlueprintObj()->Breakpoints.Num() > 0;
diff --git a/Engine/Source/Editor/Kismet/Private/BlueprintEditorCommands.cpp b/Engine/Source/Editor/Kismet/Private/BlueprintEditorCommands.cpp
--- a/Engine/Source/Editor/Kismet/Private/BlueprintEditorCommands.cpp
+++ b/Engine/Source/Editor/Kismet/Private/BlueprintEditorCommands.cpp
@@ -57,6 +57,7 @@ void FBlueprintEditorCommands::RegisterCommands()
 	UI_COMMAND( DisableAllBreakpoints, "Disable All Breakpoints", "Disable all breakpoints", EUserInterfaceActionType::Button, FInputChord() );
 	UI_COMMAND( ClearAllBreakpoints, "Delete All Breakpoints", "Delete all breakpoints", EUserInterfaceActionType::Button, FInputChord(EModifierKey::Control | EModifierKey::Shift, EKeys::F9) );
 	UI_COMMAND(	ClearAllWatches, "Delete All Watches", "Delete all watches", EUserInterfaceActionType::Button, FInputChord() );
+	UI_COMMAND( OpenBlueprintDebugger, "Blueprint Debugger", "Open the Blueprint Debugger tab.", EUserInterfaceActionType::Button, FInputChord() );
 
 	// New documents
 	UI_COMMAND( AddNewVariable, "Variable", "Adds a new variable to this blueprint.", EUserInterfaceActionType::Button, FInputChord() );
diff --git a/Engine/Source/Editor/Kismet/Private/BlueprintEditorCommands.h b/Engine/Source/Editor/Kismet/Private/BlueprintEditorCommands.h
--- a/Engine/Source/Editor/Kismet/Private/BlueprintEditorCommands.h
+++ b/Engine/Source/Editor/Kismet/Private/BlueprintEditorCommands.h
@@ -58,7 +58,8 @@ public:
 	TSharedPtr< FUICommandInfo > DisableAllBreakpoints;
 	TSharedPtr< FUICommandInfo > ClearAllBreakpoints;
 	TSharedPtr< FUICommandInfo > ClearAllWatches;
-
+	TSharedPtr< FUICommandInfo > OpenBlueprintDebugger;
+	
 	// New documents
 	TSharedPtr< FUICommandInfo > AddNewVariable;
 	TSharedPtr< FUICommandInfo > AddNewLocalVariable;
diff --git a/Engine/Source/Editor/Kismet/Private/BlueprintEditorTabs.cpp b/Engine/Source/Editor/Kismet/Private/BlueprintEditorTabs.cpp
--- a/Engine/Source/Editor/Kismet/Private/BlueprintEditorTabs.cpp
+++ b/Engine/Source/Editor/Kismet/Private/BlueprintEditorTabs.cpp
@@ -13,6 +13,7 @@ const FName FBlueprintEditorTabs::ConstructionScriptEditorID( TEXT( "Constructio
 const FName FBlueprintEditorTabs::SCSViewportID( TEXT( "SCSViewport" ) );
 const FName FBlueprintEditorTabs::MyBlueprintID( TEXT( "MyBlueprint" ) );
 const FName FBlueprintEditorTabs::ReplaceNodeReferencesID( TEXT( "ReplaceNodeReferences" ) );
+const FName FBlueprintEditorTabs::BlueprintDebuggerID( TEXT( "BlueprintDebugger" ) );
 
 const FName FBlueprintEditorTabs::GraphEditorID( TEXT( "GraphEditor" ) );
 const FName FBlueprintEditorTabs::TimelineEditorID( TEXT( "TimelineEditor" ) );
diff --git a/Engine/Source/Editor/Kismet/Private/SBlueprintEditorToolbar.cpp b/Engine/Source/Editor/Kismet/Private/SBlueprintEditorToolbar.cpp
--- a/Engine/Source/Editor/Kismet/Private/SBlueprintEditorToolbar.cpp
+++ b/Engine/Source/Editor/Kismet/Private/SBlueprintEditorToolbar.cpp
@@ -164,6 +164,16 @@ void FKismet2Menu::FillViewMenu(UToolMenu* InMenu)
 
 void FKismet2Menu::FillDebugMenu(UToolMenu* InMenu)
 {
+	{
+		FToolMenuSection& Section = InMenu->AddSection("BlueprintDebugger", LOCTEXT("BlueprintDebugger_Heading", "Blueprint Debugger"));
+		Section.AddMenuEntry(
+			FBlueprintEditorCommands::Get().OpenBlueprintDebugger,
+			/* InLabelOverride = */ LOCTEXT("BpDebuggerTitle", "Blueprint Debugger"),
+			/* InTooltipOverride = */ LOCTEXT("BpDebuggerTooltip","Open the Blueprint Debugger."),
+			FSlateIcon(FEditorStyle::GetStyleSetName(), "BlueprintDebugger.TabIcon")
+		);
+	}
+	
 	{
 		FToolMenuSection& Section = InMenu->AddSection("DebugBreakpoints", LOCTEXT("DebugMenu_BreakpointHeading", "Breakpoints"));
 		Section.AddMenuEntry( FBlueprintEditorCommands::Get().DisableAllBreakpoints );
diff --git a/Engine/Source/Editor/Kismet/Public/BlueprintEditor.h b/Engine/Source/Editor/Kismet/Public/BlueprintEditor.h
--- a/Engine/Source/Editor/Kismet/Public/BlueprintEditor.h
+++ b/Engine/Source/Editor/Kismet/Public/BlueprintEditor.h
@@ -779,6 +779,10 @@ protected:
 	bool HasAnyDisabledBreakpoints() const;
 	bool HasAnyWatches() const;
 
+	/** Spawns a new blueprint debugger tab */
+	void OpenBlueprintDebugger();
+	bool CanOpenBlueprintDebugger() const;
+
 	// Utility helper to get the currently hovered pin in the currently visible graph, or NULL if there isn't one
 	UEdGraphPin* GetCurrentlySelectedPin() const;
 
diff --git a/Engine/Source/Editor/Kismet/Public/BlueprintEditorTabs.h b/Engine/Source/Editor/Kismet/Public/BlueprintEditorTabs.h
--- a/Engine/Source/Editor/Kismet/Public/BlueprintEditorTabs.h
+++ b/Engine/Source/Editor/Kismet/Public/BlueprintEditorTabs.h
@@ -19,6 +19,7 @@ struct KISMET_API FBlueprintEditorTabs
 	static const FName MyBlueprintID;
 	static const FName ReplaceNodeReferencesID;
 	static const FName UserDefinedStructureID;
+	static const FName BlueprintDebuggerID;
 
 	// Document tab identifiers
 	static const FName GraphEditorID;
