Updating the timeline editor in Blueprints so the buttons no longer look inactive and the highlight for a selected track is more subtle

diff --git a/Engine/Source/Editor/EditorStyle/Private/StarshipStyle.cpp b/Engine/Source/Editor/EditorStyle/Private/StarshipStyle.cpp
--- a/Engine/Source/Editor/EditorStyle/Private/StarshipStyle.cpp
+++ b/Engine/Source/Editor/EditorStyle/Private/StarshipStyle.cpp
@@ -834,14 +834,6 @@ void FStarshipEditorStyle::FStyle::SetupGeneralStyles()
 		Set("CurveEd.TangentSelected", new IMAGE_BRUSH("Common/Tangent", FVector2D(7.0f, 7.0f), FLinearColor(1.0f, 1.0f, 0.0f)));
 		Set("CurveEd.TangentColor", FLinearColor(0.0f, 0.66f, 0.7f));
 		Set("CurveEd.TangentColorSelected", FLinearColor(1.0f, 1.0f, 0.0f));
-		Set("CurveEd.Visible", new IMAGE_BRUSH("Icons/icon_levels_visible_16px", Icon16x16));
-		Set("CurveEd.VisibleHighlight", new IMAGE_BRUSH("Icons/icon_levels_visible_hi_16px", Icon16x16));
-		Set("CurveEd.Invisible", new IMAGE_BRUSH("Icons/icon_levels_invisible_16px", Icon16x16));
-		Set("CurveEd.InvisibleHighlight", new IMAGE_BRUSH("Icons/icon_levels_invisible_hi_16px", Icon16x16));
-		Set("CurveEd.Locked", new IMAGE_BRUSH("Icons/icon_locked_16px", Icon16x16));
-		Set("CurveEd.LockedHighlight", new IMAGE_BRUSH("Icons/icon_locked_highlight_16px", Icon16x16));
-		Set("CurveEd.Unlocked", new IMAGE_BRUSH("Icons/icon_unlocked_16px", Icon16x16));
-		Set("CurveEd.UnlockedHighlight", new IMAGE_BRUSH("Icons/icon_unlocked_highlight_16px", Icon16x16));
 	}
 	
 	// Scrub control buttons
@@ -4165,12 +4157,21 @@ void FStarshipEditorStyle::FStyle::SetupGraphEditorStyles()
 
 		// Timeline Editor
 		{
-			Set( "TimelineEditor.AddFloatTrack", new IMAGE_BRUSH( "Icons/icon_TrackAddFloat_36x24px", Icon36x24, FLinearColor::Black ) );
-			Set( "TimelineEditor.AddVectorTrack", new IMAGE_BRUSH( "Icons/icon_TrackAddVector_36x24px", Icon36x24, FLinearColor::Black ) );
-			Set( "TimelineEditor.AddEventTrack", new IMAGE_BRUSH( "Icons/icon_TrackAddEvent_36x24px", Icon36x24, FLinearColor::Black ) );
-			Set( "TimelineEditor.AddColorTrack", new IMAGE_BRUSH( "Icons/icon_TrackAddColor_36x24px", Icon36x24, FLinearColor::Black ) );
-			Set( "TimelineEditor.AddCurveAssetTrack", new IMAGE_BRUSH( "Icons/icon_TrackAddCurve_36x24px", Icon36x24, FLinearColor::Black ) );
-			Set( "TimelineEditor.DeleteTrack", new IMAGE_BRUSH( "Icons/icon_TrackDelete_36x24px", Icon36x24, FLinearColor::Black ) );
+			Set( "TimelineEditor.AddFloatTrack", new IMAGE_BRUSH( "Icons/icon_TrackAddFloat_36x24px", Icon36x24 ) );
+			Set( "TimelineEditor.AddVectorTrack", new IMAGE_BRUSH( "Icons/icon_TrackAddVector_36x24px", Icon36x24 ) );
+			Set( "TimelineEditor.AddEventTrack", new IMAGE_BRUSH( "Icons/icon_TrackAddEvent_36x24px", Icon36x24 ) );
+			Set( "TimelineEditor.AddColorTrack", new IMAGE_BRUSH( "Icons/icon_TrackAddColor_36x24px", Icon36x24 ) );
+			Set( "TimelineEditor.AddCurveAssetTrack", new IMAGE_BRUSH( "Icons/icon_TrackAddCurve_36x24px", Icon36x24 ) );
+			Set( "TimelineEditor.DeleteTrack", new IMAGE_BRUSH( "Icons/icon_TrackDelete_36x24px", Icon36x24 ) );
+
+			Set("TimelineEditor.TrackRowSubtleHighlight", FTableRowStyle(NormalTableRowStyle)
+				.SetActiveBrush(FSlateColorBrush(FStyleColors::Panel))
+				.SetActiveHoveredBrush(FSlateColorBrush(FStyleColors::Header))
+				.SetInactiveBrush(FSlateColorBrush(FStyleColors::Recessed))
+				.SetInactiveHoveredBrush(FSlateColorBrush(FStyleColors::Panel))
+				.SetActiveHighlightedBrush(FSlateColorBrush(FStyleColors::Panel)) // This is the parent hightlight
+				.SetInactiveHighlightedBrush(FSlateColorBrush(FStyleColors::Recessed))// This is the parent highlight
+			);
 		}
 	}
 
diff --git a/Engine/Source/Editor/Kismet/Private/STimelineEditor.cpp b/Engine/Source/Editor/Kismet/Private/STimelineEditor.cpp
--- a/Engine/Source/Editor/Kismet/Private/STimelineEditor.cpp
+++ b/Engine/Source/Editor/Kismet/Private/STimelineEditor.cpp
@@ -35,6 +35,7 @@
 #include "Widgets/Notifications/SNotificationList.h"
 #include "Engine/Selection.h"
 #include "AssetData.h"
+#include "Styling/StyleColors.h"
 
 #define LOCTEXT_NAMESPACE "STimelineEditor"
 
@@ -233,6 +234,7 @@ void STimelineEdTrack::Construct(const FArguments& InArgs, TSharedPtr<FTimelineE
 					[
 						SNew(STextBlock)
 						.Text(LOCTEXT("ExternalCurveLabel", "External Curve"))
+						.ColorAndOpacity(FStyleColors::Foreground)
 					]
 
 					// External Curve Controls
@@ -240,8 +242,9 @@ void STimelineEdTrack::Construct(const FArguments& InArgs, TSharedPtr<FTimelineE
 					.AutoHeight()
 					.Padding(2, 0, 0, 4)
 					[
-						SNew(SBox)
-						.WidthOverride(170)
+						SNew(SBorder)
+						.BorderImage(FEditorStyle::GetBrush("NoBrush"))
+						.ForegroundColor(FStyleColors::Foreground)
 						[
 							SNew(SHorizontalBox)
 							+SHorizontalBox::Slot()
@@ -265,7 +268,8 @@ void STimelineEdTrack::Construct(const FArguments& InArgs, TSharedPtr<FTimelineE
 								.ToolTipText(NSLOCTEXT("TimelineEdTrack", "TimelineEdTrack_Clear", "Convert to Internal Curve"))
 								[
 									SNew(SImage)
-									.Image( FEditorStyle::GetBrush(TEXT("PropertyWindow.Button_Clear")) )
+									.Image( FEditorStyle::GetBrush(TEXT("PropertyWindow.Button_Clear")))
+									.ColorAndOpacity(FStyleColors::Foreground)
 								]
 							]
 						]
@@ -283,6 +287,7 @@ void STimelineEdTrack::Construct(const FArguments& InArgs, TSharedPtr<FTimelineE
 						[
 							SNew(STextBlock)
 							.Text(LOCTEXT("SynchronizeViewLabel", "Synchronize View"))
+							.ColorAndOpacity(FStyleColors::Foreground)
 						]
 					]
 
@@ -305,6 +310,7 @@ void STimelineEdTrack::Construct(const FArguments& InArgs, TSharedPtr<FTimelineE
 								[
 									SNew(SImage)
 									.Image( FEditorStyle::GetBrush(TEXT("ArrowUp")) )
+									.ColorAndOpacity(FStyleColors::Foreground)
 								]
 							]
 
@@ -322,6 +328,7 @@ void STimelineEdTrack::Construct(const FArguments& InArgs, TSharedPtr<FTimelineE
 								[
 									SNew(SImage)
 									.Image( FEditorStyle::GetBrush(TEXT("ArrowDown")) )
+									.ColorAndOpacity(FStyleColors::Foreground)
 								]
 							]
 							+SHorizontalBox::Slot()
@@ -331,6 +338,7 @@ void STimelineEdTrack::Construct(const FArguments& InArgs, TSharedPtr<FTimelineE
 							[
 								SNew(STextBlock)
 								.Text(LOCTEXT("ReorderLabel", "Reorder"))
+								.ColorAndOpacity(FStyleColors::Foreground)
 							]
 						]
 				]
@@ -1033,10 +1041,11 @@ void STimelineEditor::Construct(const FArguments& InArgs, TSharedPtr<FBlueprintE
 			[
 				// Add float track button
 				SNew(SButton)
-				.ContentPadding(FMargin(2,0))
+				.ButtonStyle(FAppStyle::Get(), "SimpleButton")
 				[
 					SNew(SImage)
-					.Image(FEditorStyle::GetBrush("TimelineEditor.AddFloatTrack"))
+					.Image(FAppStyle::Get().GetBrush("TimelineEditor.AddFloatTrack"))
+					.ColorAndOpacity(FSlateColor::UseForeground())
 				]
 				.ToolTipText( LOCTEXT( "AddFloatTrack", "Add Float Track" ) )
 				.OnClicked( this, &STimelineEditor::CreateNewTrack, FTTTrackBase::TT_FloatInterp )
@@ -1048,10 +1057,11 @@ void STimelineEditor::Construct(const FArguments& InArgs, TSharedPtr<FBlueprintE
 			[
 				// Add vector track button
 				SNew(SButton)
-				.ContentPadding(FMargin(2,0))
+				.ButtonStyle(FAppStyle::Get(), "SimpleButton")
 				[
 					SNew(SImage)
-					.Image(FEditorStyle::GetBrush("TimelineEditor.AddVectorTrack"))
+					.Image(FAppStyle::Get().GetBrush("TimelineEditor.AddVectorTrack"))
+					.ColorAndOpacity(FSlateColor::UseForeground())
 				]
 				.ToolTipText( LOCTEXT( "AddVectorTrack", "Add Vector Track" ) )
 				.OnClicked( this, &STimelineEditor::CreateNewTrack, FTTTrackBase::TT_VectorInterp )
@@ -1063,10 +1073,11 @@ void STimelineEditor::Construct(const FArguments& InArgs, TSharedPtr<FBlueprintE
 			[
 				// Add event track button
 				SNew(SButton)
-				.ContentPadding(FMargin(2,0))
+				.ButtonStyle(FAppStyle::Get(), "SimpleButton")
 				[
 					SNew(SImage)
-					.Image(FEditorStyle::GetBrush("TimelineEditor.AddEventTrack"))
+					.Image(FAppStyle::Get().GetBrush("TimelineEditor.AddEventTrack"))
+					.ColorAndOpacity(FSlateColor::UseForeground())
 				]
 				.ToolTipText( LOCTEXT( "AddEventTrack", "Add Event Track" ) )
 				.OnClicked( this, &STimelineEditor::CreateNewTrack, FTTTrackBase::TT_Event )
@@ -1078,10 +1089,11 @@ void STimelineEditor::Construct(const FArguments& InArgs, TSharedPtr<FBlueprintE
 			[
 				// Add color track button
 				SNew(SButton)
-				.ContentPadding(FMargin(2,0))
+				.ButtonStyle(FAppStyle::Get(), "SimpleButton")
 				[
 					SNew(SImage)
-					.Image(FEditorStyle::GetBrush("TimelineEditor.AddColorTrack"))
+					.Image(FAppStyle::Get().GetBrush("TimelineEditor.AddColorTrack"))
+					.ColorAndOpacity(FSlateColor::UseForeground())
 				]
 				.ToolTipText( LOCTEXT( "AddColorTrack", "Add Color Track" ) )
 				.OnClicked( this, &STimelineEditor::CreateNewTrack, FTTTrackBase::TT_LinearColorInterp )
@@ -1093,10 +1105,11 @@ void STimelineEditor::Construct(const FArguments& InArgs, TSharedPtr<FBlueprintE
 			[
 				// Add external curve asset button
 				SNew(SButton)
-				.ContentPadding(FMargin(2,0))
+				.ButtonStyle(FAppStyle::Get(), "SimpleButton")
 				[
 					SNew(SImage)
-					.Image(FEditorStyle::GetBrush("TimelineEditor.AddCurveAssetTrack"))
+					.Image(FAppStyle::Get().GetBrush("TimelineEditor.AddCurveAssetTrack"))
+					.ColorAndOpacity(FSlateColor::UseForeground())
 				]
 				.ToolTipText( LOCTEXT( "AddExternalAsset", "Add Selected Curve Asset" ) )
 				.IsEnabled( this, &STimelineEditor::IsCurveAssetSelected )
@@ -1133,7 +1146,7 @@ void STimelineEditor::Construct(const FArguments& InArgs, TSharedPtr<FBlueprintE
 				.IsChecked( this, &STimelineEditor::IsUseLastKeyframeChecked )
 				.OnCheckStateChanged( this, &STimelineEditor::OnUseLastKeyframeChanged )
 				[
-					SNew(STextBlock) .Text( LOCTEXT( "UseLastKeyframe", "Use Last Keyframe?" ) )
+					SNew(STextBlock) .Text( LOCTEXT( "UseLastKeyframe", "Use Last Keyframe" ) )
 					.AddMetaData<FTagMetaData>(TEXT("TimelineEditor.UseLastKeyframe"))
 				]
 			]
@@ -1310,6 +1323,7 @@ TSharedRef<ITableRow> STimelineEditor::MakeTrackWidget( TSharedPtr<FTimelineEdTr
 
 	return
 	SNew(STableRow< TSharedPtr<FTimelineEdTrack> >, OwnerTable )
+	.Style(&FAppStyle::Get().GetWidgetStyle<FTableRowStyle>("TimelineEditor.TrackRowSubtleHighlight"))
 	.Padding(FMargin(0, 0, 0, 2))
 	[
 		SNew(STimelineEdTrack, Track, SharedThis(this))
diff --git a/Engine/Source/Editor/UnrealEd/Private/SCurveEditor.cpp b/Engine/Source/Editor/UnrealEd/Private/SCurveEditor.cpp
--- a/Engine/Source/Editor/UnrealEd/Private/SCurveEditor.cpp
+++ b/Engine/Source/Editor/UnrealEd/Private/SCurveEditor.cpp
@@ -587,35 +587,37 @@ TSharedRef<SWidget> SCurveEditor::CreateCurveSelectionWidget() const
 
 				+ SHorizontalBox::Slot()
 				.AutoWidth()
-				.VAlign(VAlign_Center)
+				.VAlign(VAlign_Fill)
 				[
 					SNew(SCheckBox)
+					.Style(&FAppStyle::Get().GetWidgetStyle<FCheckBoxStyle>("ToggleButtonCheckbox"))
 					.IsChecked(this, &SCurveEditor::IsCurveVisible, CurveViewModel)
 					.OnCheckStateChanged(const_cast<SCurveEditor*>(this), &SCurveEditor::OnCurveIsVisibleChanged, CurveViewModel)
 					.ToolTipText(this, &SCurveEditor::GetIsCurveVisibleToolTip, CurveViewModel)
-					.CheckedImage(FEditorStyle::GetBrush("CurveEd.Visible"))
-					.CheckedHoveredImage(FEditorStyle::GetBrush("CurveEd.VisibleHighlight"))
-					.CheckedPressedImage(FEditorStyle::GetBrush("CurveEd.Visible"))
-					.UncheckedImage(FEditorStyle::GetBrush("CurveEd.Invisible"))
-					.UncheckedHoveredImage(FEditorStyle::GetBrush("CurveEd.InvisibleHighlight"))
-					.UncheckedPressedImage(FEditorStyle::GetBrush("CurveEd.Invisible"))
+					.CheckedImage(FAppStyle::Get().GetBrush("Icons.Visible"))
+					.CheckedHoveredImage(FAppStyle::Get().GetBrush("Icons.Visible"))
+					.CheckedPressedImage(FAppStyle::Get().GetBrush("Icons.Visible"))
+					.UncheckedImage(FAppStyle::Get().GetBrush("Icons.Hidden"))
+					.UncheckedHoveredImage(FAppStyle::Get().GetBrush("Icons.Hidden"))
+					.UncheckedPressedImage(FAppStyle::Get().GetBrush("Icons.Hidden"))
 				]
 
 				+ SHorizontalBox::Slot()
 				.AutoWidth()
-				.VAlign(VAlign_Center)
+				.VAlign(VAlign_Fill)
 				.Padding(2, 0, 0, 0)
 				[
 					SNew(SCheckBox)
+					.Style(&FAppStyle::Get().GetWidgetStyle<FCheckBoxStyle>("ToggleButtonCheckbox"))
 					.IsChecked(this, &SCurveEditor::IsCurveLocked, CurveViewModel)
 					.OnCheckStateChanged(const_cast<SCurveEditor*>(this), &SCurveEditor::OnCurveIsLockedChanged, CurveViewModel)
 					.ToolTipText(this, &SCurveEditor::GetIsCurveLockedToolTip, CurveViewModel)
-					.CheckedImage(FEditorStyle::GetBrush("CurveEd.Locked"))
-					.CheckedHoveredImage(FEditorStyle::GetBrush("CurveEd.LockedHighlight"))
-					.CheckedPressedImage(FEditorStyle::GetBrush("CurveEd.Locked"))
-					.UncheckedImage(FEditorStyle::GetBrush("CurveEd.Unlocked"))
-					.UncheckedHoveredImage(FEditorStyle::GetBrush("CurveEd.UnlockedHighlight"))
-					.UncheckedPressedImage(FEditorStyle::GetBrush("CurveEd.Unlocked"))
+					.CheckedImage(FAppStyle::Get().GetBrush("Icons.Lock"))
+					.CheckedHoveredImage(FAppStyle::Get().GetBrush("Icons.Lock"))
+					.CheckedPressedImage(FAppStyle::Get().GetBrush("Icons.Lock"))
+					.UncheckedImage(FAppStyle::Get().GetBrush("Icons.Unlock"))
+					.UncheckedHoveredImage(FAppStyle::Get().GetBrush("Icons.Unlock"))
+					.UncheckedPressedImage(FAppStyle::Get().GetBrush("Icons.Unlock"))
 					.Visibility(bCanEditTrack ? EVisibility::Visible : EVisibility::Collapsed)
 				]
 			];
