Blueprint details diff now correctly responds to Show Only Differing Properties being toggled, rather than showing no properties.

diff --git a/Engine/Source/Editor/Kismet/Private/DetailsDiff.cpp b/Engine/Source/Editor/Kismet/Private/DetailsDiff.cpp
--- a/Engine/Source/Editor/Kismet/Private/DetailsDiff.cpp
+++ b/Engine/Source/Editor/Kismet/Private/DetailsDiff.cpp
@@ -42,7 +42,7 @@ void FDetailsDiff::HighlightProperty(const FPropertySoftPath& PropertyName)
 	DetailsView->HighlightProperty(ResolvedProperty);
 }
 
-TSharedRef< SWidget > FDetailsDiff::DetailsWidget()
+TSharedRef< IDetailsView > FDetailsDiff::DetailsWidget() const
 {
 	return DetailsView.ToSharedRef();
 }
diff --git a/Engine/Source/Editor/Kismet/Private/SBlueprintDiff.cpp b/Engine/Source/Editor/Kismet/Private/SBlueprintDiff.cpp
--- a/Engine/Source/Editor/Kismet/Private/SBlueprintDiff.cpp
+++ b/Engine/Source/Editor/Kismet/Private/SBlueprintDiff.cpp
@@ -24,6 +24,7 @@
 #include "WidgetBlueprint.h"
 #include "HAL/PlatformApplicationMisc.h"
 #include "Framework/Application/SlateApplication.h"
+#include "Algo/Transform.h"
 
 #include "Subsystems/AssetEditorSubsystem.h"
 
@@ -259,6 +260,24 @@ public:
 		, NewDetails(InNewObject, FDetailsDiff::FOnDisplayedPropertiesChanged())
 	{
 		OldDetails.DiffAgainst(NewDetails, DifferingProperties, true);
+
+		TSet<FPropertyPath> PropertyPaths;
+		Algo::Transform(DifferingProperties, PropertyPaths,
+			[&InOldObject](const FSingleObjectDiffEntry& DiffEntry)
+			{
+				return DiffEntry.Identifier.ResolvePath(InOldObject);
+			});
+
+		OldDetails.DetailsWidget()->UpdatePropertiesWhitelist(PropertyPaths);
+
+		PropertyPaths.Reset();
+		Algo::Transform(DifferingProperties, PropertyPaths,
+			[&InNewObject](const FSingleObjectDiffEntry& DiffEntry)
+			{
+				return DiffEntry.Identifier.ResolvePath(InNewObject);
+			});
+
+		NewDetails.DetailsWidget()->UpdatePropertiesWhitelist(PropertyPaths);
 	}
 
 	virtual void GenerateTreeEntries(TArray< TSharedPtr<FBlueprintDifferenceTreeEntry> >& OutTreeEntries, TArray< TSharedPtr<FBlueprintDifferenceTreeEntry> >& OutRealDifferences) override
diff --git a/Engine/Source/Editor/Kismet/Public/DetailsDiff.h b/Engine/Source/Editor/Kismet/Public/DetailsDiff.h
--- a/Engine/Source/Editor/Kismet/Public/DetailsDiff.h
+++ b/Engine/Source/Editor/Kismet/Public/DetailsDiff.h
@@ -21,7 +21,7 @@ public:
 	void HighlightProperty( const FPropertySoftPath& PropertyName );
 
 	/** Returns actual widget that is used to display details */
-	TSharedRef< SWidget > DetailsWidget();
+	TSharedRef< IDetailsView > DetailsWidget() const;
 
 	/** Returns object being displayed */
 	const UObject* GetDisplayedObject() const { return DisplayedObject; }
