Do not display subobjects in Data only BP editor mode. This makes the data only BP editor only display the correct categories.

diff --git a/Engine/Source/Editor/Kismet/Private/SKismetInspector.cpp b/Engine/Source/Editor/Kismet/Private/SKismetInspector.cpp
--- a/Engine/Source/Editor/Kismet/Private/SKismetInspector.cpp
+++ b/Engine/Source/Editor/Kismet/Private/SKismetInspector.cpp
@@ -42,6 +42,7 @@
 #include "K2Node_BitmaskLiteral.h"
 #include "BitmaskLiteralDetails.h"
 #include "FormatTextDetails.h"
+#include "Kismet2/BlueprintEditorUtils.h"
 
 #define LOCTEXT_NAMESPACE "KismetInspector"
 
@@ -789,11 +790,12 @@ bool SKismetInspector::IsPropertyVisible( const FPropertyAndParent& PropertyAndP
 		}
 	}
 
+	const UBlueprint* BP = BlueprintEditorPtr.IsValid() ? BlueprintEditorPtr.Pin()->GetBlueprintObj() : nullptr;
+	const bool bIsDataOnlyBP = BP && FBlueprintEditorUtils::IsDataOnlyBlueprint(BP);
+	
 	if(const UClass* OwningClass = Property.GetOwner<UClass>())
 	{
-		const UBlueprint* BP = BlueprintEditorPtr.IsValid() ? BlueprintEditorPtr.Pin()->GetBlueprintObj() : nullptr;
 		const bool VariableAddedInCurentBlueprint = (OwningClass->ClassGeneratedBy == BP);
-
 		// If we did not add this var, hide it!
 		if(!VariableAddedInCurentBlueprint)
 		{
@@ -821,9 +823,10 @@ bool SKismetInspector::IsPropertyVisible( const FPropertyAndParent& PropertyAndP
 	}
 
 	bool bIsComponent = (ObjectProperty != nullptr && ObjectProperty->PropertyClass->IsChildOf(UActorComponent::StaticClass()));
-	if (!bShowComponents && bIsComponent)
+	// Don't show subobject properties, thats what selecting components in the component tree is for.
+	// Don't show any subobject properties in data only BP mode either
+	if ((!bShowComponents || bIsDataOnlyBP) && bIsComponent)
 	{
-		// Don't show sub components properties, thats what selecting components in the component tree is for.
 		return false;
 	}
 
