UE-112805 - Crash with BP action menu

The original issue can't be reproduced, but it seems like the user is deleting a variable and then possibly right-clicking its old location before the UI had a chance to refresh. This would cause GetProperty() to return null, and cause a crash. Other actions in this function seem to already guard against this scenario, so we can just mimic that here.

Additionally, added a couple of extra checks for other actions where the pointer isn't checked before dereferencing.

diff --git a/Engine/Source/Editor/Kismet/Private/SMyBlueprint.cpp b/Engine/Source/Editor/Kismet/Private/SMyBlueprint.cpp
--- a/Engine/Source/Editor/Kismet/Private/SMyBlueprint.cpp
+++ b/Engine/Source/Editor/Kismet/Private/SMyBlueprint.cpp
@@ -1874,7 +1874,9 @@ void SMyBlueprint::OnActionSelectedHelper(TSharedPtr<FEdGraphSchemaAction> InAct
 			if (GraphAction->EdGraph)
 			{
 				FGraphDisplayInfo DisplayInfo;
-				GraphAction->EdGraph->GetSchema()->GetGraphDisplayInformation(*GraphAction->EdGraph, DisplayInfo);
+				const UEdGraphSchema* Schema = GraphAction->EdGraph->GetSchema();
+				check(Schema != nullptr);
+				Schema->GetGraphDisplayInformation(*GraphAction->EdGraph, DisplayInfo);
 				Inspector->ShowDetailsForSingleObject(GraphAction->EdGraph, SKismetInspector::FShowDetailsOptions(DisplayInfo.PlainName));
 			}
 		}
@@ -1893,7 +1895,9 @@ void SMyBlueprint::OnActionSelectedHelper(TSharedPtr<FEdGraphSchemaAction> InAct
 			SKismetInspector::FShowDetailsOptions Options(FText::FromName(VarAction->GetVariableName()));
 			Options.bForceRefresh = true;
 
-			Inspector->ShowDetailsForSingleObject(VarAction->GetProperty()->GetUPropertyWrapper(), Options);
+			FProperty* Prop = VarAction->GetProperty();
+			UPropertyWrapper* PropWrap = (Prop ? Prop->GetUPropertyWrapper() : nullptr);
+			Inspector->ShowDetailsForSingleObject(PropWrap, Options);
 		}
 		else if (InAction->GetTypeId() == FEdGraphSchemaAction_K2LocalVar::StaticGetTypeId())
 		{
@@ -1928,8 +1932,12 @@ void SMyBlueprint::OnActionSelectedHelper(TSharedPtr<FEdGraphSchemaAction> InAct
 			InAction->GetTypeId() == FEdGraphSchemaAction_K2InputAction::StaticGetTypeId())
 		{
 			FEdGraphSchemaAction_K2TargetNode* TargetNodeAction = (FEdGraphSchemaAction_K2TargetNode*)InAction.Get();
-			SKismetInspector::FShowDetailsOptions Options(TargetNodeAction->NodeTemplate->GetNodeTitle(ENodeTitleType::EditableTitle));
-			Inspector->ShowDetailsForSingleObject(TargetNodeAction->NodeTemplate, Options);
+			
+			UK2Node* NodeTemplate = TargetNodeAction->NodeTemplate.Get();
+			check(NodeTemplate != nullptr)
+
+			SKismetInspector::FShowDetailsOptions Options(NodeTemplate->GetNodeTitle(ENodeTitleType::EditableTitle));
+			Inspector->ShowDetailsForSingleObject(NodeTemplate, Options);
 		}
 		else
 		{
