Reenable ability to override native class from blueprint
Added a config variable to allow it to be turned off, but it will default to on.

diff --git a/Engine/Config/BaseEngine.ini b/Engine/Config/BaseEngine.ini
--- a/Engine/Config/BaseEngine.ini
+++ b/Engine/Config/BaseEngine.ini
@@ -1440,6 +1440,7 @@ bExecutionAfterReturn=false
 bCanSuppressAccessViolation=false
 bLoadNativeConvertedBPClassInEditor=false
 bIgnoreCompileOnLoadErrorsOnBuildMachine=false
+bAllowNativeComponentClassOverrides=true
 
 [/Script/Engine.Blueprint]
 bRecompileOnLoad=true
diff --git a/Engine/Source/Editor/Kismet/Private/BlueprintDetailsCustomization.cpp b/Engine/Source/Editor/Kismet/Private/BlueprintDetailsCustomization.cpp
--- a/Engine/Source/Editor/Kismet/Private/BlueprintDetailsCustomization.cpp
+++ b/Engine/Source/Editor/Kismet/Private/BlueprintDetailsCustomization.cpp
@@ -5936,40 +5936,43 @@ void FBlueprintComponentDetails::CustomizeDetails(IDetailLayoutBuilder& DetailLa
 			]
 		];
 
-		// Remove UI to specify overriden component class until undo/redo can be made to work with crashing
-		/*UActorComponent* ComponentTemplate = (CachedNodePtr->IsNative() ? CachedNodePtr->GetComponentTemplate() : nullptr);
-		if (ComponentTemplate)
+		// Keep an easy way to disable UI to specify overriden component class until there is confidence that it is robust
+		if (GetAllowNativeComponentClassOverrides())
 		{
-			UClass* BaseClass = ComponentTemplate->GetClass();
-
-			if (const FBPComponentClassOverride* Override = BlueprintObj->ComponentClassOverrides.FindByKey(ComponentTemplate->GetFName()))
+			UActorComponent* ComponentTemplate = (CachedNodePtr->IsNativeComponent() ? CachedNodePtr->GetComponentTemplate() : nullptr);
+			if (ComponentTemplate)
 			{
-				AActor* Owner = ComponentTemplate->GetOwner();
-				AActor* OwnerArchetype = CastChecked<AActor>(Owner->GetArchetype());
-				if (UActorComponent* ArchetypeComponent = Cast<UActorComponent>((UObject*)FindObjectWithOuter(OwnerArchetype, UActorComponent::StaticClass(), ComponentTemplate->GetFName())))
+				UClass* BaseClass = ComponentTemplate->GetClass();
+
+				if (const FBPComponentClassOverride* Override = BlueprintObj->ComponentClassOverrides.FindByKey(ComponentTemplate->GetFName()))
 				{
-					BaseClass = ArchetypeComponent->GetClass();
+					AActor* Owner = ComponentTemplate->GetOwner();
+					AActor* OwnerArchetype = CastChecked<AActor>(Owner->GetArchetype());
+					if (UActorComponent* ArchetypeComponent = Cast<UActorComponent>((UObject*)FindObjectWithOuter(OwnerArchetype, UActorComponent::StaticClass(), ComponentTemplate->GetFName())))
+					{
+						BaseClass = ArchetypeComponent->GetClass();
+					}
 				}
-			}
 
-			const FText ComponentClassTooltip = LOCTEXT("BlueprintComponentDetails_ComponentClassOverrideTooltip", "The class to use when creating this component for this class. This can only be done for components defined in native at this time.");
+				const FText ComponentClassTooltip = LOCTEXT("BlueprintComponentDetails_ComponentClassOverrideTooltip", "The class to use when creating this component for this class. This can only be done for components defined in native at this time.");
 
-			VariableCategory.AddCustomRow( LOCTEXT("BlueprintComponentDetails_ComponentClassOverride", "Component Class") )
-			.NameContent()
-			[
-				SNew(STextBlock)
-				.Text(LOCTEXT("BlueprintComponentDetails_ComponentClassOverrideLabel", "Component Class"))
-				.ToolTipText(ComponentClassTooltip)
-				.Font(IDetailLayoutBuilder::GetDetailFont())
-			]
-			.ValueContent()
-			[
-				SNew(SClassPropertyEntryBox)
-				.MetaClass(BaseClass)
-				.AllowNone(false)
-				.SelectedClass(this, &FBlueprintComponentDetails::GetSelectedEntryClass)
-				.OnSetClass(this, &FBlueprintComponentDetails::HandleNewEntryClassSelected)];
-		}*/
+				VariableCategory.AddCustomRow( LOCTEXT("BlueprintComponentDetails_ComponentClassOverride", "Component Class") )
+				.NameContent()
+				[
+					SNew(STextBlock)
+					.Text(LOCTEXT("BlueprintComponentDetails_ComponentClassOverrideLabel", "Component Class"))
+					.ToolTipText(ComponentClassTooltip)
+					.Font(IDetailLayoutBuilder::GetDetailFont())
+				]
+				.ValueContent()
+				[
+					SNew(SClassPropertyEntryBox)
+					.MetaClass(BaseClass)
+					.AllowNone(false)
+					.SelectedClass(this, &FBlueprintComponentDetails::GetSelectedEntryClass)
+					.OnSetClass(this, &FBlueprintComponentDetails::HandleNewEntryClassSelected)];
+			}
+		}
 
 		IDetailCategoryBuilder& SocketsCategory = DetailLayout.EditCategory("Sockets", LOCTEXT("BlueprintComponentDetailsCategory", "Sockets"), ECategoryPriority::Important);
 
diff --git a/Engine/Source/Editor/KismetCompiler/Private/KismetCompiler.cpp b/Engine/Source/Editor/KismetCompiler/Private/KismetCompiler.cpp
--- a/Engine/Source/Editor/KismetCompiler/Private/KismetCompiler.cpp
+++ b/Engine/Source/Editor/KismetCompiler/Private/KismetCompiler.cpp
@@ -4059,7 +4059,10 @@ void FKismetCompilerContext::CompileClassLayout(EInternalCompilerFlags InternalF
 	// This validation requires CDO object.
 	ValidateVariableNames();
 
-	ValidateComponentClassOverrides();
+	if (GetAllowNativeComponentClassOverrides())
+	{
+		ValidateComponentClassOverrides();
+	}
 
 	OldCDO = NULL;
 	OldGenLinkerIdx = INDEX_NONE;
diff --git a/Engine/Source/Editor/UnrealEd/Private/Kismet2/BlueprintEditorUtils.cpp b/Engine/Source/Editor/UnrealEd/Private/Kismet2/BlueprintEditorUtils.cpp
--- a/Engine/Source/Editor/UnrealEd/Private/Kismet2/BlueprintEditorUtils.cpp
+++ b/Engine/Source/Editor/UnrealEd/Private/Kismet2/BlueprintEditorUtils.cpp
@@ -3167,7 +3167,7 @@ bool FBlueprintEditorUtils::IsDataOnlyBlueprint(const UBlueprint* Blueprint)
 		return false;
 	}
 
-	if (Blueprint->ComponentTemplates.Num() > 0 || Blueprint->Timelines.Num() > 0 || Blueprint->ComponentClassOverrides.Num() > 0)
+	if (Blueprint->ComponentTemplates.Num() > 0 || Blueprint->Timelines.Num() > 0 || (Blueprint->ComponentClassOverrides.Num() > 0 && GetAllowNativeComponentClassOverrides()))
 	{
 		return false;
 	}
diff --git a/Engine/Source/Editor/UnrealEd/Private/Kismet2/Kismet2.cpp b/Engine/Source/Editor/UnrealEd/Private/Kismet2/Kismet2.cpp
--- a/Engine/Source/Editor/UnrealEd/Private/Kismet2/Kismet2.cpp
+++ b/Engine/Source/Editor/UnrealEd/Private/Kismet2/Kismet2.cpp
@@ -807,16 +807,19 @@ static void ConformComponentsUtils::ConformRemovedNativeComponents(UObject* BpCd
 		}
 
 		// If we have overriden the class of a native component then ensure that the component still exists and that the overriden class is a valid subclass of it
-		if (const FBPComponentClassOverride* BPCO = CastChecked<UBlueprintGeneratedClass>(BlueprintClass)->ComponentClassOverrides.FindByKey(Component->GetFName()))
+		if (GetAllowNativeComponentClassOverrides())
 		{
-			if (BPCO->ComponentClass == Component->GetClass())
+			if (const FBPComponentClassOverride* BPCO = CastChecked<UBlueprintGeneratedClass>(BlueprintClass)->ComponentClassOverrides.FindByKey(Component->GetFName()))
 			{
-				if (UObject* OverridenComponent = (UObject*)FindObjectWithOuter(NativeCDO, UActorComponent::StaticClass(), Component->GetFName()))
+				if (BPCO->ComponentClass == Component->GetClass())
 				{
-					if (Component->IsA(OverridenComponent->GetClass()))
+					if (UObject* OverridenComponent = (UObject*)FindObjectWithOuter(NativeCDO, UActorComponent::StaticClass(), Component->GetFName()))
 					{
-						NewNativeComponents.Add(Component);
-						continue;
+						if (Component->IsA(OverridenComponent->GetClass()))
+						{
+							NewNativeComponents.Add(Component);
+							continue;
+						}
 					}
 				}
 			}
diff --git a/Engine/Source/Runtime/CoreUObject/Private/UObject/Class.cpp b/Engine/Source/Runtime/CoreUObject/Private/UObject/Class.cpp
--- a/Engine/Source/Runtime/CoreUObject/Private/UObject/Class.cpp
+++ b/Engine/Source/Runtime/CoreUObject/Private/UObject/Class.cpp
@@ -6301,9 +6301,12 @@ void UDynamicClass::AddReferencedObjects(UObject* InThis, FReferenceCollector& C
 	Collector.AddReferencedObjects(This->ComponentTemplates, This);
 	Collector.AddReferencedObjects(This->Timelines, This);
 
-	for (TPair<FName, UClass*>& Override : This->ComponentClassOverrides)
+	if (GetAllowNativeComponentClassOverrides())
 	{
-		Collector.AddReferencedObject(Override.Value);
+		for (TPair<FName, UClass*>& Override : This->ComponentClassOverrides)
+		{
+			Collector.AddReferencedObject(Override.Value);
+		}
 	}
 
 	Collector.AddReferencedObject(This->AnimClassImplementation, This);
diff --git a/Engine/Source/Runtime/CoreUObject/Private/UObject/UObjectGlobals.cpp b/Engine/Source/Runtime/CoreUObject/Private/UObject/UObjectGlobals.cpp
--- a/Engine/Source/Runtime/CoreUObject/Private/UObject/UObjectGlobals.cpp
+++ b/Engine/Source/Runtime/CoreUObject/Private/UObject/UObjectGlobals.cpp
@@ -78,6 +78,16 @@ static FAutoConsoleVariableRef CVarAllowUnversionedContentInEditor(
 	ECVF_Default
 );
 
+COREUOBJECT_API bool GetAllowNativeComponentClassOverrides()
+{
+	static const bool bAllowNativeComponentClassOverrides = []()
+	{
+		bool bAllow;
+		GConfig->GetBool(TEXT("Kismet"), TEXT("bAllowNativeComponentClassOverrides"), bAllow, GEngineIni);
+		return bAllow;
+	}();
+	return bAllowNativeComponentClassOverrides;
+}
 DEFINE_STAT(STAT_InitProperties);
 DEFINE_STAT(STAT_ConstructObject);
 DEFINE_STAT(STAT_AllocateObject);
@@ -2687,7 +2697,7 @@ void FObjectInitializer::Construct_Internal()
 	ThreadContext.ConstructedObject = Obj;
 	ThreadContext.PushInitializer(this);
 
-	if (Obj)
+	if (Obj && GetAllowNativeComponentClassOverrides())
 	{
 		Obj->GetClass()->SetupObjectInitializer(*this);
 	}
diff --git a/Engine/Source/Runtime/CoreUObject/Public/UObject/UObjectGlobals.h b/Engine/Source/Runtime/CoreUObject/Public/UObject/UObjectGlobals.h
--- a/Engine/Source/Runtime/CoreUObject/Public/UObject/UObjectGlobals.h
+++ b/Engine/Source/Runtime/CoreUObject/Public/UObject/UObjectGlobals.h
@@ -2495,6 +2495,9 @@ COREUOBJECT_API UPackage* FindOrConstructDynamicTypePackage(const TCHAR* Package
 /** Get names of "virtual" packages, that contain Dynamic types  */
 COREUOBJECT_API TMap<FName, FName>& GetConvertedDynamicPackageNameToTypeName();
 
+/** Utility accessor for whether we are running with component class overrides enabled */
+COREUOBJECT_API bool GetAllowNativeComponentClassOverrides();
+
 struct COREUOBJECT_API FDynamicClassStaticData
 {
 	/** Autogenerated Z_Construct* function pointer */
diff --git a/Engine/Source/Runtime/Engine/Private/BlueprintGeneratedClass.cpp b/Engine/Source/Runtime/Engine/Private/BlueprintGeneratedClass.cpp
--- a/Engine/Source/Runtime/Engine/Private/BlueprintGeneratedClass.cpp
+++ b/Engine/Source/Runtime/Engine/Private/BlueprintGeneratedClass.cpp
@@ -136,7 +136,7 @@ void UBlueprintGeneratedClass::PostLoad()
 #endif // WITH_EDITORONLY_DATA
 
 	// Update any component names that have been redirected
-	if (!FPlatformProperties::RequiresCookedData())
+	if (!FPlatformProperties::RequiresCookedData() && GetAllowNativeComponentClassOverrides())
 	{
 		for (FBPComponentClassOverride& Override : ComponentClassOverrides)
 		{
@@ -1517,12 +1517,15 @@ void UBlueprintGeneratedClass::GetDefaultObjectPreloadDependencies(TArray<UObjec
 		}
 	}
 
-	// Add the classes that will be used for overriding components defined in base classes
-	for (const FBPComponentClassOverride& Override : ComponentClassOverrides)
+	if (GetAllowNativeComponentClassOverrides())
 	{
-		if (Override.ComponentClass)
+		// Add the classes that will be used for overriding components defined in base classes
+		for (const FBPComponentClassOverride& Override : ComponentClassOverrides)
 		{
-			OutDeps.Add(Override.ComponentClass);
+			if (Override.ComponentClass)
+			{
+				OutDeps.Add(Override.ComponentClass);
+			}
 		}
 	}
 }
