During the creation of the BP Action menu builder, iterate in such a way that it will not crash if one of the filters triggers a compile on load, which would re-fresh the menu builder and cause it to be modified.

diff --git a/Engine/Source/Editor/Kismet/Private/BlueprintActionMenuBuilder.cpp b/Engine/Source/Editor/Kismet/Private/BlueprintActionMenuBuilder.cpp
--- a/Engine/Source/Editor/Kismet/Private/BlueprintActionMenuBuilder.cpp
+++ b/Engine/Source/Editor/Kismet/Private/BlueprintActionMenuBuilder.cpp
@@ -522,11 +522,15 @@ void FBlueprintActionMenuBuilder::RebuildActionList()
 	
 	FBlueprintActionDatabase& ActionDatabase = FBlueprintActionDatabase::Get();
 	FBlueprintActionDatabase::FActionRegistry const& ActionRegistry = ActionDatabase.GetAllActions();
-	for (auto const& ActionEntry : ActionRegistry)
+	
+	for (auto Iterator(ActionRegistry.CreateConstIterator()); Iterator; ++Iterator)
 	{
-		if (UObject *ActionObject = ActionEntry.Key.ResolveObjectPtr())
+		const FObjectKey& ObjKey = Iterator->Key;
+		const FBlueprintActionDatabase::FActionList& ActionList = Iterator->Value;
+
+		if (UObject* ActionObject = ObjKey.ResolveObjectPtr())
 		{
-			for (UBlueprintNodeSpawner const* NodeSpawner : ActionEntry.Value)
+			for (UBlueprintNodeSpawner const* NodeSpawner : ActionList)
 			{
 				FBlueprintActionInfo BlueprintAction(ActionObject, NodeSpawner);
 
@@ -544,9 +548,9 @@ void FBlueprintActionMenuBuilder::RebuildActionList()
 		else
 		{
 			// Remove this (invalid) entry on the next tick.
-			ActionDatabase.DeferredRemoveEntry(ActionEntry.Key);
+			ActionDatabase.DeferredRemoveEntry(ObjKey);
 		}
-	}	
+	}
 }
 
 #undef LOCTEXT_NAMESPACE
