UE-107119 - Crash in FKismetFunctionDragDropAction::DroppedOnPanel

Changed the GetHoveredGraph check to a soft-fail.

WM_MOUSEMOVE messages can arrive with the latest cursor coordinates, which can be in the "My Blueprint" pane. When this happen, the action's graph is set null. However, we can also end up processing WM_LBUTTONUP in the same tick, with cursor coordinates that are still in the graph pane. When this happens, we'll call DroppedOnPin, but the graph pointer will be null.

Null-checking should be safe here since we'll essentially ignore the action, which is then disposed of.

diff --git a/Engine/Source/Editor/Kismet/Private/BPFunctionDragDropAction.cpp b/Engine/Source/Editor/Kismet/Private/BPFunctionDragDropAction.cpp
--- a/Engine/Source/Editor/Kismet/Private/BPFunctionDragDropAction.cpp
+++ b/Engine/Source/Editor/Kismet/Private/BPFunctionDragDropAction.cpp
@@ -232,44 +232,48 @@ FReply FKismetFunctionDragDropAction::DroppedOnPin(FVector2D ScreenPosition, FVe
 	FReply Reply = FReply::Unhandled();
 
 	UEdGraph* Graph = GetHoveredGraph();
-	check(Graph); 
 
-	// The ActionNode set during construction points to the Graph, this is suitable for displaying the mouse decorator but needs to be more complete based on the current graph
-	UBlueprintFunctionNodeSpawner* FunctionNodeSpawner = GetDropAction(*Graph);
-
-	if (FunctionNodeSpawner)
+	// In certain cases, mouse movement messages can jump the event queue and process OnDragLeave before us.
+	// This ends up setting our graph pointer to null, so we need to guard against that.
+	if (Graph != nullptr)
 	{
-		FText CannotDropReason = FText::GetEmpty();
-		if (!CanBeDroppedDelegate.IsBound() || CanBeDroppedDelegate.Execute(nullptr, Graph, CannotDropReason))
+		// The ActionNode set during construction points to the Graph, this is suitable for displaying the mouse decorator but needs to be more complete based on the current graph
+		UBlueprintFunctionNodeSpawner* FunctionNodeSpawner = GetDropAction(*Graph);
+
+		if (FunctionNodeSpawner)
 		{
-			UFunction const* Function = GetFunctionProperty();
-			if ((Function != nullptr) && UEdGraphSchema_K2::CanUserKismetCallFunction(Function))
+			FText CannotDropReason = FText::GetEmpty();
+			if (!CanBeDroppedDelegate.IsBound() || CanBeDroppedDelegate.Execute(nullptr, Graph, CannotDropReason))
 			{
-				AnalyticCallback.ExecuteIfBound();
+				UFunction const* Function = GetFunctionProperty();
+				if ((Function != nullptr) && UEdGraphSchema_K2::CanUserKismetCallFunction(Function))
+				{
+					AnalyticCallback.ExecuteIfBound();
 
-				const FScopedTransaction Transaction(LOCTEXT("KismetFunction_DroppedOnPanel", "Function Dropped on Graph"));
+					const FScopedTransaction Transaction(LOCTEXT("KismetFunction_DroppedOnPanel", "Function Dropped on Graph"));
 
-				IBlueprintNodeBinder::FBindingSet Bindings;
-				UEdGraphNode* ResultNode = FunctionNodeSpawner->Invoke(Graph, Bindings, GraphPosition);
+					IBlueprintNodeBinder::FBindingSet Bindings;
+					UEdGraphNode* ResultNode = FunctionNodeSpawner->Invoke(Graph, Bindings, GraphPosition);
 
-				// Autowire the node if we were dragging on top of a pin
-				if (ResultNode != nullptr)
-				{
-					if (UEdGraphPin* FromPin = GetHoveredPin())
+					// Autowire the node if we were dragging on top of a pin
+					if (ResultNode != nullptr)
 					{
-						ResultNode->AutowireNewNode(FromPin);
+						if (UEdGraphPin* FromPin = GetHoveredPin())
+						{
+							ResultNode->AutowireNewNode(FromPin);
+						}
 					}
-				}
 
-				Reply = FReply::Handled();
+					Reply = FReply::Handled();
+				}
 			}
 		}
-	}
-	else if (SourceAction->GetTypeId() == FEdGraphSchemaAction_K2Event::StaticGetTypeId())
-	{
-		if (FEdGraphSchemaAction_K2Event* FuncAction = (FEdGraphSchemaAction_K2Event*)SourceAction.Get())
+		else if (SourceAction->GetTypeId() == FEdGraphSchemaAction_K2Event::StaticGetTypeId())
 		{
-			FKismetEditorUtilities::BringKismetToFocusAttentionOnObject(FuncAction->NodeTemplate);
+			if (FEdGraphSchemaAction_K2Event* FuncAction = (FEdGraphSchemaAction_K2Event*)SourceAction.Get())
+			{
+				FKismetEditorUtilities::BringKismetToFocusAttentionOnObject(FuncAction->NodeTemplate);
+			}
 		}
 	}
 
