[Blueprints]
Replace uses of the old SCSEditor with the new Subobject Editor in Blueprints.

The behavior will be functionally the same, but now has a properly seperated data model for modifying subobjects in the editor.

diff --git a/Engine/Source/Editor/Kismet/Kismet.Build.cs b/Engine/Source/Editor/Kismet/Kismet.Build.cs
--- a/Engine/Source/Editor/Kismet/Kismet.Build.cs
+++ b/Engine/Source/Editor/Kismet/Kismet.Build.cs
@@ -58,8 +58,9 @@ public class Kismet : ModuleRules
                 "UMG", // for SBlueprintDiff
                 "WorkspaceMenuStructure",
 				"ToolMenus",
-				"SubobjectDataInterface",
-			}
+				"SubobjectEditor",
+				"SubobjectDataInterface"
+            }
 			);
 
         DynamicallyLoadedModuleNames.AddRange(
diff --git a/Engine/Source/Editor/Kismet/Private/BlueprintDetailsCustomization.cpp b/Engine/Source/Editor/Kismet/Private/BlueprintDetailsCustomization.cpp
--- a/Engine/Source/Editor/Kismet/Private/BlueprintDetailsCustomization.cpp
+++ b/Engine/Source/Editor/Kismet/Private/BlueprintDetailsCustomization.cpp
@@ -53,7 +53,10 @@
 #include "IDetailsView.h"
 #include "Widgets/Colors/SColorPicker.h"
 #include "SKismetInspector.h"
-#include "SSCSEditor.h"
+#include "Widgets/SToolTip.h"
+#include "SSubobjectEditor.h"
+#include "SSubobjectBlueprintEditor.h"
+#include "SubobjectDataSubsystem.h"
 #include "SPinTypeSelector.h"
 #include "NodeFactory.h"
 #include "Kismet2/Kismet2NameValidators.h"
@@ -2417,7 +2420,7 @@ EVisibility FBlueprintVarActionDetails::ReplicationVisibility() const
 
 TSharedRef<SWidget> FBlueprintVarActionDetails::BuildEventsMenuForVariable() const
 {
-	if( MyBlueprint.IsValid() )
+	if(MyBlueprint.IsValid())
 	{
 		TSharedPtr<SMyBlueprint> MyBlueprintPtr = MyBlueprint.Pin();
 		FEdGraphSchemaAction_K2Var* Variable = MyBlueprintPtr->SelectionAsVar();
@@ -2425,8 +2428,8 @@ TSharedRef<SWidget> FBlueprintVarActionDetails::BuildEventsMenuForVariable() con
 		TWeakPtr<FBlueprintEditor> BlueprintEditorPtr = MyBlueprintPtr->GetBlueprintEditor();
 		if( BlueprintEditorPtr.IsValid() && ComponentProperty )
 		{
-			TSharedPtr<SSCSEditor> Editor =  BlueprintEditorPtr.Pin()->GetSCSEditor();
-			FMenuBuilder MenuBuilder( true, NULL );
+			TSharedPtr<SSubobjectBlueprintEditor> Editor = StaticCastSharedPtr<SSubobjectBlueprintEditor>(BlueprintEditorPtr.Pin()->GetSubobjectEditor());
+			FMenuBuilder MenuBuilder(true, nullptr);
 			Editor->BuildMenuEventsSection( MenuBuilder, BlueprintEditorPtr.Pin()->GetBlueprintObj(), ComponentProperty->PropertyClass, 
 											FCanExecuteAction::CreateSP(BlueprintEditorPtr.Pin().Get(), &FBlueprintEditor::InEditingMode),
 											FGetSelectedObjectsDelegate::CreateSP(MyBlueprintPtr.Get(), &SMyBlueprint::GetSelectedItemsForContextMenu));
@@ -5836,12 +5839,12 @@ void FBlueprintGlobalOptionsDetails::CustomizeDetails(IDetailLayoutBuilder& Deta
 void FBlueprintComponentDetails::CustomizeDetails(IDetailLayoutBuilder& DetailLayout)
 {
 	check( BlueprintEditorPtr.IsValid() );
-	TSharedPtr<SSCSEditor> Editor = BlueprintEditorPtr.Pin()->GetSCSEditor();
+	TSharedPtr<SSubobjectEditor> Editor = BlueprintEditorPtr.Pin()->GetSubobjectEditor();
 	check( Editor.IsValid() );
 	const UBlueprint* BlueprintObj = GetBlueprintObj();
 	check(BlueprintObj != nullptr);
 
-	TArray<FSCSEditorTreeNodePtrType> Nodes = Editor->GetSelectedNodes();
+	TArray<FSubobjectEditorTreeNodePtrType> Nodes = Editor->GetSelectedNodes();
 
 	if (!Nodes.Num())
 	{
@@ -5939,7 +5942,7 @@ void FBlueprintComponentDetails::CustomizeDetails(IDetailLayoutBuilder& DetailLa
 		// Keep an easy way to disable UI to specify overriden component class until there is confidence that it is robust
 		if (GetAllowNativeComponentClassOverrides())
 		{
-			UActorComponent* ComponentTemplate = (CachedNodePtr->IsNativeComponent() ? CachedNodePtr->GetComponentTemplate() : nullptr);
+			const UActorComponent* ComponentTemplate = (CachedNodePtr->IsNativeComponent() ? CachedNodePtr->GetComponentTemplate() : nullptr);
 			if (ComponentTemplate)
 			{
 				UClass* BaseClass = ComponentTemplate->GetClass();
@@ -6025,7 +6028,7 @@ void FBlueprintComponentDetails::CustomizeDetails(IDetailLayoutBuilder& DetailLa
 	if ( FBlueprintEditorUtils::DoesSupportEventGraphs(BlueprintObj) && Nodes.Num() == 1 )
 	{
 		// Use the component template to support native components as well
-		if (UActorComponent* ComponentTemplate = CachedNodePtr->GetComponentTemplate())
+		if (const UActorComponent* ComponentTemplate = CachedNodePtr->GetComponentTemplate())
 		{
 			AddEventsCategory(DetailLayout, CachedNodePtr->GetVariableName(), ComponentTemplate->GetClass());
 		}
@@ -6062,17 +6065,12 @@ void FBlueprintComponentDetails::OnVariableTextChanged(const FText& InNewText)
 
 	const FString& NewTextStr = InNewText.ToString();
 
-	if (USCS_Node* SCS_Node = CachedNodePtr->GetSCSNode())
+	if (USubobjectDataSubsystem* System = USubobjectDataSubsystem::Get())
 	{
-		if (!NewTextStr.IsEmpty() && !FComponentEditorUtils::IsValidVariableNameString(SCS_Node->ComponentTemplate, NewTextStr))
+		FText ErrorMsg;
+		if (!System->IsValidRename(CachedNodePtr->GetDataHandle(), InNewText, ErrorMsg))
 		{
-			VariableNameEditableTextBox->SetError(LOCTEXT("ComponentVariableRenameFailed_NotValid", "This name is reserved for engine use."));
-			return;
-		}
-
-		if (!FComponentEditorUtils::IsComponentNameAvailable(NewTextStr, SCS_Node->ComponentTemplate->GetOuter(), SCS_Node->ComponentTemplate))
-		{
-			VariableNameEditableTextBox->SetError(FText::Format(LOCTEXT("ComponentVariableRenameFailed_InUse", "{0} is in use by another variable or function!"), InNewText));
+			VariableNameEditableTextBox->SetError(ErrorMsg);
 			return;
 		}
 	}
@@ -6101,16 +6099,12 @@ void FBlueprintComponentDetails::OnVariableTextChanged(const FText& InNewText)
 
 void FBlueprintComponentDetails::OnVariableTextCommitted(const FText& InNewName, ETextCommit::Type InTextCommit)
 {
-	if ( !bIsVariableNameInvalid )
+	if (!bIsVariableNameInvalid)
 	{
 		check(CachedNodePtr.IsValid());
 
-		USCS_Node* SCS_Node = CachedNodePtr->GetSCSNode();
-		if(SCS_Node != NULL)
-		{
-			const FScopedTransaction Transaction( LOCTEXT("RenameComponentVariable", "Rename Component Variable") );
-			FBlueprintEditorUtils::RenameComponentMemberVariable(GetBlueprintObj(), CachedNodePtr->GetSCSNode(), FName( *InNewName.ToString() ));
-		}
+		const FScopedTransaction Transaction(LOCTEXT("RenameComponentVariable", "Rename Component Variable"));
+		USubobjectDataSubsystem::RenameSubobjectMemberVariable(GetBlueprintObj(), CachedNodePtr->GetDataHandle(), FName(*InNewName.ToString()));
 	}
 
 	bIsVariableNameInvalid = false;
@@ -6140,8 +6134,8 @@ void FBlueprintComponentDetails::OnTooltipTextCommitted(const FText& NewText, ET
 bool FBlueprintComponentDetails::OnVariableCategoryChangeEnabled() const
 {
 	check(CachedNodePtr.IsValid());
-	
-	return !CachedNodePtr->IsInheritedComponent();
+	const FSubobjectData* Data = CachedNodePtr->GetDataSource();
+	return !Data->IsInheritedComponent();
 }
 
 FText FBlueprintComponentDetails::OnGetVariableCategoryText() const
@@ -6249,7 +6243,7 @@ const UClass* FBlueprintComponentDetails::GetSelectedEntryClass() const
 {
 	check(CachedNodePtr.IsValid());
 
-	if (UActorComponent* ComponentTemplate = CachedNodePtr->GetComponentTemplate())
+	if (const UActorComponent* ComponentTemplate = CachedNodePtr->GetComponentTemplate())
 	{
 		UBlueprint* BlueprintObj = GetBlueprintObj();
 		check(BlueprintObj);
@@ -6280,7 +6274,7 @@ void FBlueprintComponentDetails::HandleNewEntryClassSelected(const UClass* NewEn
 	{
 		check(CachedNodePtr.IsValid());
 
-		if (UActorComponent* ComponentTemplate = CachedNodePtr->GetComponentTemplate())
+		if (const UActorComponent* ComponentTemplate = CachedNodePtr->GetComponentTemplate())
 		{
 			const FScopedTransaction Transaction(LOCTEXT("SetComponentClassOverride", "Set Component Class Override"));
 
@@ -6333,10 +6327,9 @@ void FBlueprintComponentDetails::HandleNewEntryClassSelected(const UClass* NewEn
 FText FBlueprintComponentDetails::GetSocketName() const
 {
 	check(CachedNodePtr.IsValid());
-
-	if (CachedNodePtr->GetSCSNode() != NULL)
+	if (FSubobjectData* Data = CachedNodePtr->GetDataSource())
 	{
-		return FText::FromName(CachedNodePtr->GetSCSNode()->AttachToName);
+		return Data->GetSocketName();
 	}
 	return FText::GetEmpty();
 }
@@ -6345,9 +6338,9 @@ bool FBlueprintComponentDetails::CanChangeSocket() const
 {
 	check(CachedNodePtr.IsValid());
 
-	if (CachedNodePtr->GetSCSNode() != NULL)
+	if (FSubobjectData* Data = CachedNodePtr->GetDataSource())
 	{
-		return !CachedNodePtr->IsInheritedComponent();
+		return !Data->IsInheritedComponent();
 	}
 	return true;
 }
@@ -6355,17 +6348,18 @@ bool FBlueprintComponentDetails::CanChangeSocket() const
 void FBlueprintComponentDetails::OnBrowseSocket()
 {
 	check(CachedNodePtr.IsValid());
-
-	if (CachedNodePtr->GetSCSNode() != NULL)
+	FSubobjectData* Data = CachedNodePtr->GetDataSource();
+	if (Data && Data->HasValidSocket())
 	{
-		TSharedPtr<SSCSEditor> Editor = BlueprintEditorPtr.Pin()->GetSCSEditor();
-		check( Editor.IsValid() );
+		TSharedPtr<SSubobjectEditor> Editor = BlueprintEditorPtr.Pin()->GetSubobjectEditor();
+		check(Editor.IsValid());
 
-		FSCSEditorTreeNodePtrType ParentFNode = CachedNodePtr->GetParent();
-
-		if (ParentFNode.IsValid())
+		FSubobjectEditorTreeNodePtrType ParentFNode = CachedNodePtr->GetParent();
+		FSubobjectData* ParentData = ParentFNode.IsValid() ? ParentFNode->GetDataSource() : nullptr;
+		if (ParentData)
 		{
-			if (USceneComponent* ParentSceneComponent = Cast<USceneComponent>(ParentFNode->GetOrCreateEditableComponentTemplate(Editor->GetBlueprint())))
+			// #TODO_BH Remove const cast
+			if (USceneComponent* ParentSceneComponent = const_cast<USceneComponent*>(ParentData->GetObjectForBlueprint<USceneComponent>(Editor->GetBlueprint())))
 			{
 				if (ParentSceneComponent->HasAnySockets())
 				{
@@ -6388,23 +6382,24 @@ void FBlueprintComponentDetails::OnBrowseSocket()
 void FBlueprintComponentDetails::OnClearSocket()
 {
 	check(CachedNodePtr.IsValid());
-
-	if (CachedNodePtr->GetSCSNode() != NULL)
+	FSubobjectData* Data = CachedNodePtr->GetDataSource();
+	
+	if (Data && Data->HasValidSocket())
 	{
-		CachedNodePtr->GetSCSNode()->AttachToName = NAME_None;
+		Data->SetupAttachment(NAME_None);
 		FBlueprintEditorUtils::MarkBlueprintAsStructurallyModified(GetBlueprintObj());
 	}
 }
 
-void FBlueprintComponentDetails::OnSocketSelection( FName SocketName )
+void FBlueprintComponentDetails::OnSocketSelection(FName SocketName)
 {
 	check(CachedNodePtr.IsValid());
 
-	USCS_Node* SCS_Node = CachedNodePtr->GetSCSNode();
-	if (SCS_Node != NULL)
+	FSubobjectData* Data = CachedNodePtr->GetDataSource();
+	if (Data && Data->HasValidSocket())
 	{
 		// Record selection if there is an actual asset attached
-		SCS_Node->AttachToName = SocketName;
+		Data->SetSocketName(SocketName);
 		FBlueprintEditorUtils::MarkBlueprintAsStructurallyModified(GetBlueprintObj());
 	}
 }
diff --git a/Engine/Source/Editor/Kismet/Private/BlueprintDetailsCustomization.h b/Engine/Source/Editor/Kismet/Private/BlueprintDetailsCustomization.h
--- a/Engine/Source/Editor/Kismet/Private/BlueprintDetailsCustomization.h
+++ b/Engine/Source/Editor/Kismet/Private/BlueprintDetailsCustomization.h
@@ -29,6 +29,7 @@ class SEditableTextBox;
 class SMultiLineEditableTextBox;
 class UEdGraphNode_Documentation;
 class UK2Node_Variable;
+class FSubobjectEditorTreeNode;
 
 /**
  * Variable network replication options.
@@ -897,7 +898,7 @@ private:
 	TWeakPtr<FBlueprintEditor> BlueprintEditorPtr;
 
 	/** The cached tree Node we're editing */
-	TSharedPtr<class FSCSEditorTreeNode> CachedNodePtr;
+	TSharedPtr<FSubobjectEditorTreeNode> CachedNodePtr;
 
 	/** The widget used when in variable name editing mode */ 
 	TSharedPtr<SEditableTextBox> VariableNameEditableTextBox;
diff --git a/Engine/Source/Editor/Kismet/Private/BlueprintEditor.cpp b/Engine/Source/Editor/Kismet/Private/BlueprintEditor.cpp
--- a/Engine/Source/Editor/Kismet/Private/BlueprintEditor.cpp
+++ b/Engine/Source/Editor/Kismet/Private/BlueprintEditor.cpp
@@ -112,7 +112,8 @@
 
 // Core kismet tabs
 #include "SGraphNode.h"
-#include "SSCSEditor.h"
+#include "SSubobjectBlueprintEditor.h"
+#include "SubobjectDataSubsystem.h"
 #include "SSCSEditorViewport.h"
 #include "SKismetInspector.h"
 #include "SBlueprintPalette.h"
@@ -173,6 +174,7 @@
 
 #include "BlendSpaceGraph.h"
 #include "AnimationBlendSpaceSampleGraph.h"
+#include "SSubobjectEditor.h"
 
 DEFINE_LOG_CATEGORY_STATIC(LogBlueprintEditor, Log, All);
 
@@ -446,9 +448,9 @@ const FSlateBrush* FBlueprintEditor::GetGlyphForGraph(const UEdGraph* Graph, boo
 {
 	const FSlateBrush* ReturnValue = FEditorStyle::GetBrush( bInLargeIcon ? TEXT("GraphEditor.Function_24x") : TEXT("GraphEditor.Function_16x") );
 
-	check(Graph != NULL);
+	check(Graph != nullptr);
 	const UEdGraphSchema* Schema = Graph->GetSchema();
-	if (Schema != NULL)
+	if (Schema != nullptr)
 	{
 		const EGraphType GraphType = Schema->GetGraphType(Graph);
 		switch (GraphType)
@@ -464,7 +466,7 @@ const FSlateBrush* FBlueprintEditor::GetGlyphForGraph(const UEdGraph* Graph, boo
 				if ( Graph->IsA(UAnimationTransitionGraph::StaticClass()) )
 				{
 					UObject* GraphOuter = Graph->GetOuter();
-					if ( GraphOuter != NULL && GraphOuter->IsA(UAnimStateConduitNode::StaticClass()) )
+					if ( GraphOuter != nullptr && GraphOuter->IsA(UAnimStateConduitNode::StaticClass()) )
 					{
 						ReturnValue = FEditorStyle::GetBrush( bInLargeIcon ? TEXT("GraphEditor.Conduit_24x") : TEXT("GraphEditor.Conduit_16x") );
 					}
@@ -532,7 +534,7 @@ const FSlateBrush* FBlueprintEditor::GetGlyphForGraph(const UEdGraph* Graph, boo
 
 FSlateBrush const* FBlueprintEditor::GetVarIconAndColor(const UStruct* VarScope, FName VarName, FSlateColor& IconColorOut, FSlateBrush const*& SecondaryBrushOut, FSlateColor& SecondaryColorOut)
 {
-	if (VarScope != NULL)
+	if (VarScope != nullptr)
 	{
 		FProperty* Property = FindFProperty<FProperty>(VarScope, VarName);
 		return GetVarIconAndColorFromProperty(Property, IconColorOut, SecondaryBrushOut, SecondaryColorOut);
@@ -542,7 +544,7 @@ FSlateBrush const* FBlueprintEditor::GetVarIconAndColor(const UStruct* VarScope,
 
 FSlateBrush const* FBlueprintEditor::GetVarIconAndColorFromProperty(const FProperty* Property, FSlateColor& IconColorOut, FSlateBrush const*& SecondaryBrushOut, FSlateColor& SecondaryColorOut)
 {
-	if (Property != NULL)
+	if (Property != nullptr)
 	{
 		const UEdGraphSchema_K2* K2Schema = GetDefault<UEdGraphSchema_K2>();
 
@@ -580,6 +582,7 @@ bool FBlueprintEditor::OnRequestClose()
 	}
 
 	bEditorMarkedAsClosed = true;
+	
 	return FWorkflowCentricApplication::OnRequestClose();
 }
 
@@ -597,7 +600,7 @@ bool FBlueprintEditor::IsCompilingEnabled() const
 
 bool FBlueprintEditor::InDebuggingMode() const
 {
-	return GEditor->PlayWorld != NULL;
+	return GEditor->PlayWorld != nullptr;
 }
 
 EVisibility FBlueprintEditor::IsDebuggerVisible() const
@@ -632,7 +635,7 @@ void FBlueprintEditor::AnalyticsTrackNodeEvent( UBlueprint* Blueprint, UEdGraphN
 	if( Blueprint && GraphNode && FEngineAnalytics::IsAvailable() )
 	{
 		// we'd like to see if this was happening in normal blueprint editor or persona 
-		const FString EditorName = Cast<UAnimBlueprint>(Blueprint) != NULL ? TEXT("Persona") : TEXT("BlueprintEditor");
+		const FString EditorName = Cast<UAnimBlueprint>(Blueprint) != nullptr ? TEXT("Persona") : TEXT("BlueprintEditor");
 
 		// Build Node Details
 		const UGeneralProjectSettings& ProjectSettings = *GetDefault<UGeneralProjectSettings>();
@@ -672,7 +675,7 @@ void FBlueprintEditor::AnalyticsTrackCompileEvent( UBlueprint* Blueprint, int32
 	if(Blueprint && FEngineAnalytics::IsAvailable())
 	{
 		// we'd like to see if this was happening in normal blueprint editor or persona 
-		const FString EditorName = Cast<UAnimBlueprint>(Blueprint) != NULL ? TEXT("Persona") : TEXT("BlueprintEditor");
+		const FString EditorName = Cast<UAnimBlueprint>(Blueprint) != nullptr ? TEXT("Persona") : TEXT("BlueprintEditor");
 
 		// Build Node Details
 		const UGeneralProjectSettings& ProjectSettings = *GetDefault<UGeneralProjectSettings>();
@@ -701,9 +704,9 @@ void FBlueprintEditor::RefreshEditors(ERefreshBlueprintEditorReason::Type Reason
 	}
 	else if (CurrentUISelection == SelectionState_Components)
 	{
-		if (SCSEditor.IsValid())
+		if(SubobjectEditor.IsValid())
 		{
-			SCSEditor->RefreshSelectionDetails();
+			SubobjectEditor->RefreshSelectionDetails();
 		}
 	}
 	else if (CurrentUISelection == SelectionState_Graph)
@@ -746,12 +749,12 @@ void FBlueprintEditor::RefreshEditors(ERefreshBlueprintEditorReason::Type Reason
 		MyBlueprintWidget->Refresh();
 	}
 
-	if (SCSEditor.IsValid())
+	if(SubobjectEditor.IsValid())
 	{
-		SCSEditor->UpdateTree();
-		
+		SubobjectEditor->UpdateTree();
+
 		// Note: Don't pass 'true' here because we don't want the preview actor to be reconstructed until after Blueprint modification is complete.
-		UpdateSCSPreview();
+		UpdateSubobjectPreview();
 	}
 
 	if (BookmarksWidget.IsValid())
@@ -817,10 +820,10 @@ void FBlueprintEditor::ClearSelectionStateFor(FName SelectionOwner)
 		}
 	}
 	else if ( SelectionOwner == SelectionState_Components )
-	{
-		if ( SCSEditor.IsValid() )
+	{	
+		if ( SubobjectEditor.IsValid() )
 		{
-			SCSEditor->ClearSelection();
+			SubobjectEditor->ClearSelection();
 		}
 	}
 	else if ( SelectionOwner == SelectionState_MyBlueprint )
@@ -858,18 +861,18 @@ void FBlueprintEditor::SummonFindAndReplaceUI()
 	TabManager->TryInvokeTab(FBlueprintEditorTabs::ReplaceNodeReferencesID);
 }
 
-void FBlueprintEditor::EnableSCSPreview(bool bEnable)
+void FBlueprintEditor::EnableSubobjectPreview(bool bEnable)
 {
-	if(SCSViewport.IsValid())
+	if(SubobjectViewport.IsValid())
 	{
-		SCSViewport->EnablePreview(bEnable);
+		SubobjectViewport->EnablePreview(bEnable);
 	}
 }
 
-void FBlueprintEditor::UpdateSCSPreview(bool bUpdateNow)
+void FBlueprintEditor::UpdateSubobjectPreview(bool bUpdateNow)
 {
 	// refresh widget
-	if(SCSViewport.IsValid())
+	if(SubobjectViewport.IsValid())
 	{
 		TSharedPtr<SDockTab> OwnerTab = Inspector->GetOwnerTab();
 		if ( OwnerTab.IsValid() )
@@ -878,27 +881,27 @@ void FBlueprintEditor::UpdateSCSPreview(bool bUpdateNow)
 		}
 
 		// Only request a refresh immediately if the viewport tab is in the foreground.
-		SCSViewport->RequestRefresh(false, bUpdateNow);
+		SubobjectViewport->RequestRefresh(false, bUpdateNow);
 	}
 }
 
-AActor* FBlueprintEditor::GetSCSEditorActorContext() const
+UObject* FBlueprintEditor::GetSubobjectEditorObjectContext() const
 {
 	// Return the current CDO that was last generated for the class
 	UBlueprint* Blueprint = GetBlueprintObj();
 	if(Blueprint != nullptr && Blueprint->GeneratedClass != nullptr)
 	{
-		return Blueprint->GeneratedClass->GetDefaultObject<AActor>();
+		return Blueprint->GeneratedClass->GetDefaultObject();
 	}
 	
 	return nullptr;
 }
 
-void FBlueprintEditor::OnSelectionUpdated(const TArray<FSCSEditorTreeNodePtrType>& SelectedNodes)
+void FBlueprintEditor::OnSelectionUpdated(const TArray<TSharedPtr<FSubobjectEditorTreeNode>>& SelectedNodes)
 {
-	if (SCSViewport.IsValid())
+	if (SubobjectViewport.IsValid())
 	{
-		SCSViewport->OnComponentSelectionChanged();
+		SubobjectViewport->OnComponentSelectionChanged();
 	}
 
 	UBlueprint* Blueprint = GetBlueprintObj();
@@ -920,7 +923,7 @@ void FBlueprintEditor::OnSelectionUpdated(const TArray<FSCSEditorTreeNodePtrType
 	if (Inspector.IsValid())
 	{
 		// Clear the my blueprints selection
-		if ( SelectedNodes.Num() > 0 )
+		if (SelectedNodes.Num() > 0)
 		{
 			SetUISelectionState(FBlueprintEditor::SelectionState_Components);
 		}
@@ -930,37 +933,38 @@ void FBlueprintEditor::OnSelectionUpdated(const TArray<FSCSEditorTreeNodePtrType
 		TArray<UObject*> InspectorObjects;
 		bool bShowComponents = true;
 		InspectorObjects.Empty(SelectedNodes.Num());
-		for (FSCSEditorTreeNodePtrType NodePtr : SelectedNodes)
+		for (FSubobjectEditorTreeNodePtrType NodePtr : SelectedNodes)
 		{
 			if (NodePtr.IsValid())
 			{
-				if (NodePtr->IsActorNode())
+				const FSubobjectData* NodeData = NodePtr->GetDataSource();
+				if (const AActor* Actor = NodeData->GetObject<AActor>())
 				{
-					if (AActor* DefaultActor = NodePtr->GetEditableObjectForBlueprint<AActor>(GetBlueprintObj()))
+					if (const AActor* DefaultActor = NodeData->GetObjectForBlueprint<AActor>(GetBlueprintObj()))
 					{
-						InspectorObjects.Add(DefaultActor);
-
+						InspectorObjects.Add(const_cast<AActor*>(DefaultActor));
+					
 						FString Title;
 						DefaultActor->GetName(Title);
 						InspectorTitle = FText::FromString(Title);
 						bShowComponents = false;
-
+					
 						TryInvokingDetailsTab();
 					}
 				}
 				else
 				{
-					UActorComponent* EditableComponent = NodePtr->GetOrCreateEditableComponentTemplate(GetBlueprintObj());
+					const UActorComponent* EditableComponent = NodeData->GetObjectForBlueprint<UActorComponent>(GetBlueprintObj());
 					if (EditableComponent)
 					{
 						InspectorTitle = FText::FromString(NodePtr->GetDisplayString());
-						InspectorObjects.Add(EditableComponent);
+						InspectorObjects.Add(const_cast<UActorComponent*>(EditableComponent));
 					}
-
-					if ( SCSViewport.IsValid() )
+				
+					if (SubobjectViewport.IsValid())
 					{
-						TSharedPtr<SDockTab> OwnerTab = SCSViewport->GetOwnerTab();
-						if ( OwnerTab.IsValid() )
+						TSharedPtr<SDockTab> OwnerTab = SubobjectViewport->GetOwnerTab();
+						if (OwnerTab.IsValid())
 						{
 							OwnerTab->FlashTab();
 						}
@@ -976,7 +980,7 @@ void FBlueprintEditor::OnSelectionUpdated(const TArray<FSCSEditorTreeNodePtrType
 	}
 }
 
-void FBlueprintEditor::OnComponentDoubleClicked(TSharedPtr<class FSCSEditorTreeNode> Node)
+void FBlueprintEditor::OnComponentDoubleClicked(TSharedPtr<FSubobjectEditorTreeNode> Node)
 {
 	TSharedPtr<SDockTab> OwnerTab = Inspector->GetOwnerTab();
 	if ( OwnerTab.IsValid() )
@@ -998,7 +1002,7 @@ TSharedRef<SWidget> FBlueprintEditor::CreateGraphTitleBarWidget(TSharedRef<FTabI
 /** Create new tab for the supplied graph - don't call this directly.*/
 TSharedRef<SGraphEditor> FBlueprintEditor::CreateGraphEditorWidget(TSharedRef<FTabInfo> InTabInfo, UEdGraph* InGraph)
 {
-	check((InGraph != NULL) && IsEditingSingleBlueprint());
+	check((InGraph != nullptr) && IsEditingSingleBlueprint());
 
 	// No need to regenerate the commands.
 	if(!GraphEditorCommands.IsValid())
@@ -1396,7 +1400,7 @@ void FBlueprintEditor::SetupGraphEditorEvents(UEdGraph* InGraph, SGraphEditor::F
 	InEvents.OnDoubleClicked = SGraphEditor::FOnDoubleClicked::CreateSP(this, &FBlueprintEditor::NavigateToParentGraphByDoubleClick);
 
 	// Custom menu for K2 schemas
-	if(InGraph->Schema != NULL && InGraph->Schema->IsChildOf(UEdGraphSchema_K2::StaticClass()))
+	if(InGraph->Schema != nullptr && InGraph->Schema->IsChildOf(UEdGraphSchema_K2::StaticClass()))
 	{
 		InEvents.OnCreateActionMenu = SGraphEditor::FOnCreateActionMenu::CreateSP(this, &FBlueprintEditor::OnCreateGraphActionMenu);
 	}
@@ -1413,7 +1417,7 @@ FGraphAppearanceInfo FBlueprintEditor::GetGraphAppearance(UEdGraph* InGraph) con
 	FGraphAppearanceInfo AppearanceInfo;
 
 	UBlueprint* Blueprint = (InGraph != nullptr) ? FBlueprintEditorUtils::FindBlueprintForGraph(InGraph) : GetBlueprintObj();
-	if (Blueprint != NULL)
+	if (Blueprint != nullptr)
 	{
 	switch (Blueprint->BlueprintType)
 	{
@@ -1494,7 +1498,7 @@ FBlueprintEditor::FBlueprintEditor()
 
 void FBlueprintEditor::EnsureBlueprintIsUpToDate(UBlueprint* BlueprintObj)
 {
-	// Purge any NULL graphs
+	// Purge any nullptr graphs
 	FBlueprintEditorUtils::PurgeNullGraphs(BlueprintObj);
 
 	// Make sure the blueprint is cosmetically up to date
@@ -1503,14 +1507,14 @@ void FBlueprintEditor::EnsureBlueprintIsUpToDate(UBlueprint* BlueprintObj)
 	if (FBlueprintEditorUtils::SupportsConstructionScript(BlueprintObj))
 	{
 		// If we don't have an SCS yet, make it
-		if(BlueprintObj->SimpleConstructionScript == NULL)
+		if(BlueprintObj->SimpleConstructionScript == nullptr)
 		{
-			check(NULL != BlueprintObj->GeneratedClass);
+			check(nullptr != BlueprintObj->GeneratedClass);
 			BlueprintObj->SimpleConstructionScript = NewObject<USimpleConstructionScript>(BlueprintObj->GeneratedClass);
 			BlueprintObj->SimpleConstructionScript->SetFlags(RF_Transactional);
 
 			// Recreate (or create) any widgets that depend on the SCS
-			CreateSCSEditors();
+			CreateSubobjectEditors();
 		}
 
 		// If we should have a UCS but don't yet, make it
@@ -1547,7 +1551,7 @@ void FBlueprintEditor::EnsureBlueprintIsUpToDate(UBlueprint* BlueprintObj)
 			}
 		
 			// Remove the SCS object reference
-			BlueprintObj->SimpleConstructionScript = NULL;
+			BlueprintObj->SimpleConstructionScript = nullptr;
 
 			// Mark the Blueprint as having been structurally modified
 			FBlueprintEditorUtils::MarkBlueprintAsStructurallyModified(BlueprintObj);
@@ -1593,7 +1597,7 @@ struct FLoadObjectsFromAssetRegistryHelper
 			if(AssetData[AssetIndex].IsValid())
 			{
 				FString AssetPath = AssetData[AssetIndex].ObjectPath.ToString();
-				TObjectType* Object = LoadObject<TObjectType>(NULL, *AssetPath, NULL, 0, NULL);
+				TObjectType* Object = LoadObject<TObjectType>(nullptr, *AssetPath, nullptr, 0, nullptr);
 				if (Object)
 				{
 					Collection.Add( MakeWeakObjectPtr(Object) );
@@ -2116,8 +2120,8 @@ void FBlueprintEditor::PostRegenerateMenusAndToolbars()
 
 FText FBlueprintEditor::GetParentClassNameText() const
 {
-	UClass* ParentClass = (GetBlueprintObj() != NULL) ? GetBlueprintObj()->ParentClass : NULL;
-	return (ParentClass != NULL) ? ParentClass->GetDisplayNameText() : LOCTEXT("BlueprintEditor_NoParentClass", "None");
+	UClass* ParentClass = (GetBlueprintObj() != nullptr) ? GetBlueprintObj()->ParentClass : nullptr;
+	return (ParentClass != nullptr) ? ParentClass->GetDisplayNameText() : LOCTEXT("BlueprintEditor_NoParentClass", "None");
 }
 
 bool FBlueprintEditor::IsParentClassOfObjectABlueprint(const UBlueprint* Blueprint) const
@@ -2148,10 +2152,10 @@ EVisibility FBlueprintEditor::GetEditParentClassVisibility() const
 bool FBlueprintEditor::IsParentClassNative() const
 {
 	const UBlueprint* Blueprint = GetBlueprintObj();
-	if (Blueprint != NULL)
+	if (Blueprint != nullptr)
 	{
 		UClass* ParentClass = Blueprint->ParentClass;
-		if (ParentClass != NULL)
+		if (ParentClass != nullptr)
 		{
 			if (ParentClass->HasAllClassFlags(CLASS_Native))
 			{
@@ -2404,7 +2408,7 @@ void FBlueprintEditor::FocusInspectorOnGraphSelection(const FGraphPanelSelection
 
 void FBlueprintEditor::CreateDefaultTabContents(const TArray<UBlueprint*>& InBlueprints)
 {
-	UBlueprint* InBlueprint = InBlueprints.Num() == 1 ? InBlueprints[0] : NULL;
+	UBlueprint* InBlueprint = InBlueprints.Num() == 1 ? InBlueprints[0] : nullptr;
 
 	// Cache off whether or not this is an interface, since it is used to govern multiple widget's behavior
 	const bool bIsInterface = (InBlueprint && InBlueprint->BlueprintType == BPTYPE_Interface);
@@ -2467,20 +2471,20 @@ void FBlueprintEditor::CreateDefaultTabContents(const TArray<UBlueprint*>& InBlu
 		InBlueprint->ParentClass->IsChildOf(AActor::StaticClass()) && 
 		InBlueprint->SimpleConstructionScript )
 	{
-		CreateSCSEditors();
+		CreateSubobjectEditors();
 	}
 }
 
-void FBlueprintEditor::CreateSCSEditors()
+void FBlueprintEditor::CreateSubobjectEditors()
 {
-	SCSEditor = SAssignNew(SCSEditor, SSCSEditor)
-		.ActorContext(this, &FBlueprintEditor::GetSCSEditorActorContext)
+	SubobjectEditor = SAssignNew(SubobjectEditor, SSubobjectBlueprintEditor)
+		.ObjectContext(this, &FBlueprintEditor::GetSubobjectEditorObjectContext)
 		.PreviewActor(this, &FBlueprintEditor::GetPreviewActor)
 		.AllowEditing(this, &FBlueprintEditor::InEditingMode)
 		.OnSelectionUpdated(this, &FBlueprintEditor::OnSelectionUpdated)
 		.OnItemDoubleClicked(this, &FBlueprintEditor::OnComponentDoubleClicked);
-
-	SCSViewport = SAssignNew(SCSViewport, SSCSEditorViewport)
+	
+	SubobjectViewport = SAssignNew(SubobjectViewport, SSCSEditorViewport)
 		.BlueprintEditor(SharedThis(this));
 }
 
@@ -2889,7 +2893,7 @@ void FBlueprintEditor::ReparentBlueprint_NewParentChosen(UClass* ChosenClass)
 {
 	UBlueprint* BlueprintObj = GetBlueprintObj();
 
-	if ((BlueprintObj != NULL) && (ChosenClass != NULL) && (ChosenClass != BlueprintObj->ParentClass))
+	if ((BlueprintObj != nullptr) && (ChosenClass != nullptr) && (ChosenClass != BlueprintObj->ParentClass))
 	{
 		// Notify user, about common interfaces
 		bool bReparent = true;
@@ -2971,9 +2975,9 @@ void FBlueprintEditor::ReparentBlueprint_NewParentChosen(UClass* ChosenClass)
 			// Ensure that the Blueprint is up-to-date (valid SCS etc.) after compiling (new parent class)
 			EnsureBlueprintIsUpToDate(BlueprintObj);
 
-			if (SCSEditor.IsValid())
+			if(SubobjectEditor.IsValid())
 			{
-				SCSEditor->UpdateTree();
+				SubobjectEditor->UpdateTree();
 			}
 		}
 	}
@@ -2984,7 +2988,7 @@ void FBlueprintEditor::ReparentBlueprint_NewParentChosen(UClass* ChosenClass)
 bool FBlueprintEditor::ReparentBlueprint_IsVisible() const
 {
 	UBlueprint* Blueprint = GetBlueprintObj();
-	if (Blueprint != NULL)
+	if (Blueprint != nullptr)
 	{
 		// Don't show the reparent option if it's an Interface or we're not in editing mode
 		return !FBlueprintEditorUtils::IsInterfaceBlueprint(Blueprint) && InEditingMode() && (BPTYPE_FunctionLibrary != Blueprint->BlueprintType);
@@ -3179,7 +3183,12 @@ void FBlueprintEditor::HandleUndoTransaction(const FTransaction* Transaction)
 		// Transaction affects the Blueprint this editor handles, so react as necessary
 		if (bAffectsBlueprint)
 		{
-			SetUISelectionState(NAME_None);
+			// Do not clear the selection on undo for the component tree so that the details panel
+			// does not get cleared unnecessarily
+			if(CurrentUISelection != SelectionState_Components)
+			{
+				SetUISelectionState(NAME_None);
+			}
 
 			RefreshEditors();
 
@@ -3439,7 +3448,7 @@ void FBlueprintEditor::OnBlueprintCompiled(UBlueprint* InBlueprint)
 	if( InBlueprint )
 	{
 		UUnrealEdEngine* EditorEngine = GUnrealEd;
-		// GUnrealEd can be null after a hot-reload... this seems like a bigger 
+		// GUnrealEd can be nullptr after a hot-reload... this seems like a bigger 
 		// problem worth investigating (that could affect other systems), but 
 		// as I cannot repro it a second time (to see if it gets reset soon after), 
 		// we'll just gaurd here for now and see if we can tie this ensure to any 
@@ -3946,12 +3955,12 @@ void FBlueprintEditor::JumpToNode(const UEdGraphNode* Node, bool bRequestRename)
 
 UBlueprint* FBlueprintEditor::GetBlueprintObj() const
 {
-	return GetEditingObjects().Num() == 1 ? Cast<UBlueprint>(GetEditingObjects()[0]) : NULL;
+	return GetEditingObjects().Num() == 1 ? Cast<UBlueprint>(GetEditingObjects()[0]) : nullptr;
 }
 
 bool FBlueprintEditor::IsEditingSingleBlueprint() const
 {
-	return GetBlueprintObj() != NULL;
+	return GetBlueprintObj() != nullptr;
 }
 
 FString FBlueprintEditor::GetDocumentationLink() const
@@ -4142,7 +4151,7 @@ void FBlueprintEditor::DoPromoteToVariable( UBlueprint* InBlueprint, UEdGraphPin
 		FEdGraphSchemaAction_K2NewNode NodeInfo;
 
 		// Create get or set node, depending on whether we clicked on an input or output pin
-		UK2Node_Variable* TemplateNode = NULL;
+		UK2Node_Variable* TemplateNode = nullptr;
 		if (InTargetPin->Direction == EGPD_Input)
 		{
 			TemplateNode = NewObject<UK2Node_VariableGet>();
@@ -4294,7 +4303,7 @@ void FBlueprintEditor::OnAddExecutionPin()
 	for (FGraphPanelSelectionSet::TConstIterator It(SelectedNodes); It; ++It)
 	{
 		UK2Node_ExecutionSequence* SeqNode = Cast<UK2Node_ExecutionSequence>(*It);
-		if (SeqNode != NULL)
+		if (SeqNode != nullptr)
 		{
 			const FScopedTransaction Transaction( LOCTEXT("AddExecutionPin", "Add Execution Pin") );
 			SeqNode->Modify();
@@ -4456,8 +4465,8 @@ bool FBlueprintEditor::CanRemoveExecutionPin() const
 void  FBlueprintEditor::OnRemoveThisStructVarPin()
 {
 	TSharedPtr<SGraphEditor> FocusedGraphEd = FocusedGraphEdPtr.Pin();
-	UEdGraphPin* SelectedPin = FocusedGraphEd.IsValid() ? FocusedGraphEd->GetGraphPinForMenu() : NULL;
-	UEdGraphNode* OwningNode = SelectedPin ? SelectedPin->GetOwningNodeUnchecked() : NULL;
+	UEdGraphPin* SelectedPin = FocusedGraphEd.IsValid() ? FocusedGraphEd->GetGraphPinForMenu() : nullptr;
+	UEdGraphNode* OwningNode = SelectedPin ? SelectedPin->GetOwningNodeUnchecked() : nullptr;
 	if (UK2Node_SetFieldsInStruct* SetFilestInStructNode = Cast<UK2Node_SetFieldsInStruct>(OwningNode))
 	{
 		const FScopedTransaction Transaction(LOCTEXT("RemoveThisStructVarPin", "Remove Struct Var Pin"));
@@ -4477,15 +4486,15 @@ void  FBlueprintEditor::OnRemoveThisStructVarPin()
 bool FBlueprintEditor::CanRemoveThisStructVarPin() const
 {
 	TSharedPtr<SGraphEditor> FocusedGraphEd = FocusedGraphEdPtr.Pin();
-	const UEdGraphPin*  SelectedPin = FocusedGraphEd.IsValid() ? FocusedGraphEd->GetGraphPinForMenu() : NULL;
+	const UEdGraphPin*  SelectedPin = FocusedGraphEd.IsValid() ? FocusedGraphEd->GetGraphPinForMenu() : nullptr;
 	return UK2Node_SetFieldsInStruct::ShowCustomPinActions(SelectedPin, false);
 }
 
 void  FBlueprintEditor::OnRemoveOtherStructVarPins()
 {
 	TSharedPtr<SGraphEditor> FocusedGraphEd = FocusedGraphEdPtr.Pin();
-	UEdGraphPin* SelectedPin = FocusedGraphEd.IsValid() ? FocusedGraphEd->GetGraphPinForMenu() : NULL;
-	UEdGraphNode* OwningNode = SelectedPin ? SelectedPin->GetOwningNodeUnchecked() : NULL;
+	UEdGraphPin* SelectedPin = FocusedGraphEd.IsValid() ? FocusedGraphEd->GetGraphPinForMenu() : nullptr;
+	UEdGraphNode* OwningNode = SelectedPin ? SelectedPin->GetOwningNodeUnchecked() : nullptr;
 	if (UK2Node_SetFieldsInStruct* SetFilestInStructNode = Cast<UK2Node_SetFieldsInStruct>(OwningNode))
 	{
 		const FScopedTransaction Transaction(LOCTEXT("RemoveOtherStructVarPins", "Remove Other Struct Var Pins"));
@@ -4505,7 +4514,7 @@ void  FBlueprintEditor::OnRemoveOtherStructVarPins()
 bool FBlueprintEditor::CanRemoveOtherStructVarPins() const
 {
 	TSharedPtr<SGraphEditor> FocusedGraphEd = FocusedGraphEdPtr.Pin();
-	const UEdGraphPin* SelectedPin = FocusedGraphEd.IsValid() ? FocusedGraphEd->GetGraphPinForMenu() : NULL;
+	const UEdGraphPin* SelectedPin = FocusedGraphEd.IsValid() ? FocusedGraphEd->GetGraphPinForMenu() : nullptr;
 	return UK2Node_SetFieldsInStruct::ShowCustomPinActions(SelectedPin, false);
 }
 
@@ -4513,7 +4522,7 @@ void FBlueprintEditor::OnRestoreAllStructVarPins()
 {
 	const FGraphPanelSelectionSet& SelectedNodes = GetSelectedNodes();
 	FGraphPanelSelectionSet::TConstIterator It(SelectedNodes);
-	UK2Node_SetFieldsInStruct* Node = (!!It) ? Cast<UK2Node_SetFieldsInStruct>(*It) : NULL;
+	UK2Node_SetFieldsInStruct* Node = (!!It) ? Cast<UK2Node_SetFieldsInStruct>(*It) : nullptr;
 	if (Node && !Node->AllPinsAreShown())
 	{
 		const FScopedTransaction Transaction(LOCTEXT("RestoreAllStructVarPins", "Restore all struct var pins"));
@@ -4539,7 +4548,7 @@ bool FBlueprintEditor::CanRestoreAllStructVarPins() const
 {
 	const FGraphPanelSelectionSet& SelectedNodes = GetSelectedNodes();
 	FGraphPanelSelectionSet::TConstIterator It(SelectedNodes);
-	UK2Node_SetFieldsInStruct* Node = (!!It) ? Cast<UK2Node_SetFieldsInStruct>(*It) : NULL;
+	UK2Node_SetFieldsInStruct* Node = (!!It) ? Cast<UK2Node_SetFieldsInStruct>(*It) : nullptr;
 	return Node && !Node->AllPinsAreShown();
 }
 
@@ -4584,7 +4593,7 @@ void FBlueprintEditor::OnAddOptionPin()
 	for (FGraphPanelSelectionSet::TConstIterator It(SelectedNodes); It; ++It)
 	{
 		UK2Node_Select* SeqNode = Cast<UK2Node_Select>(*It);
-		if (SeqNode != NULL)
+		if (SeqNode != nullptr)
 		{
 			const FScopedTransaction Transaction( LOCTEXT("AddOptionPin", "Add Option Pin") );
 			SeqNode->Modify();
@@ -4630,7 +4639,7 @@ void FBlueprintEditor::OnRemoveOptionPin()
 	for (FGraphPanelSelectionSet::TConstIterator It(SelectedNodes); It; ++It)
 	{
 		UK2Node_Select* SeqNode = Cast<UK2Node_Select>(*It);
-		if (SeqNode != NULL)
+		if (SeqNode != nullptr)
 		{
 			const FScopedTransaction Transaction( LOCTEXT("RemoveOptionPin", "Remove Option Pin") );
 			SeqNode->Modify();
@@ -4659,7 +4668,7 @@ bool FBlueprintEditor::CanRemoveOptionPin() const
 	{
 		UK2Node_Select* SeqNode = Cast<UK2Node_Select>(*It);
 		// There's a bad node so return false
-		if (SeqNode == NULL || !SeqNode->CanRemoveOptionPinToNode())
+		if (SeqNode == nullptr || !SeqNode->CanRemoveOptionPinToNode())
 		{
 			return false;
 		}
@@ -4810,7 +4819,7 @@ bool FBlueprintEditor::CanAddParentNode() const
 		FFunctionFromNodeHelper FunctionFromNode(SelectedObj);
 		if (FunctionFromNode.Function)
 		{
-			return (Schema->GetCallableParentFunction(FunctionFromNode.Function) != NULL);
+			return (Schema->GetCallableParentFunction(FunctionFromNode.Function) != nullptr);
 		}
 	}
 
@@ -4847,10 +4856,10 @@ void FBlueprintEditor::OnToggleBreakpoint()
 	for (FGraphPanelSelectionSet::TConstIterator NodeIt(SelectedNodes); NodeIt; ++NodeIt)
 	{
 		UK2Node* SelectedNode = Cast<UK2Node>(*NodeIt);
-		if ((SelectedNode != NULL) && SelectedNode->CanPlaceBreakpoints())
+		if ((SelectedNode != nullptr) && SelectedNode->CanPlaceBreakpoints())
 		{
 			UBreakpoint* ExistingBreakpoint = FKismetDebugUtilities::FindBreakpointForNode(GetBlueprintObj(), SelectedNode);
-			if (ExistingBreakpoint == NULL)
+			if (ExistingBreakpoint == nullptr)
 			{
 				// Add a breakpoint on this node if there isn't one there already
 				UBreakpoint* NewBreakpoint = NewObject<UBreakpoint>(GetBlueprintObj());
@@ -4874,7 +4883,7 @@ bool FBlueprintEditor::CanToggleBreakpoint() const
 	for (FGraphPanelSelectionSet::TConstIterator NodeIt(SelectedNodes); NodeIt; ++NodeIt)
 	{
 		UK2Node* SelectedNode = Cast<UK2Node>(*NodeIt);
-		if ((SelectedNode != NULL) && SelectedNode->CanPlaceBreakpoints())
+		if ((SelectedNode != nullptr) && SelectedNode->CanPlaceBreakpoints())
 		{
 			return true;
 		}
@@ -4889,11 +4898,11 @@ void FBlueprintEditor::OnAddBreakpoint()
 	for (FGraphPanelSelectionSet::TConstIterator NodeIt(SelectedNodes); NodeIt; ++NodeIt)
 	{
 		UK2Node* SelectedNode = Cast<UK2Node>(*NodeIt);
-		if ((SelectedNode != NULL) && SelectedNode->CanPlaceBreakpoints())
+		if ((SelectedNode != nullptr) && SelectedNode->CanPlaceBreakpoints())
 		{
 			// Add a breakpoint on this node if there isn't one there already
 			UBreakpoint* ExistingBreakpoint = FKismetDebugUtilities::FindBreakpointForNode(GetBlueprintObj(), SelectedNode);
-			if (ExistingBreakpoint == NULL)
+			if (ExistingBreakpoint == nullptr)
 			{
 				// Add a new breakpoint
 				UBreakpoint* NewBreakpoint = NewObject<UBreakpoint>(GetBlueprintObj());
@@ -4913,10 +4922,10 @@ bool FBlueprintEditor::CanAddBreakpoint() const
 	for (FGraphPanelSelectionSet::TConstIterator NodeIt(SelectedNodes); NodeIt; ++NodeIt)
 	{
 		UK2Node* SelectedNode = Cast<UK2Node>(*NodeIt);
-		if ((SelectedNode != NULL) && SelectedNode->CanPlaceBreakpoints())
+		if ((SelectedNode != nullptr) && SelectedNode->CanPlaceBreakpoints())
 		{
 			UBreakpoint* ExistingBreakpoint = FKismetDebugUtilities::FindBreakpointForNode(GetBlueprintObj(), SelectedNode);
-			if (ExistingBreakpoint == NULL)
+			if (ExistingBreakpoint == nullptr)
 			{
 				return true;
 			}
@@ -5121,7 +5130,7 @@ void FBlueprintEditor::OnCollapseSelectionToFunction()
 		const FScopedTransaction Transaction( FGraphEditorCommands::Get().CollapseNodes->GetDescription() );
 		BlueprintObj->Modify();
 
-		UEdGraphNode* FunctionNode = NULL;
+		UEdGraphNode* FunctionNode = nullptr;
 		UEdGraph* FunctionGraph = CollapseSelectionToFunction(FocusedGraphEdPtr.Pin(), CollapsableNodes, FunctionNode);
 
 		FBlueprintEditorUtils::MarkBlueprintAsStructurallyModified( BlueprintObj );
@@ -5286,7 +5295,7 @@ void FBlueprintEditor::OnCollapseSelectionToMacro()
 		const FScopedTransaction Transaction( FGraphEditorCommands::Get().CollapseNodes->GetDescription() );
 		BlueprintObj->Modify();
 
-		UEdGraphNode* MacroNode = NULL;
+		UEdGraphNode* MacroNode = nullptr;
 		UEdGraph* MacroGraph = CollapseSelectionToMacro(FocusedGraphEdPtr.Pin(), CollapsableNodes, MacroNode);
 
 		FBlueprintEditorUtils::MarkBlueprintAsStructurallyModified( BlueprintObj );
@@ -5299,7 +5308,7 @@ bool FBlueprintEditor::CanCollapseSelectionToMacro(TSet<class UEdGraphNode*>& In
 {
 	// Create a temporary macro graph
 	UEdGraph* MacroGraph = FBlueprintEditorUtils::CreateNewGraph(GetBlueprintObj(), FBlueprintEditorUtils::FindUniqueKismetName(GetBlueprintObj(), TEXT("TempGraph")), UEdGraph::StaticClass(), UEdGraphSchema_K2::StaticClass());
-	FBlueprintEditorUtils::AddMacroGraph(GetBlueprintObj(), MacroGraph, /*bIsUserCreated=*/ true, NULL);
+	FBlueprintEditorUtils::AddMacroGraph(GetBlueprintObj(), MacroGraph, /*bIsUserCreated=*/ true, nullptr);
 
 	// Does the selection set contain anything that is legal to collapse?
 	const UEdGraphSchema_K2* Schema = GetDefault<UEdGraphSchema_K2>();
@@ -5628,7 +5637,7 @@ bool FBlueprintEditor::CanExpandNodes() const
 		}
 		else if (UK2Node_CallFunction* SelectedCallFunctionNode = Cast<UK2Node_CallFunction>(*NodeIt))
 		{
-			// If ResultEventNode is non-NULL, it means it is sourced by an event, we do not want to expand events
+			// If ResultEventNode is non-nullptr, it means it is sourced by an event, we do not want to expand events
 			const UEdGraphNode* ResultEventNode = nullptr;
 			return SelectedCallFunctionNode->GetFunctionGraph(ResultEventNode) != nullptr && ResultEventNode == nullptr;
 		}
@@ -6499,7 +6508,7 @@ void FBlueprintEditor::DeleteSelectedDuplicatableNodes()
 		for (FGraphPanelSelectionSet::TConstIterator SelectedIter(OldSelectedNodes); SelectedIter; ++SelectedIter)
 		{
 			UEdGraphNode* Node = Cast<UEdGraphNode>(*SelectedIter);
-			if ((Node != NULL) && Node->CanDuplicateNode())
+			if ((Node != nullptr) && Node->CanDuplicateNode())
 			{
 				FocusedGraphEd->SetNodeSelection(Node, true);
 			}
@@ -6640,7 +6649,7 @@ struct FUpdatePastedNodes
 
 					if (bCanCurrentBlueprintReplace) 
 					{
-						UEdGraphNode* TargetNode = TargetInPin->LinkedTo.Num() ? TargetInPin->LinkedTo[0]->GetOwningNode() : NULL;
+						UEdGraphNode* TargetNode = TargetInPin->LinkedTo.Num() ? TargetInPin->LinkedTo[0]->GetOwningNode() : nullptr;
 						UK2Node_Literal* TargetLiteralNode = Cast<UK2Node_Literal>(TargetNode);
 
 						const bool bPastedNodeShouldBeReplacedWithTarget = TargetLiteralNode
@@ -6729,7 +6738,7 @@ private:
 		bool bResult = true;
 		check(OldCall);
 
-		UK2Node_VariableGet* NewTarget = NULL;
+		UK2Node_VariableGet* NewTarget = nullptr;
 		
 		const FProperty* Property = OldCall->MemberVariableToCallOn.ResolveMember<FProperty>();
 		for (UK2Node_VariableGet* AddedTarget : AddedTargets)
@@ -6879,7 +6888,7 @@ void FBlueprintEditor::PasteNodesHere(class UEdGraph* DestinationGraph, const FV
 			Node->CreateNewGuid();
 
 			UK2Node* K2Node = Cast<UK2Node>(Node);
-			if ((K2Node != NULL) && K2Node->NodeCausesStructuralBlueprintChange())
+			if ((K2Node != nullptr) && K2Node->NodeCausesStructuralBlueprintChange())
 			{
 				bNeedToModifyStructurally = true;
 			}
@@ -6991,17 +7000,17 @@ void FBlueprintEditor::OnAssignReferencedActor()
 {
 	const FGraphPanelSelectionSet SelectedNodes = GetSelectedNodes();
 	USelection* SelectedActors = GEditor->GetSelectedActors();
-	if ( SelectedNodes.Num() > 0 && SelectedActors != NULL && SelectedActors->Num() == 1 )
+	if ( SelectedNodes.Num() > 0 && SelectedActors != nullptr && SelectedActors->Num() == 1 )
 	{
 		AActor* SelectedActor = Cast<AActor>( SelectedActors->GetSelectedObject(0) );
-		if ( SelectedActor != NULL )
+		if ( SelectedActor != nullptr )
 		{
 			TArray<UK2Node_ActorBoundEvent*> NodesToAlter;
 
 			for ( FGraphPanelSelectionSet::TConstIterator NodeIt(SelectedNodes); NodeIt; ++NodeIt )
 			{
 				UK2Node_ActorBoundEvent* SelectedNode = Cast<UK2Node_ActorBoundEvent>(*NodeIt);
-				if ( SelectedNode != NULL )
+				if ( SelectedNode != nullptr )
 				{
 					NodesToAlter.Add( SelectedNode );
 				}
@@ -7043,15 +7052,15 @@ bool FBlueprintEditor::CanAssignReferencedActor() const
 
 		// If there is only one actor selected and at least one Blueprint graph
 		// node is able to receive the assignment then return true.
-		if ( SelectedActors != NULL && SelectedActors->Num() == 1 )
+		if ( SelectedActors != nullptr && SelectedActors->Num() == 1 )
 		{
 			AActor* SelectedActor = Cast<AActor>( SelectedActors->GetSelectedObject(0) );
-			if ( SelectedActor != NULL )
+			if ( SelectedActor != nullptr )
 			{
 				for ( FGraphPanelSelectionSet::TConstIterator NodeIt(SelectedNodes); NodeIt; ++NodeIt )
 				{
 					UK2Node_ActorBoundEvent* SelectedNode = Cast<UK2Node_ActorBoundEvent>(*NodeIt);
-					if ( SelectedNode != NULL )
+					if ( SelectedNode != nullptr )
 					{
 						if ( SelectedNode->EventOwner != SelectedActor )
 						{
@@ -7079,9 +7088,9 @@ void FBlueprintEditor::OnSelectReferenceInLevel()
 		for (FGraphPanelSelectionSet::TConstIterator NodeIt(SelectedNodes); NodeIt; ++NodeIt)
 		{
 			UK2Node* SelectedNode = Cast<UK2Node>(*NodeIt);
-			AActor* ReferencedActor = (SelectedNode) ? SelectedNode->GetReferencedLevelActor() : NULL;
+			AActor* ReferencedActor = (SelectedNode) ? SelectedNode->GetReferencedLevelActor() : nullptr;
 
-			if (ReferencedActor != NULL)
+			if (ReferencedActor != nullptr)
 			{
 				ActorsToSelect.AddUnique(ReferencedActor);
 			}
@@ -7116,10 +7125,10 @@ bool FBlueprintEditor::CanSelectReferenceInLevel() const
 		for (FGraphPanelSelectionSet::TConstIterator NodeIt(SelectedNodes); NodeIt; ++NodeIt)
 		{
 			UK2Node* SelectedNode = Cast<UK2Node>(*NodeIt);
-			const AActor* ReferencedActor = (SelectedNode) ? SelectedNode->GetReferencedLevelActor() : NULL;
+			const AActor* ReferencedActor = (SelectedNode) ? SelectedNode->GetReferencedLevelActor() : nullptr;
 
-			bCanSelectActors = (ReferencedActor != NULL);
-			if (ReferencedActor == NULL)
+			bCanSelectActors = (ReferencedActor != nullptr);
+			if (ReferencedActor == nullptr)
 			{
 				// Bail early if the selected node isn't referencing an actor
 				return false;
@@ -7130,7 +7139,7 @@ bool FBlueprintEditor::CanSelectReferenceInLevel() const
 	return bCanSelectActors;
 }
 
-// Utility helper to get the currently hovered pin in the currently visible graph, or NULL if there isn't one
+// Utility helper to get the currently hovered pin in the currently visible graph, or nullptr if there isn't one
 UEdGraphPin* FBlueprintEditor::GetCurrentlySelectedPin() const
 {
 	TSharedPtr<SGraphEditor> FocusedGraphEd = FocusedGraphEdPtr.Pin();
@@ -7139,7 +7148,7 @@ UEdGraphPin* FBlueprintEditor::GetCurrentlySelectedPin() const
 		return FocusedGraphEd->GetGraphPinForMenu();
 	}
 
-	return NULL;
+	return nullptr;
 }
 
 void FBlueprintEditor::SetDetailsCustomization(TSharedPtr<FDetailsViewObjectFilter> DetailsObjectFilter, TSharedPtr<IDetailRootObjectCustomization> DetailsRootCustomization)
@@ -7155,22 +7164,22 @@ void FBlueprintEditor::SetDetailsCustomization(TSharedPtr<FDetailsViewObjectFilt
 	}
 }
 
-void FBlueprintEditor::SetSCSEditorUICustomization(TSharedPtr<ISCSEditorUICustomization> SCSEditorUICustomization)
+void FBlueprintEditor::SetSubobjectEditorUICustomization(TSharedPtr<ISCSEditorUICustomization> SCSEditorUICustomization)
 {
-	if (SCSEditor.IsValid())
+	if(SubobjectEditor.IsValid())
 	{
-		SCSEditor->SetUICustomization(SCSEditorUICustomization);
+		SubobjectEditor->SetUICustomization(SCSEditorUICustomization);
 	}
 }
 
 void FBlueprintEditor::RegisterSCSEditorCustomization(const FName& InComponentName, TSharedPtr<ISCSEditorCustomization> InCustomization)
 {
-	SCSEditorCustomizations.Add(InComponentName, InCustomization);
+	SubobjectEditorCustomizations.Add(InComponentName, InCustomization);
 }
 
 void FBlueprintEditor::UnregisterSCSEditorCustomization(const FName& InComponentName)
 {
-	SCSEditorCustomizations.Remove(InComponentName);
+	SubobjectEditorCustomizations.Remove(InComponentName);
 }
 
 void FBlueprintEditor::CreateMergeToolTab()
@@ -7199,26 +7208,31 @@ void FBlueprintEditor::CloseMergeTool()
 	}
 }
 
-TArray<FSCSEditorTreeNodePtrType>  FBlueprintEditor::GetSelectedSCSEditorTreeNodes() const
+TArray<TSharedPtr<class FSCSEditorTreeNode> > FBlueprintEditor::GetSelectedSCSEditorTreeNodes() const
 {
-	TArray<FSCSEditorTreeNodePtrType>  Nodes;
-	if (SCSEditor.IsValid())
+	return TArray<TSharedPtr<class FSCSEditorTreeNode>>();
+}
+
+TArray<FSubobjectEditorTreeNodePtrType> FBlueprintEditor::GetSelectedSubobjectEditorTreeNodes() const
+{
+	TArray<FSubobjectEditorTreeNodePtrType>  Nodes;
+	if (SubobjectEditor.IsValid())
 	{
-		Nodes = SCSEditor->GetSelectedNodes();
+		Nodes = SubobjectEditor->GetSelectedNodes();
 	}
 	return Nodes;
 }
 
-FSCSEditorTreeNodePtrType FBlueprintEditor::FindAndSelectSCSEditorTreeNode(const UActorComponent* InComponent, bool IsCntrlDown) 
+FSubobjectEditorTreeNodePtrType FBlueprintEditor::FindAndSelectSubobjectEditorTreeNode(const UActorComponent* InComponent, bool IsCntrlDown)
 {
-	FSCSEditorTreeNodePtrType NodePtr;
+	FSubobjectEditorTreeNodePtrType NodePtr;
 
-	if (SCSEditor.IsValid())
+	if (SubobjectEditor.IsValid())
 	{
-		NodePtr = SCSEditor->GetNodeFromActorComponent(InComponent);
+		NodePtr = SubobjectEditor->FindSlateNodeForObject(InComponent);
 		if(NodePtr.IsValid())
 		{
-			SCSEditor->SelectNode(NodePtr, IsCntrlDown);
+			SubobjectEditor->SelectNode(NodePtr, IsCntrlDown);
 		}
 	}
 
@@ -7874,7 +7888,7 @@ void FBlueprintEditor::CollapseNodesIntoGraph(UEdGraphNode* InGatewayNode, UK2No
 		// Move the node over
 		InSourceGraph->Nodes.Remove(Node);
 		InDestinationGraph->Nodes.Add(Node);
-		Node->Rename(/*NewName=*/ NULL, /*NewOuter=*/ InDestinationGraph);
+		Node->Rename(/*NewName=*/ nullptr, /*NewOuter=*/ InDestinationGraph);
 
 		// Move the sub-graph to the new graph
 		if(UK2Node_Composite* Composite = Cast<UK2Node_Composite>(Node))
@@ -7928,8 +7942,8 @@ void FBlueprintEditor::CollapseNodesIntoGraph(UEdGraphNode* InGatewayNode, UK2No
 				UK2Node_EditablePinBase* LocalPort = (LocalPin->Direction == EGPD_Input) ? InEntryNode : InResultNode;
 
 				// Add a new pin to the entry/exit node and to the composite node
-				UEdGraphPin* LocalPortPin = NULL;
-				UEdGraphPin* RemotePortPin = NULL;
+				UEdGraphPin* LocalPortPin = nullptr;
+				UEdGraphPin* RemotePortPin = nullptr;
 
 				// Function graphs have a single exec path through them, so only one exec pin for input and another for output. In this fashion, they must not be handled by name.
 				if(InGatewayNode->GetClass() == UK2Node_CallFunction::StaticClass() && LocalPin->PinType.PinCategory == UEdGraphSchema_K2::PC_Exec)
@@ -8099,7 +8113,7 @@ void FBlueprintEditor::CollapseNodes(TSet<UEdGraphNode*>& InCollapsableNodes)
 	SourceGraph->Modify();
 
 	// Create the composite node that will serve as the gateway into the subgraph
-	UK2Node_Composite* GatewayNode = NULL;
+	UK2Node_Composite* GatewayNode = nullptr;
 	{
 		GatewayNode = FEdGraphSchemaAction_K2NewNode::SpawnNode<UK2Node_Composite>(SourceGraph, FVector2D(0,0), EK2NewNodeFlags::SelectNewNode);
 		GatewayNode->bCanRenameNode = true;
@@ -8492,9 +8506,9 @@ void FBlueprintEditor::StartEditingDefaults(bool bAutoFocus, bool bForceRefresh)
 	{
 		if (GetBlueprintObj()->GeneratedClass != nullptr)
 		{
-			if ( SCSEditor.IsValid() && GetBlueprintObj()->GeneratedClass->IsChildOf<AActor>() )
+			if (SubobjectEditor.IsValid() && GetBlueprintObj()->GeneratedClass->IsChildOf<AActor>())
 			{
-				SCSEditor->SelectRoot();
+				SubobjectEditor->SelectRoot();
 			}
 			else
 			{
@@ -8621,9 +8635,9 @@ void FBlueprintEditor::OnAddNewDelegate()
 	}
 
 	const UEdGraphSchema_K2* K2Schema = GetDefault<UEdGraphSchema_K2>();
-	check(NULL != K2Schema);
+	check(nullptr != K2Schema);
 	UBlueprint* const Blueprint = GetBlueprintObj();
-	check(NULL != Blueprint);
+	check(nullptr != Blueprint);
 
 	FName Name = FBlueprintEditorUtils::FindUniqueKismetName(GetBlueprintObj(), TEXT("NewEventDispatcher"));
 
@@ -8651,7 +8665,7 @@ void FBlueprintEditor::OnAddNewDelegate()
 	NewGraph->bEditable = false;
 
 	K2Schema->CreateDefaultNodesForGraph(*NewGraph);
-	K2Schema->CreateFunctionGraphTerminators(*NewGraph, (UClass*)NULL);
+	K2Schema->CreateFunctionGraphTerminators(*NewGraph, (UClass*)nullptr);
 	K2Schema->AddExtraFunctionFlags(NewGraph, (FUNC_BlueprintCallable|FUNC_BlueprintEvent|FUNC_Public));
 	K2Schema->MarkFunctionEntryAsEditable(NewGraph, true);
 
@@ -8704,17 +8718,17 @@ void FBlueprintEditor::NewDocument_OnClicked(ECreatedDocumentType GraphType)
 	const FScopedTransaction Transaction( LOCTEXT("AddNewFunction", "Add New Function") ); 
 	GetBlueprintObj()->Modify();
 
-	UEdGraph* NewGraph = NULL;
+	UEdGraph* NewGraph = nullptr;
 
 	if (GraphType == CGT_NewFunctionGraph)
 	{
 		NewGraph = FBlueprintEditorUtils::CreateNewGraph(GetBlueprintObj(), DocumentName, UEdGraph::StaticClass(), GetDefaultSchemaClass());
-		FBlueprintEditorUtils::AddFunctionGraph<UClass>(GetBlueprintObj(), NewGraph, /*bIsUserCreated=*/ true, NULL);
+		FBlueprintEditorUtils::AddFunctionGraph<UClass>(GetBlueprintObj(), NewGraph, /*bIsUserCreated=*/ true, nullptr);
 	}
 	else if (GraphType == CGT_NewMacroGraph)
 	{
 		NewGraph = FBlueprintEditorUtils::CreateNewGraph(GetBlueprintObj(), DocumentName, UEdGraph::StaticClass(), GetDefaultSchemaClass());
-		FBlueprintEditorUtils::AddMacroGraph(GetBlueprintObj(), NewGraph, /*bIsUserCreated=*/ true, NULL);
+		FBlueprintEditorUtils::AddMacroGraph(GetBlueprintObj(), NewGraph, /*bIsUserCreated=*/ true, nullptr);
 	}
 	else if (GraphType == CGT_NewEventGraph)
 	{
@@ -8796,7 +8810,7 @@ TSubclassOf<UEdGraphSchema> FBlueprintEditor::GetDefaultSchemaClass() const
 bool FBlueprintEditor::AddNewDelegateIsVisible() const
 {
 	const UBlueprint* Blueprint = GetBlueprintObj();
-	return (NULL != Blueprint)
+	return (nullptr != Blueprint)
 		&& (Blueprint->BlueprintType != BPTYPE_Interface) 
 		&& (Blueprint->BlueprintType != BPTYPE_MacroLibrary)
 		&& (Blueprint->BlueprintType != BPTYPE_FunctionLibrary);
@@ -8850,7 +8864,7 @@ void FBlueprintEditor::NotifyPostChange(const FPropertyChangedEvent& PropertyCha
 			}
 
 			// Force updates to occur immediately during interactive mode (otherwise the preview won't refresh because it won't be ticking)
-			UpdateSCSPreview(PropertyChangedEvent.ChangeType == EPropertyChangeType::Interactive);
+			UpdateSubobjectPreview(PropertyChangedEvent.ChangeType == EPropertyChangeType::Interactive);
 		}
 	}
 }
@@ -8924,7 +8938,7 @@ FText FBlueprintEditor::GetToolkitName() const
 		}
 	}
 
-	TSubclassOf< UObject > SharedParentClass = NULL;
+	TSubclassOf< UObject > SharedParentClass = nullptr;
 
 	for (UObject* EditingObj : EditingObjs)
 	{
@@ -8932,7 +8946,7 @@ FText FBlueprintEditor::GetToolkitName() const
 		check( Blueprint );
 
 		// Initialize with the class of the first object we encounter.
-		if( *SharedParentClass == NULL )
+		if( *SharedParentClass == nullptr )
 		{
 			SharedParentClass = Blueprint->ParentClass;
 		}
@@ -8971,7 +8985,7 @@ FText FBlueprintEditor::GetToolkitToolTipText() const
 		}
 	}
 
-	TSubclassOf< UObject > SharedParentClass = NULL;
+	TSubclassOf< UObject > SharedParentClass = nullptr;
 
 	for (UObject* EditingObj : EditingObjs)
 	{
@@ -8979,7 +8993,7 @@ FText FBlueprintEditor::GetToolkitToolTipText() const
 		check( Blueprint );
 
 		// Initialize with the class of the first object we encounter.
-		if( *SharedParentClass == NULL )
+		if( *SharedParentClass == nullptr )
 		{
 			SharedParentClass = Blueprint->ParentClass;
 		}
@@ -9165,7 +9179,7 @@ TSharedPtr<SDockTab> FBlueprintEditor::OpenDocument(const UObject* DocumentID, F
 
 void FBlueprintEditor::NavigateTab(FDocumentTracker::EOpenDocumentCause InCause)
 {
-	OpenDocument(NULL, InCause);
+	OpenDocument(nullptr, InCause);
 }
 
 void FBlueprintEditor::CloseDocumentTab(const UObject* DocumentID)
@@ -9266,7 +9280,7 @@ void FBlueprintEditor::OnCreateComment()
 				if (Schema->IsA(UEdGraphSchema_K2::StaticClass()))
 				{
 					FEdGraphSchemaAction_K2AddComment CommentAction;
-					CommentAction.PerformAction(Graph, NULL, GraphEditor->GetPasteLocation());
+					CommentAction.PerformAction(Graph, nullptr, GraphEditor->GetPasteLocation());
 				}
 			}
 		}
@@ -9290,7 +9304,7 @@ void FBlueprintEditor::OnFindReferences()
 		for (FGraphPanelSelectionSet::TConstIterator NodeIt(SelectedNodes); NodeIt; ++NodeIt)
 		{
 			UEdGraphNode* SelectedNode = Cast<UEdGraphNode>(*NodeIt);
-			if (SelectedNode != NULL)
+			if (SelectedNode != nullptr)
 			{
 				if (UK2Node_CustomEvent* CustomEvent = Cast<UK2Node_CustomEvent>(SelectedNode))
 				{
@@ -9353,7 +9367,7 @@ void FBlueprintEditor::UpdatePreviewActor(UBlueprint* InBlueprint, bool bInForce
 	AActor* PreviewActor = GetPreviewActor();
 
 	// Signal that we're going to be constructing editor components
-	if ( InBlueprint != NULL && InBlueprint->SimpleConstructionScript != NULL )
+	if ( InBlueprint != nullptr && InBlueprint->SimpleConstructionScript != nullptr )
 	{
 		InBlueprint->SimpleConstructionScript->BeginEditorComponentConstruction();
 	}
@@ -9518,11 +9532,11 @@ void FBlueprintEditor::UpdateNodeCreationStats( const ENodeCreateAction::Type Cr
 	}
 }
 
-TSharedPtr<ISCSEditorCustomization> FBlueprintEditor::CustomizeSCSEditor(USceneComponent* InComponentToCustomize) const
+TSharedPtr<ISCSEditorCustomization> FBlueprintEditor::CustomizeSubobjectEditor(const USceneComponent* InComponentToCustomize) const
 {
-	check(InComponentToCustomize != NULL);
-	const TSharedPtr<ISCSEditorCustomization>* FoundCustomization = SCSEditorCustomizations.Find(InComponentToCustomize->GetClass()->GetFName());
-	if(FoundCustomization != NULL)
+	check(InComponentToCustomize);
+	const TSharedPtr<ISCSEditorCustomization>* FoundCustomization = SubobjectEditorCustomizations.Find(InComponentToCustomize->GetClass()->GetFName());
+	if(FoundCustomization)
 	{
 		return *FoundCustomization;
 	}
diff --git a/Engine/Source/Editor/Kismet/Private/BlueprintEditorModes.cpp b/Engine/Source/Editor/Kismet/Private/BlueprintEditorModes.cpp
--- a/Engine/Source/Editor/Kismet/Private/BlueprintEditorModes.cpp
+++ b/Engine/Source/Editor/Kismet/Private/BlueprintEditorModes.cpp
@@ -20,6 +20,7 @@
 #include "BlueprintEditorTabFactories.h"
 #include "BlueprintEditorSharedTabFactories.h"
 #include "BlueprintEditorSettings.h"
+#include "SSubobjectEditor.h"
 
 #define LOCTEXT_NAMESPACE "BlueprintEditor"
 
@@ -272,18 +273,19 @@ void FBlueprintComponentsApplicationMode::PreDeactivateMode()
 	FApplicationMode::PreDeactivateMode();
 
 	TSharedPtr<FBlueprintEditor> BP = MyBlueprintEditor.Pin();
-	BP->GetSCSEditor()->SetEnabled(true);
-	BP->GetSCSEditor()->UpdateTree();
+	BP->GetSubobjectEditor()->SetEnabled(true);
+    BP->GetSubobjectEditor()->UpdateTree();
 	BP->GetInspector()->SetEnabled(true);
 	BP->GetInspector()->EnableComponentDetailsCustomization(false);
-	BP->EnableSCSPreview(false);
+	BP->EnableSubobjectPreview(false);
 
 	// Cache component selection before clearing so it can be restored
-	for( auto& SCSNode : BP->GetSCSEditor()->GetSelectedNodes() )
+	for(FSubobjectEditorTreeNodePtrType& Node : BP->GetSubobjectEditor()->GetSelectedNodes())
 	{
-		CachedComponentSelection.AddUnique(SCSNode->GetComponentTemplate());
+		CachedComponentSelection.AddUnique(Node->GetComponentTemplate());
 	}
-	BP->GetSCSEditor()->ClearSelection();
+	
+	BP->GetSubobjectEditor()->ClearSelection();
 }
 
 void FBlueprintComponentsApplicationMode::PostActivateMode()
@@ -291,25 +293,25 @@ void FBlueprintComponentsApplicationMode::PostActivateMode()
 	TSharedPtr<FBlueprintEditor> BP = MyBlueprintEditor.Pin();
 	if (BP.IsValid())
 	{
-		auto SCSEditor = BP->GetSCSEditor();
-		SCSEditor->UpdateTree();
-		BP->EnableSCSPreview(true);
-		BP->UpdateSCSPreview();
+		TSharedPtr<SSubobjectEditor> SubobjectEditor = BP->GetSubobjectEditor();
+		SubobjectEditor->UpdateTree();
+		BP->EnableSubobjectPreview(true);
+		BP->UpdateSubobjectPreview();
 		BP->GetInspector()->EnableComponentDetailsCustomization(true);
 
 		// Reselect the cached components
 		TArray<TSharedPtr<FSCSEditorTreeNode>> Selection;
-		for (auto Component : CachedComponentSelection)
+		for (TWeakObjectPtr<const UActorComponent>& Component : CachedComponentSelection)
 		{
 			if (Component.IsValid())
 			{
-				SCSEditor->SCSTreeWidget->SetItemSelection(SCSEditor->GetNodeFromActorComponent(Component.Get()), true);
+				SubobjectEditor->GetDragDropTree()->SetItemSelection(SubobjectEditor->FindSlateNodeForObject(Component.Get()), true);
 			}
 		}
 
-		if (BP->GetSCSViewport()->GetIsSimulateEnabled())
+		if (BP->GetSubobjectViewport()->GetIsSimulateEnabled())
 		{
-			SCSEditor->SetEnabled(false);
+			SubobjectEditor->SetEnabled(false);
 			BP->GetInspector()->SetEnabled(false);
 		}
 	}
diff --git a/Engine/Source/Editor/Kismet/Private/BlueprintEditorModule.cpp b/Engine/Source/Editor/Kismet/Private/BlueprintEditorModule.cpp
--- a/Engine/Source/Editor/Kismet/Private/BlueprintEditorModule.cpp
+++ b/Engine/Source/Editor/Kismet/Private/BlueprintEditorModule.cpp
@@ -271,7 +271,7 @@ TSharedRef<IBlueprintEditor> FBlueprintEditorModule::CreateBlueprintEditor(const
 	NewBlueprintEditor->InitBlueprintEditor(Mode, InitToolkitHost, BlueprintsToEdit, bShouldOpenInDefaultsMode);
 
 	NewBlueprintEditor->SetDetailsCustomization(DetailsObjectFilter, DetailsRootCustomization);
-	NewBlueprintEditor->SetSCSEditorUICustomization(SCSEditorUICustomization);
+	NewBlueprintEditor->SetSubobjectEditorUICustomization(SCSEditorUICustomization);
 
 	for(auto It(SCSEditorCustomizations.CreateConstIterator()); It; ++It)
 	{
@@ -345,13 +345,13 @@ void FBlueprintEditorModule::SetDetailsCustomization(TSharedPtr<FDetailsViewObje
 	}
 }
 
-void FBlueprintEditorModule::SetSCSEditorUICustomization(TSharedPtr<ISCSEditorUICustomization> InSCSEditorUICustomization)
+void FBlueprintEditorModule::SetSubobjectEditorUICustomization(TSharedPtr<ISCSEditorUICustomization> InSCSEditorUICustomization)
 {
 	SCSEditorUICustomization = InSCSEditorUICustomization;
 
 	for (const TSharedRef<IBlueprintEditor>& BlueprintEditor : GetBlueprintEditors())
 	{
-		StaticCastSharedRef<FBlueprintEditor>(BlueprintEditor)->SetSCSEditorUICustomization(SCSEditorUICustomization);
+		StaticCastSharedRef<FBlueprintEditor>(BlueprintEditor)->SetSubobjectEditorUICustomization(SCSEditorUICustomization);
 	}
 }
 
diff --git a/Engine/Source/Editor/Kismet/Private/BlueprintEditorTabFactories.cpp b/Engine/Source/Editor/Kismet/Private/BlueprintEditorTabFactories.cpp
--- a/Engine/Source/Editor/Kismet/Private/BlueprintEditorTabFactories.cpp
+++ b/Engine/Source/Editor/Kismet/Private/BlueprintEditorTabFactories.cpp
@@ -13,7 +13,7 @@
 #include "STimelineEditor.h"
 #include "Debugging/SKismetDebuggingView.h"
 #include "SKismetInspector.h"
-#include "SSCSEditor.h"
+#include "SSubobjectEditor.h"
 #include "SSCSEditorViewport.h"
 #include "SBlueprintPalette.h"
 #include "SBlueprintBookmarks.h"
@@ -299,7 +299,7 @@ TSharedRef<SWidget> FConstructionScriptEditorSummoner::CreateTabBody(const FWork
 {
 	TSharedPtr<FBlueprintEditor> BlueprintEditorPtr = StaticCastSharedPtr<FBlueprintEditor>(HostingApp.Pin());
 
-	return BlueprintEditorPtr->GetSCSEditor().ToSharedRef();
+	return BlueprintEditorPtr->GetSubobjectEditor().ToSharedRef();
 }
 
 FSCSViewportSummoner::FSCSViewportSummoner(TSharedPtr<class FAssetEditorToolkit> InHostingApp)
@@ -322,7 +322,7 @@ TSharedRef<SWidget> FSCSViewportSummoner::CreateTabBody(const FWorkflowTabSpawnI
 	TSharedPtr<SWidget> Result;
 	if (BlueprintEditorPtr->CanAccessComponentsMode())
 	{
-		Result = BlueprintEditorPtr->GetSCSViewport();
+		Result = BlueprintEditorPtr->GetSubobjectViewport();
 	}
 
 	if (Result.IsValid())
@@ -342,7 +342,7 @@ TSharedRef<SDockTab> FSCSViewportSummoner::SpawnTab(const FWorkflowTabSpawnInfo&
 	TSharedRef<SDockTab> Tab = FWorkflowTabFactory::SpawnTab(Info);
 
 	TSharedPtr<FBlueprintEditor> BlueprintEditorPtr = StaticCastSharedPtr<FBlueprintEditor>(HostingApp.Pin());
-	BlueprintEditorPtr->GetSCSViewport()->SetOwnerTab(Tab);
+	BlueprintEditorPtr->GetSubobjectViewport()->SetOwnerTab(Tab);
 
 	return Tab;
 }
diff --git a/Engine/Source/Editor/Kismet/Private/FindInBlueprints.cpp b/Engine/Source/Editor/Kismet/Private/FindInBlueprints.cpp
--- a/Engine/Source/Editor/Kismet/Private/FindInBlueprints.cpp
+++ b/Engine/Source/Editor/Kismet/Private/FindInBlueprints.cpp
@@ -431,7 +431,7 @@ FReply FFindInBlueprintsProperty::OnClick()
 							UActorComponent* Component = Node->GetActualComponentTemplate(GeneratedClass);
 							if (Component)
 							{
-								BlueprintEditor->FindAndSelectSCSEditorTreeNode(Component, false);
+								BlueprintEditor->FindAndSelectSubobjectEditorTreeNode(Component, false);
 							}
 						}
 						break;
diff --git a/Engine/Source/Editor/Kismet/Private/InstancedStaticMeshSCSEditorCustomization.cpp b/Engine/Source/Editor/Kismet/Private/InstancedStaticMeshSCSEditorCustomization.cpp
--- a/Engine/Source/Editor/Kismet/Private/InstancedStaticMeshSCSEditorCustomization.cpp
+++ b/Engine/Source/Editor/Kismet/Private/InstancedStaticMeshSCSEditorCustomization.cpp
@@ -47,7 +47,7 @@ bool FInstancedStaticMeshSCSEditorCustomization::HandleViewportClick(const TShar
 		if (BlueprintEditorPtr.IsValid())
 		{
 			// Note: This will find and select any node associated with the component instance that's attached to the proxy (including visualizers)
-			BlueprintEditorPtr.Pin()->FindAndSelectSCSEditorTreeNode(InstancedStaticMeshInstanceProxy->Component, bIsCtrlKeyDown);
+			BlueprintEditorPtr.Pin()->FindAndSelectSubobjectEditorTreeNode(InstancedStaticMeshInstanceProxy->Component, bIsCtrlKeyDown);
 		}
 
 		return true;
@@ -56,12 +56,12 @@ bool FInstancedStaticMeshSCSEditorCustomization::HandleViewportClick(const TShar
 	return false;
 }
 
-bool FInstancedStaticMeshSCSEditorCustomization::HandleViewportDrag(class USceneComponent* InSceneComponent, class USceneComponent* InComponentTemplate, const FVector& InDeltaTranslation, const FRotator& InDeltaRotation, const FVector& InDeltaScale, const FVector& InPivot)
+bool FInstancedStaticMeshSCSEditorCustomization::HandleViewportDrag(const USceneComponent* InSceneComponent, USceneComponent* InComponentTemplate, const FVector& InDeltaTranslation, const FRotator& InDeltaRotation, const FVector& InDeltaScale, const FVector& InPivot)
 {
 	check(InSceneComponent->IsA(UInstancedStaticMeshComponent::StaticClass()));
 
-	UInstancedStaticMeshComponent* InstancedStaticMeshComponentScene = CastChecked<UInstancedStaticMeshComponent>(InSceneComponent);
-	UInstancedStaticMeshComponent* InstancedStaticMeshComponentTemplate = CastChecked<UInstancedStaticMeshComponent>(InComponentTemplate);
+	UInstancedStaticMeshComponent* InstancedStaticMeshComponentScene = const_cast<UInstancedStaticMeshComponent*>(CastChecked<const UInstancedStaticMeshComponent>(InSceneComponent));
+	UInstancedStaticMeshComponent* InstancedStaticMeshComponentTemplate = const_cast<UInstancedStaticMeshComponent*>(CastChecked<const UInstancedStaticMeshComponent>(InComponentTemplate));
 
 	// transform pivot into component's space
 	const FVector LocalPivot = InstancedStaticMeshComponentScene->GetComponentToWorld().InverseTransformPosition(InPivot);
@@ -138,12 +138,12 @@ bool FInstancedStaticMeshSCSEditorCustomization::HandleViewportDrag(class UScene
 	return bMovedInstance;
 }
 
-bool FInstancedStaticMeshSCSEditorCustomization::HandleGetWidgetLocation(class USceneComponent* InSceneComponent, FVector& OutWidgetLocation)
+bool FInstancedStaticMeshSCSEditorCustomization::HandleGetWidgetLocation(const USceneComponent* InSceneComponent, FVector& OutWidgetLocation)
 {
 	// location is average of selected instances
 	float SelectedInstanceCount = 0.0f;
 	FVector AverageLocation = FVector::ZeroVector;
-	UInstancedStaticMeshComponent* InstancedStaticMeshComponent = CastChecked<UInstancedStaticMeshComponent>(InSceneComponent);
+	UInstancedStaticMeshComponent* InstancedStaticMeshComponent = const_cast<UInstancedStaticMeshComponent*>(CastChecked<UInstancedStaticMeshComponent>(InSceneComponent));
 
 	// Ensure that selected instances are up-to-date
 	ValidateSelectedInstances(InstancedStaticMeshComponent);
@@ -166,14 +166,14 @@ bool FInstancedStaticMeshSCSEditorCustomization::HandleGetWidgetLocation(class U
 	return false;
 }
 
-bool FInstancedStaticMeshSCSEditorCustomization::HandleGetWidgetTransform(class USceneComponent* InSceneComponent, FMatrix& OutWidgetTransform)
+bool FInstancedStaticMeshSCSEditorCustomization::HandleGetWidgetTransform(const USceneComponent* InSceneComponent, FMatrix& OutWidgetTransform)
 {
 	// transform is first selected instance
 	bool bInstanceSelected = false;
-	UInstancedStaticMeshComponent* InstancedStaticMeshComponent = CastChecked<UInstancedStaticMeshComponent>(InSceneComponent);
+	const UInstancedStaticMeshComponent* InstancedStaticMeshComponent = CastChecked<UInstancedStaticMeshComponent>(InSceneComponent);
 
 	// Ensure that selected instances are up-to-date
-	ValidateSelectedInstances(InstancedStaticMeshComponent);
+	ValidateSelectedInstances(const_cast<UInstancedStaticMeshComponent*>(InstancedStaticMeshComponent));
 
 	for (int32 InstanceIndex = 0; InstanceIndex < InstancedStaticMeshComponent->SelectedInstances.Num(); InstanceIndex++)
 	{
diff --git a/Engine/Source/Editor/Kismet/Private/InstancedStaticMeshSCSEditorCustomization.h b/Engine/Source/Editor/Kismet/Private/InstancedStaticMeshSCSEditorCustomization.h
--- a/Engine/Source/Editor/Kismet/Private/InstancedStaticMeshSCSEditorCustomization.h
+++ b/Engine/Source/Editor/Kismet/Private/InstancedStaticMeshSCSEditorCustomization.h
@@ -8,6 +8,8 @@
 #include "BlueprintEditorModule.h"
 #include "ISCSEditorCustomization.h"
 
+class USceneComponent;
+
 class FInstancedStaticMeshSCSEditorCustomization : public ISCSEditorCustomization
 {
 public:
@@ -16,9 +18,9 @@ public:
 public:
 	/** ISCSEditorCustomization interface */
 	virtual bool HandleViewportClick(const TSharedRef<FEditorViewportClient>& InViewportClient, class FSceneView& InView, class HHitProxy* InHitProxy, FKey InKey, EInputEvent InEvent, uint32 InHitX, uint32 InHitY) override;
-	virtual bool HandleViewportDrag(class USceneComponent* InComponentScene, class USceneComponent* InComponentTemplate, const FVector& InDeltaTranslation, const FRotator& InDeltaRotation, const FVector& InDeltaScale, const FVector& InPivot) override;
-	virtual bool HandleGetWidgetLocation(class USceneComponent* InSceneComponent, FVector& OutWidgetLocation) override;
-	virtual bool HandleGetWidgetTransform(class USceneComponent* InSceneComponent, FMatrix& OutWidgetTransform) override;
+	virtual bool HandleViewportDrag(const USceneComponent* InComponentScene, class USceneComponent* InComponentTemplate, const FVector& InDeltaTranslation, const FRotator& InDeltaRotation, const FVector& InDeltaScale, const FVector& InPivot) override;
+	virtual bool HandleGetWidgetLocation(const USceneComponent* InSceneComponent, FVector& OutWidgetLocation) override;
+	virtual bool HandleGetWidgetTransform(const USceneComponent* InSceneComponent, FMatrix& OutWidgetTransform) override;
 
 protected:
 	/** Ensure that selection bits are in sync w/ the number of instances */
diff --git a/Engine/Source/Editor/Kismet/Private/SBlueprintActionMenu.cpp b/Engine/Source/Editor/Kismet/Private/SBlueprintActionMenu.cpp
--- a/Engine/Source/Editor/Kismet/Private/SBlueprintActionMenu.cpp
+++ b/Engine/Source/Editor/Kismet/Private/SBlueprintActionMenu.cpp
@@ -10,6 +10,7 @@
 #include "Widgets/Images/SImage.h"
 #include "Widgets/Input/SComboButton.h"
 #include "Widgets/Input/SCheckBox.h"
+#include "Widgets/SToolTip.h"
 #include "EditorStyleSet.h"
 #include "Editor/EditorPerProjectUserSettings.h"
 #include "EdGraphSchema_K2.h"
@@ -22,7 +23,7 @@
 #include "BlueprintActionMenuUtils.h"
 #include "BlueprintPaletteFavorites.h"
 #include "IDocumentation.h"
-#include "SSCSEditor.h"
+#include "SSubobjectEditor.h"
 #include "SBlueprintContextTargetMenu.h"
 
 #define LOCTEXT_NAMESPACE "SBlueprintGraphContextMenu"
@@ -492,10 +493,10 @@ void SBlueprintActionMenu::ConstructActionContext(FBlueprintActionContext& Conte
 			ContextDescOut.SelectedObjects.Add(SelectedVar->GetProperty());
 		}
 		// If the selection come from the SCS editor, add it to the filter context.
-		else if (Blueprint->SkeletonGeneratedClass && BlueprintEditor->GetSCSEditor().IsValid())
+		else if (Blueprint->SkeletonGeneratedClass && BlueprintEditor->GetSubobjectEditor().IsValid())
 		{
-			TArray<FSCSEditorTreeNodePtrType> Nodes = BlueprintEditor->GetSCSEditor()->GetSelectedNodes();
-			if (Nodes.Num() == 1 && Nodes[0]->GetNodeType() == FSCSEditorTreeNode::ComponentNode)
+			TArray<FSubobjectEditorTreeNodePtrType> Nodes = BlueprintEditor->GetSubobjectEditor()->GetSelectedNodes();
+			if (Nodes.Num() == 1 && Nodes[0]->IsComponentNode())
 			{
 				FName PropertyName = Nodes[0]->GetVariableName();
 				FObjectProperty* VariableProperty = FindFProperty<FObjectProperty>(Blueprint->SkeletonGeneratedClass, PropertyName);
diff --git a/Engine/Source/Editor/Kismet/Private/SCSEditorViewportClient.cpp b/Engine/Source/Editor/Kismet/Private/SCSEditorViewportClient.cpp
--- a/Engine/Source/Editor/Kismet/Private/SCSEditorViewportClient.cpp
+++ b/Engine/Source/Editor/Kismet/Private/SCSEditorViewportClient.cpp
@@ -17,7 +17,7 @@
 #include "SEditorViewport.h"
 #include "EngineGlobals.h"
 #include "Editor.h"
-#include "SSCSEditor.h"
+#include "SSubobjectEditor.h"
 #include "Kismet2/BlueprintEditorUtils.h"
 #include "SKismetInspector.h"
 #include "ScopedTransaction.h"
@@ -80,13 +80,14 @@ namespace
 	}
 
 	// Determine whether or not the given node has a parent node that is not the root node, is movable and is selected
-	bool IsMovableParentNodeSelected(const FSCSEditorTreeNodePtrType& NodePtr, const TArray<FSCSEditorTreeNodePtrType>& SelectedNodes)
+	bool IsMovableParentNodeSelected(const FSubobjectEditorTreeNodePtrType& NodePtr, const TArray<FSubobjectEditorTreeNodePtrType>& SelectedNodes)
 	{
 		if(NodePtr.IsValid())
 		{
 			// Check for a valid parent node
-			FSCSEditorTreeNodePtrType ParentNodePtr = NodePtr->GetParent();
-			if(ParentNodePtr.IsValid() && !ParentNodePtr->IsRootComponent())
+			FSubobjectEditorTreeNodePtrType ParentNodePtr = NodePtr->GetParent();
+			const FSubobjectData* ParentData = ParentNodePtr.IsValid()  ? ParentNodePtr->GetDataSource() : nullptr;
+			if(ParentData && !ParentData->IsRootComponent())
 			{
 				if(SelectedNodes.Contains(ParentNodePtr))
 				{
@@ -166,7 +167,7 @@ void FSCSEditorViewportClient::Tick(float DeltaSeconds)
 	FEditorViewportClient::Tick(DeltaSeconds);
 
 	// Register the selection override delegate for the preview actor's components
-	TSharedPtr<SSCSEditor> SCSEditor = BlueprintEditorPtr.Pin()->GetSCSEditor();
+	TSharedPtr<SSubobjectEditor> SubobjectEditor = BlueprintEditorPtr.Pin()->GetSubobjectEditor();
 	AActor* PreviewActor = GetPreviewActor();
 	if (PreviewActor != nullptr)
 	{
@@ -176,7 +177,7 @@ void FSCSEditorViewportClient::Tick(float DeltaSeconds)
 			{
 				if (!PrimComponent->SelectionOverrideDelegate.IsBound())
 				{
-					SCSEditor->SetSelectionOverride(PrimComponent);
+					SubobjectEditor->SetSelectionOverride(PrimComponent);
 				}
 			}
 		}
@@ -215,15 +216,14 @@ void FSCSEditorViewportClient::Draw(const FSceneView* View, FPrimitiveDrawInterf
 	AActor* PreviewActor = GetPreviewActor();
 	if(PreviewActor)
 	{
-		if(GUnrealEd != NULL)
+		if(GUnrealEd != nullptr)
 		{
-			TArray<FSCSEditorTreeNodePtrType> SelectedNodes = BlueprintEditorPtr.Pin()->GetSelectedSCSEditorTreeNodes();
-			for (int32 SelectionIndex = 0; SelectionIndex < SelectedNodes.Num(); ++SelectionIndex)
+			TArray<FSubobjectEditorTreeNodePtrType> SelectedNodes = BlueprintEditorPtr.Pin()->GetSelectedSubobjectEditorTreeNodes();
+			for(FSubobjectEditorTreeNodePtrType SelectedNode : SelectedNodes)
 			{
-				FSCSEditorTreeNodePtrType SelectedNode = SelectedNodes[SelectionIndex];
-
-				UActorComponent* Comp = SelectedNode->FindComponentInstanceInActor(PreviewActor);
-				if(Comp != NULL && Comp->IsRegistered())
+				const FSubobjectData* Data = SelectedNode->GetDataSource();
+				const UActorComponent* Comp = Data ? Data->FindComponentInstanceInActor(PreviewActor) : nullptr;
+				if(Comp != nullptr && Comp->IsRegistered())
 				{
 					// Try and find a visualizer
 					TSharedPtr<FComponentVisualizer> Visualizer = GUnrealEd->FindComponentVisualizer(Comp->GetClass());
@@ -242,15 +242,15 @@ void FSCSEditorViewportClient::DrawCanvas( FViewport& InViewport, FSceneView& Vi
 	AActor* PreviewActor = GetPreviewActor();
 	if(PreviewActor)
 	{
-		if (GUnrealEd != NULL)
+		if (GUnrealEd != nullptr)
 		{
-			TArray<FSCSEditorTreeNodePtrType> SelectedNodes = BlueprintEditorPtr.Pin()->GetSelectedSCSEditorTreeNodes();
-			for (int32 SelectionIndex = 0; SelectionIndex < SelectedNodes.Num(); ++SelectionIndex)
+			TArray<FSubobjectEditorTreeNodePtrType> SelectedNodes = BlueprintEditorPtr.Pin()->GetSelectedSubobjectEditorTreeNodes();
+			for(FSubobjectEditorTreeNodePtrType SelectedNode : SelectedNodes)
 			{
-				FSCSEditorTreeNodePtrType SelectedNode = SelectedNodes[SelectionIndex];
-
-				UActorComponent* Comp = SelectedNode->FindComponentInstanceInActor(PreviewActor);
-				if (Comp != NULL && Comp->IsRegistered())
+				const FSubobjectData* Data = SelectedNode->GetDataSource();
+				const UActorComponent* Comp = Data ? Data->FindComponentInstanceInActor(PreviewActor) : nullptr;
+								
+				if (Comp != nullptr && Comp->IsRegistered())
 				{
 					// Try and find a visualizer
 					TSharedPtr<FComponentVisualizer> Visualizer = GUnrealEd->FindComponentVisualizer(Comp->GetClass());
@@ -267,10 +267,12 @@ void FSCSEditorViewportClient::DrawCanvas( FViewport& InViewport, FSceneView& Vi
 		const int32 HalfX = 0.5f * Viewport->GetSizeXY().X;
 		const int32 HalfY = 0.5f * Viewport->GetSizeXY().Y;
 
-		TArray<TSharedPtr<FSCSEditorTreeNode>> SelectedNodes = BlueprintEditorPtr.Pin()->GetSelectedSCSEditorTreeNodes();
+		TArray<FSubobjectEditorTreeNodePtrType> SelectedNodes = BlueprintEditorPtr.Pin()->GetSelectedSubobjectEditorTreeNodes();
+
 		if(bIsManipulating && SelectedNodes.Num() > 0)
 		{
-			USceneComponent* SceneComp = Cast<USceneComponent>(SelectedNodes[0]->FindComponentInstanceInActor(PreviewActor));
+			const FSubobjectData* Data = SelectedNodes[0]->GetDataSource();
+			const USceneComponent* SceneComp = Data ? Cast<USceneComponent>(Data->FindComponentInstanceInActor(PreviewActor)) : nullptr;
 			if(SceneComp)
 			{
 				const FVector WidgetLocation = GetWidgetLocation();
@@ -308,7 +310,7 @@ void FSCSEditorViewportClient::ProcessClick(class FSceneView& View, class HHitPr
 		{
 			HInstancedStaticMeshInstance* InstancedStaticMeshInstanceProxy = ((HInstancedStaticMeshInstance*)HitProxy);
 
-			TSharedPtr<ISCSEditorCustomization> Customization = BlueprintEditorPtr.Pin()->CustomizeSCSEditor(InstancedStaticMeshInstanceProxy->Component);
+			TSharedPtr<ISCSEditorCustomization> Customization = BlueprintEditorPtr.Pin()->CustomizeSubobjectEditor(InstancedStaticMeshInstanceProxy->Component);
 			if (Customization.IsValid() && Customization->HandleViewportClick(AsShared(), View, HitProxy, Key, Event, HitX, HitY))
 			{
 				Invalidate();
@@ -382,14 +384,14 @@ void FSCSEditorViewportClient::ProcessClick(class FSceneView& View, class HHitPr
 
 				if (SelectedCompInstance)
 				{
-					TSharedPtr<ISCSEditorCustomization> Customization = BlueprintEditorPtr.Pin()->CustomizeSCSEditor(SelectedCompInstance);
+					TSharedPtr<ISCSEditorCustomization> Customization = BlueprintEditorPtr.Pin()->CustomizeSubobjectEditor(SelectedCompInstance);
 					if (!(Customization.IsValid() && Customization->HandleViewportClick(AsShared(), View, HitProxy, Key, Event, HitX, HitY)))
 					{
 						const bool bIsCtrlKeyDown = Viewport->KeyState(EKeys::LeftControl) || Viewport->KeyState(EKeys::RightControl);
 						if (BlueprintEditorPtr.IsValid())
 						{
 							// Note: This will find and select any node associated with the component instance that's attached to the proxy (including visualizers)
-							BlueprintEditorPtr.Pin()->FindAndSelectSCSEditorTreeNode(SelectedCompInstance, bIsCtrlKeyDown);
+							BlueprintEditorPtr.Pin()->FindAndSelectSubobjectEditorTreeNode(SelectedCompInstance, bIsCtrlKeyDown);
 						}
 					}
 				}
@@ -413,7 +415,7 @@ bool FSCSEditorViewportClient::InputWidgetDelta( FViewport* InViewport, EAxisLis
 		TSharedPtr<FBlueprintEditor> BlueprintEditor = BlueprintEditorPtr.Pin();
 		if (PreviewActor && BlueprintEditor.IsValid())
 		{
-			TArray<FSCSEditorTreeNodePtrType> SelectedNodes = BlueprintEditor->GetSelectedSCSEditorTreeNodes();
+			TArray<FSubobjectEditorTreeNodePtrType> SelectedNodes = BlueprintEditorPtr.Pin()->GetSelectedSubobjectEditorTreeNodes();
 			if(SelectedNodes.Num() > 0)
 			{
 				FVector ModifiedScale = Scale;
@@ -431,11 +433,12 @@ bool FSCSEditorViewportClient::InputWidgetDelta( FViewport* InViewport, EAxisLis
 					ModifiedScale = FVector::ZeroVector;
 				}
 
-				for (const FSCSEditorTreeNodePtrType& SelectedNodePtr : SelectedNodes)
+				for (const FSubobjectEditorTreeNodePtrType& SelectedNodePtr : SelectedNodes)
 				{
+					const FSubobjectData* Data = SelectedNodePtr->GetDataSource();
 					// Don't allow editing of a root node, inherited SCS node or child node that also has a movable (non-root) parent node selected
 					const bool bCanEdit = GUnrealEd->ComponentVisManager.IsActive() ||
-						(!SelectedNodePtr->IsRootComponent() && !IsMovableParentNodeSelected(SelectedNodePtr, SelectedNodes));
+						(Data && !Data->IsRootComponent() && !IsMovableParentNodeSelected(SelectedNodePtr, SelectedNodes));
 
 					if(bCanEdit)
 					{
@@ -445,10 +448,10 @@ bool FSCSEditorViewportClient::InputWidgetDelta( FViewport* InViewport, EAxisLis
 							Invalidate();
 							return true;
 						}
-						
-						USceneComponent* SceneComp = Cast<USceneComponent>(SelectedNodePtr->FindComponentInstanceInActor(PreviewActor));
-						USceneComponent* SelectedTemplate = Cast<USceneComponent>(SelectedNodePtr->GetOrCreateEditableComponentTemplate(BlueprintEditor->GetBlueprintObj()));
-						if(SceneComp != NULL && SelectedTemplate != NULL)
+						// #TODO_BH Clean up const casts
+						USceneComponent* SceneComp = const_cast<USceneComponent*>(Cast<USceneComponent>(Data->FindComponentInstanceInActor(PreviewActor)));
+						USceneComponent* SelectedTemplate = const_cast<USceneComponent*>(Cast<USceneComponent>(Data->GetObjectForBlueprint(BlueprintEditor->GetBlueprintObj())));
+						if(SceneComp && SelectedTemplate)
 						{
 							// Cache the current default values for propagation
 							FVector OldRelativeLocation = SelectedTemplate->GetRelativeLocation();
@@ -458,7 +461,7 @@ bool FSCSEditorViewportClient::InputWidgetDelta( FViewport* InViewport, EAxisLis
 							// Adjust the deltas as necessary
 							FComponentEditorUtils::AdjustComponentDelta(SceneComp, Drag, Rot);
 
-							TSharedPtr<ISCSEditorCustomization> Customization = BlueprintEditor->CustomizeSCSEditor(SceneComp);
+							TSharedPtr<ISCSEditorCustomization> Customization = BlueprintEditor->CustomizeSubobjectEditor(SceneComp);
 							if(Customization.IsValid() && Customization->HandleViewportDrag(SceneComp, SelectedTemplate, Drag, Rot, ModifiedScale, GetWidgetLocation()))
 							{
 								// Handled by SCS Editor customization
@@ -477,7 +480,7 @@ bool FSCSEditorViewportClient::InputWidgetDelta( FViewport* InViewport, EAxisLis
 							}
 
 							UBlueprint* PreviewBlueprint = UBlueprint::GetBlueprintFromClass(PreviewActor->GetClass());
-							if(PreviewBlueprint != NULL)
+							if(PreviewBlueprint != nullptr)
 							{
 								// Like PostEditMove(), but we only need to re-run construction scripts
 								if(PreviewBlueprint && PreviewBlueprint->bRunConstructionScriptOnDrag)
@@ -585,12 +588,13 @@ UE::Widget::EWidgetMode FSCSEditorViewportClient::GetWidgetMode() const
 	if(!bIsSimulateEnabled && PreviewActor)
 	{
 		const TSharedPtr<FBlueprintEditor> BluePrintEditor = BlueprintEditorPtr.Pin();
-		if ( BluePrintEditor.IsValid() )
+		if (BluePrintEditor.IsValid())
 		{
-			TArray<FSCSEditorTreeNodePtrType> SelectedNodes = BluePrintEditor->GetSelectedSCSEditorTreeNodes();
-			if (BluePrintEditor->GetSCSEditor()->GetActorNode().IsValid())
+			TArray<FSubobjectEditorTreeNodePtrType> SelectedNodes = BlueprintEditorPtr.Pin()->GetSelectedSubobjectEditorTreeNodes();
+
+			if (BluePrintEditor->GetSubobjectEditor()->GetSceneRootNode().IsValid())
 			{
-				const TArray<FSCSEditorTreeNodePtrType>& RootNodes = BluePrintEditor->GetSCSEditor()->GetActorNode()->GetComponentNodes();
+				TArray<FSubobjectEditorTreeNodePtrType> RootNodes = BluePrintEditor->GetSubobjectEditor()->GetRootNodes();
 
 				if (GUnrealEd->ComponentVisManager.IsActive() &&
 					GUnrealEd->ComponentVisManager.IsVisualizingArchetype())
@@ -602,19 +606,26 @@ UE::Widget::EWidgetMode FSCSEditorViewportClient::GetWidgetMode() const
 				{
 					// if the selected nodes array is empty, or only contains entries from the
 					// root nodes array, or isn't visible in the preview actor, then don't display a transform widget
-					for (int32 CurrentNodeIndex = 0; CurrentNodeIndex < SelectedNodes.Num(); CurrentNodeIndex++)
+					for(FSubobjectEditorTreeNodePtrType CurrentNodePtr : SelectedNodes)
 					{
-						FSCSEditorTreeNodePtrType CurrentNodePtr = SelectedNodes[CurrentNodeIndex];
-						if ((CurrentNodePtr.IsValid() &&
-							((!RootNodes.Contains(CurrentNodePtr) && !CurrentNodePtr->IsRootComponent()) ||
-							(CurrentNodePtr->GetObject<UInstancedStaticMeshComponent>() && // show widget if we are editing individual instances even if it is the root component
-								CastChecked<UInstancedStaticMeshComponent>(CurrentNodePtr->FindComponentInstanceInActor(GetPreviewActor()))->SelectedInstances.Contains(true))) &&
-							CurrentNodePtr->CanEdit() &&
-							CurrentNodePtr->FindComponentInstanceInActor(PreviewActor)))
+						if (CurrentNodePtr.IsValid())
 						{
-							// a non-NULL, non-root item is selected, draw the widget
-							ReturnWidgetMode = WidgetMode;
-							break;
+							FSubobjectData* Data = CurrentNodePtr->GetDataSource();
+							if (Data && Data->CanEdit())
+							{
+								const bool bIsNotRootComponent = !RootNodes.Contains(CurrentNodePtr) && !Data->IsRootComponent();
+								const bool bIsISM = 
+									Data->GetObject<UInstancedStaticMeshComponent>() && 
+									CastChecked<UInstancedStaticMeshComponent>(Data->FindComponentInstanceInActor(GetPreviewActor()))->SelectedInstances.Contains(true);
+								const bool bHasInstanceInActor = Data->FindComponentInstanceInActor(PreviewActor) != nullptr;
+
+								if ((bIsNotRootComponent || bIsISM) && bHasInstanceInActor)
+								{
+									// a non-NULL, non-root item is selected, draw the widget
+									ReturnWidgetMode = WidgetMode;
+									break;
+								}
+							}
 						}
 					}
 				}
@@ -650,16 +661,18 @@ FVector FSCSEditorViewportClient::GetWidgetLocation() const
 	AActor* PreviewActor = GetPreviewActor();
 	if(PreviewActor)
 	{
-		TArray<FSCSEditorTreeNodePtrType> SelectedNodes = BlueprintEditorPtr.Pin()->GetSelectedSCSEditorTreeNodes();
+		TArray<FSubobjectEditorTreeNodePtrType> SelectedNodes = BlueprintEditorPtr.Pin()->GetSelectedSubobjectEditorTreeNodes();
 		if(SelectedNodes.Num() > 0)
 		{
 			// Use the last selected item for the widget location
-			USceneComponent* SceneComp = Cast<USceneComponent>(SelectedNodes.Last().Get()->FindComponentInstanceInActor(PreviewActor));
-			if( SceneComp )
+			const FSubobjectData* Data = SelectedNodes[0]->GetDataSource();
+
+			const USceneComponent* SceneComp = Cast<USceneComponent>(Data->FindComponentInstanceInActor(PreviewActor));
+			if(SceneComp)
 			{
-				TSharedPtr<ISCSEditorCustomization> Customization = BlueprintEditorPtr.Pin()->CustomizeSCSEditor(SceneComp);
+				TSharedPtr<ISCSEditorCustomization> Customization = BlueprintEditorPtr.Pin()->CustomizeSubobjectEditor(SceneComp);
 				FVector CustomLocation;
-				if(Customization.IsValid() && Customization->HandleGetWidgetLocation(SceneComp, CustomLocation))
+				if(Customization.IsValid() && Customization->HandleGetWidgetLocation(const_cast<USceneComponent*>(SceneComp), CustomLocation))
 				{
 					Location = CustomLocation;
 				}
@@ -690,14 +703,15 @@ FMatrix FSCSEditorViewportClient::GetWidgetCoordSystem() const
 		TSharedPtr<FBlueprintEditor> BlueprintEditor = BlueprintEditorPtr.Pin();
 		if (PreviewActor && BlueprintEditor.IsValid())
 		{
-			TArray<FSCSEditorTreeNodePtrType> SelectedNodes = BlueprintEditor->GetSelectedSCSEditorTreeNodes();
+			TArray<FSubobjectEditorTreeNodePtrType> SelectedNodes = BlueprintEditorPtr.Pin()->GetSelectedSubobjectEditorTreeNodes();
 			if(SelectedNodes.Num() > 0)
 			{
-				const FSCSEditorTreeNodePtrType SelectedNode = SelectedNodes.Last();
-				USceneComponent* SceneComp = SelectedNode.IsValid() ? Cast<USceneComponent>(SelectedNode->FindComponentInstanceInActor(PreviewActor)) : NULL;
-				if( SceneComp )
+				const FSubobjectEditorTreeNodePtrType SelectedNode = SelectedNodes.Last();
+				const FSubobjectData* Data = SelectedNode->GetDataSource();
+				const USceneComponent* SceneComp = SelectedNode.IsValid() ? Cast<USceneComponent>(Data->FindComponentInstanceInActor(PreviewActor)) : nullptr;
+				if(SceneComp)
 				{
-					TSharedPtr<ISCSEditorCustomization> Customization = BlueprintEditor->CustomizeSCSEditor(SceneComp);
+					TSharedPtr<ISCSEditorCustomization> Customization = BlueprintEditor->CustomizeSubobjectEditor(SceneComp);
 					FMatrix CustomTransform;
 					if(Customization.IsValid() && Customization->HandleGetWidgetTransform(SceneComp, CustomTransform))
 					{
@@ -808,19 +822,20 @@ void FSCSEditorViewportClient::FocusViewportToSelection()
 	AActor* PreviewActor = GetPreviewActor();
 	if(PreviewActor)
 	{
-		TArray<FSCSEditorTreeNodePtrType> SelectedNodes = BlueprintEditorPtr.Pin()->GetSelectedSCSEditorTreeNodes();
+		TArray<FSubobjectEditorTreeNodePtrType> SelectedNodes = BlueprintEditorPtr.Pin()->GetSelectedSubobjectEditorTreeNodes();
 		if(SelectedNodes.Num() > 0)
 		{
+			const FSubobjectData* Data = SelectedNodes[0]->GetDataSource();
 			// Use the last selected item for the widget location
-			USceneComponent* SceneComp = Cast<USceneComponent>(SelectedNodes.Last()->FindComponentInstanceInActor(PreviewActor));
-			if( SceneComp )
+			const USceneComponent* SceneComp = Data ? Cast<USceneComponent>(Data->FindComponentInstanceInActor(PreviewActor)) : nullptr;
+			if(SceneComp)
 			{
-				FocusViewportOnBox( SceneComp->Bounds.GetBox() );
+				FocusViewportOnBox(SceneComp->Bounds.GetBox());
 			}
 		}
 		else
 		{
-			FocusViewportOnBox( PreviewActor->GetComponentsBoundingBox( true ) );
+			FocusViewportOnBox(PreviewActor->GetComponentsBoundingBox(true));
 		}
 	}
 }
@@ -839,13 +854,13 @@ void FSCSEditorViewportClient::ToggleIsSimulateEnabled()
 	PreviewScene->GetWorld()->bBegunPlay = bIsSimulateEnabled;
 	PreviewScene->GetWorld()->bShouldSimulatePhysics = bIsSimulateEnabled;
 
-	TSharedPtr<SWidget> SCSEditor = BlueprintEditorPtr.Pin()->GetSCSEditor();
+	TSharedPtr<SWidget> SubobjectEditor = BlueprintEditorPtr.Pin()->GetSubobjectEditor();
 	TSharedRef<SWidget> Inspector = BlueprintEditorPtr.Pin()->GetInspector();
 
 	// When simulate is enabled, we don't want to allow the user to modify the components
 	BlueprintEditorPtr.Pin()->UpdatePreviewActor(BlueprintEditorPtr.Pin()->GetBlueprintObj(), true);
 
-	SCSEditor->SetEnabled(!bIsSimulateEnabled);
+	SubobjectEditor->SetEnabled(!bIsSimulateEnabled);
 	Inspector->SetEnabled(!bIsSimulateEnabled);
 
 	if(!IsRealtime())
@@ -900,42 +915,13 @@ void FSCSEditorViewportClient::BeginTransaction(const FText& Description)
 	//UE_LOG(LogSCSEditorViewport, Log, TEXT("FSCSEditorViewportClient::BeginTransaction() pre: %s %08x"), SessionName, *((uint32*)&ScopedTransaction));
 
 	if(!ScopedTransaction)
+	
 	{
-		ScopedTransaction = new FScopedTransaction(Description);
-
-		TSharedPtr<FBlueprintEditor> BlueprintEditor = BlueprintEditorPtr.Pin();
-		if (BlueprintEditor.IsValid())
+		if(USubobjectDataSubsystem* System = USubobjectDataSubsystem::Get())
 		{
-			UBlueprint* PreviewBlueprint = BlueprintEditor->GetBlueprintObj();
-			if (PreviewBlueprint != nullptr)
-			{
-				FBlueprintEditorUtils::MarkBlueprintAsModified(PreviewBlueprint);
-			}
-
-			TArray<FSCSEditorTreeNodePtrType> SelectedNodes = BlueprintEditor->GetSelectedSCSEditorTreeNodes();
-			for (const FSCSEditorTreeNodePtrType& Node : SelectedNodes)
-			{
-				if(Node.IsValid())
-				{
-					if(USCS_Node* SCS_Node = Node->GetSCSNode())
-					{
-						USimpleConstructionScript* SCS = SCS_Node->GetSCS();
-						UBlueprint* Blueprint = SCS ? SCS->GetBlueprint() : nullptr;
-						if (Blueprint == PreviewBlueprint)
-						{
-							SCS_Node->Modify();
-						}
-					}
-
-					// Modify template, any instances will be reconstructed as part of PostUndo:
-					UActorComponent* ComponentTemplate = Node->GetOrCreateEditableComponentTemplate(PreviewBlueprint);
-					if (ComponentTemplate != nullptr)
-					{
-						ComponentTemplate->SetFlags(RF_Transactional);
-						ComponentTemplate->Modify();
-					}
-				}
-			}
+			TArray<FSubobjectDataHandle> SelectedNodes = BlueprintEditorPtr.Pin()->GetSubobjectEditor()->GetSelectedHandles();
+			TSharedPtr<FBlueprintEditor> BlueprintEditor = BlueprintEditorPtr.Pin();
+			ScopedTransaction = System->BeginTransaction(SelectedNodes, Description, BlueprintEditor.IsValid() ? BlueprintEditor->GetBlueprintObj() : nullptr);
 		}
 	}
 
@@ -949,7 +935,7 @@ void FSCSEditorViewportClient::EndTransaction()
 	if(ScopedTransaction)
 	{
 		delete ScopedTransaction;
-		ScopedTransaction = NULL;
+		ScopedTransaction = nullptr;
 	}
 }
 
diff --git a/Engine/Source/Editor/Kismet/Private/SSCSEditorViewport.cpp b/Engine/Source/Editor/Kismet/Private/SSCSEditorViewport.cpp
--- a/Engine/Source/Editor/Kismet/Private/SSCSEditorViewport.cpp
+++ b/Engine/Source/Editor/Kismet/Private/SSCSEditorViewport.cpp
@@ -6,7 +6,7 @@
 #include "Framework/MultiBox/MultiBoxBuilder.h"
 #include "EditorStyleSet.h"
 #include "BlueprintEditorCommands.h"
-#include "SSCSEditor.h"
+#include "SSubobjectEditor.h"
 #include "Slate/SceneViewport.h"
 #include "SViewportToolBar.h"
 #include "STransformViewportToolbar.h"
@@ -283,8 +283,9 @@ void SSCSEditorViewport::BindCommands()
 	FSCSEditorViewportCommands::Register(); // make sure the viewport specific commands have been registered
 
 	TSharedPtr<FBlueprintEditor> BlueprintEditor = BlueprintEditorPtr.Pin();
-	TSharedPtr<SSCSEditor> SCSEditorWidgetPtr = BlueprintEditor->GetSCSEditor();
-	SSCSEditor* SCSEditorWidget = SCSEditorWidgetPtr.Get();
+	TSharedPtr<SSubobjectEditor> SubobjectEditorPtr = BlueprintEditor->GetSubobjectEditor();
+	SSubobjectEditor* SubobjectEditorWidget = SubobjectEditorPtr.Get(); 
+	
 	// for mac, we have to bind a command that would override the BP-Editor's 
 	// "NavigateToParentBackspace" command, because the delete key is the 
 	// backspace key for that platform (and "NavigateToParentBackspace" does not 
@@ -292,13 +293,18 @@ void SSCSEditorViewport::BindCommands()
 	// 
 	// NOTE: this needs to come before we map any other actions (so it is 
 	// prioritized first)
-	CommandList->MapAction(
-		FSCSEditorViewportCommands::Get().DeleteComponent,
-		FExecuteAction::CreateSP(SCSEditorWidget, &SSCSEditor::OnDeleteNodes),
-		FCanExecuteAction::CreateSP(SCSEditorWidget, &SSCSEditor::CanDeleteNodes)
-	);
 
-	CommandList->Append(BlueprintEditor->GetSCSEditor()->CommandList.ToSharedRef());
+	if(SubobjectEditorWidget)
+	{
+		CommandList->MapAction(
+		    FSCSEditorViewportCommands::Get().DeleteComponent,
+		    FExecuteAction::CreateSP(SubobjectEditorWidget, &SSubobjectEditor::OnDeleteNodes),
+		    FCanExecuteAction::CreateSP(SubobjectEditorWidget, &SSubobjectEditor::CanDeleteNodes)
+		);
+		
+		CommandList->Append(SubobjectEditorWidget->GetCommandList().ToSharedRef());
+	}
+	
 	CommandList->Append(BlueprintEditor->GetToolkitCommands());
 	SEditorViewport::BindCommands();
 
@@ -418,8 +424,9 @@ TSharedPtr<SDockTab> SSCSEditorViewport::GetOwnerTab() const
 
 FReply SSCSEditorViewport::OnDrop(const FGeometry& MyGeometry, const FDragDropEvent& DragDropEvent)
 {
-	TSharedPtr<SSCSEditor> SCSEditor = BlueprintEditorPtr.Pin()->GetSCSEditor();
-	return SCSEditor->TryHandleAssetDragDropOperation(DragDropEvent);
+	TSharedPtr<SSubobjectEditor> SubobjectEditor = BlueprintEditorPtr.Pin()->GetSubobjectEditor();
+
+	return SubobjectEditor->TryHandleAssetDragDropOperation(DragDropEvent);
 }
 
 EActiveTimerReturnType SSCSEditorViewport::DeferredUpdatePreview(double InCurrentTime, float InDeltaTime, bool bResetCamera)
diff --git a/Engine/Source/Editor/Kismet/Private/Tests/BlueprintEditorTests.cpp b/Engine/Source/Editor/Kismet/Private/Tests/BlueprintEditorTests.cpp
--- a/Engine/Source/Editor/Kismet/Private/Tests/BlueprintEditorTests.cpp
+++ b/Engine/Source/Editor/Kismet/Private/Tests/BlueprintEditorTests.cpp
@@ -221,7 +221,7 @@ namespace BlueprintEditorPromotionUtils
 		// Recompile skeleton because of the new component we added
 		FKismetEditorUtilities::GenerateBlueprintSkeleton(InBlueprint, true);
 
-		CurrentBlueprintEditor->UpdateSCSPreview(true);
+		CurrentBlueprintEditor->UpdateSubobjectPreview(true);
 
 		return NewNode;
 	}
diff --git a/Engine/Source/Editor/Kismet/Public/BlueprintEditor.h b/Engine/Source/Editor/Kismet/Public/BlueprintEditor.h
--- a/Engine/Source/Editor/Kismet/Public/BlueprintEditor.h
+++ b/Engine/Source/Editor/Kismet/Public/BlueprintEditor.h
@@ -52,6 +52,8 @@ class UK2Node_FunctionEntry;
 class UK2Node_Event;
 class UToolMenu;
 class FBlueprintNamespaceHelper;
+class SSubobjectEditor;
+class FSubobjectEditorTreeNode;
 
 /* Enums to use when grouping the blueprint members in the list panel. The order here will determine the order in the list */
 namespace NodeSectionID
@@ -90,7 +92,7 @@ public:
 
 public:
 	FCustomDebugObject()
-		: Object(NULL)
+		: Object(nullptr)
 	{
 	}
 
@@ -226,8 +228,15 @@ public:
 	virtual void SummonSearchUI(bool bSetFindWithinBlueprint, FString NewSearchTerms = FString(), bool bSelectFirstResult = false) override;
 	virtual void SummonFindAndReplaceUI() override;
 	virtual TSharedPtr<SGraphEditor> OpenGraphAndBringToFront(UEdGraph* Graph, bool bSetFocus = true) override;
-	virtual TArray<TSharedPtr<class FSCSEditorTreeNode> > GetSelectedSCSEditorTreeNodes() const override;
-	virtual TSharedPtr<class FSCSEditorTreeNode> FindAndSelectSCSEditorTreeNode(const UActorComponent* InComponent, bool IsCntrlDown) override;
+
+	UE_DEPRECATED(5.0, "GetSelectedSCSEditorTreeNodes has been deprecated. Use GetSelectedSubobjectEditorTreeNodes instead.")
+	virtual TArray<TSharedPtr<class FSCSEditorTreeNode>> GetSelectedSCSEditorTreeNodes() const override;
+	
+	UE_DEPRECATED(5.0, "FindAndSelectSCSEditorTreeNode has been deprecated. Use FindAndSelectSubobjectEditorTreeNode instead.")
+	virtual TSharedPtr<class FSCSEditorTreeNode> FindAndSelectSCSEditorTreeNode(const UActorComponent* InComponent, bool IsCntrlDown) override { return nullptr; }
+
+	virtual TArray<TSharedPtr<FSubobjectEditorTreeNode>> GetSelectedSubobjectEditorTreeNodes() const override;
+	virtual TSharedPtr<FSubobjectEditorTreeNode> FindAndSelectSubobjectEditorTreeNode(const UActorComponent* InComponent, bool IsCntrlDown) override;
 	virtual int32 GetNumberOfSelectedNodes() const override;
 	virtual void AnalyticsTrackNodeEvent(UBlueprint* Blueprint, UEdGraphNode *GraphNode, bool bNodeDelete = false) const override;
 	void AnalyticsTrackCompileEvent(UBlueprint* Blueprint, int32 NumErrors, int32 NumWarnings) const;
@@ -260,19 +269,22 @@ public:
 	bool IsEditingSingleBlueprint() const;
 
 	/** Getters for the various Kismet2 widgets */
-	TSharedRef<class SKismetInspector> GetInspector() const { return Inspector.ToSharedRef(); }
-	TSharedRef<class SKismetInspector> GetDefaultEditor() const { return DefaultEditor.ToSharedRef(); }
-	TSharedRef<class SKismetDebuggingView> GetDebuggingView() const { return DebuggingView.ToSharedRef(); }
-	TSharedRef<class SBlueprintPalette> GetPalette() const { return Palette.ToSharedRef(); }
-	TSharedRef<class SBlueprintBookmarks> GetBookmarksWidget() const { return BookmarksWidget.ToSharedRef(); }
-	TSharedRef<class SWidget> GetCompilerResults() const { return CompilerResults.ToSharedRef(); }
-	TSharedRef<class SFindInBlueprints> GetFindResults() const { return FindResults.ToSharedRef(); }
+	TSharedRef<SKismetInspector> GetInspector() const { return Inspector.ToSharedRef(); }
+	TSharedRef<SKismetInspector> GetDefaultEditor() const { return DefaultEditor.ToSharedRef(); }
+	TSharedRef<SKismetDebuggingView> GetDebuggingView() const { return DebuggingView.ToSharedRef(); }
+	TSharedRef<SBlueprintPalette> GetPalette() const { return Palette.ToSharedRef(); }
+	TSharedRef<SBlueprintBookmarks> GetBookmarksWidget() const { return BookmarksWidget.ToSharedRef(); }
+	TSharedRef<SWidget> GetCompilerResults() const { return CompilerResults.ToSharedRef(); }
+	TSharedRef<SFindInBlueprints> GetFindResults() const { return FindResults.ToSharedRef(); }
 
 	/** Getters for the various optional Kismet2 widgets */
-	TSharedPtr<class SSCSEditor> GetSCSEditor() const { return SCSEditor; }
-	TSharedPtr<class SSCSEditorViewport> GetSCSViewport() const { return SCSViewport; }
-	TSharedPtr<class SMyBlueprint> GetMyBlueprintWidget() const { return MyBlueprintWidget; }
-	TSharedPtr<class SReplaceNodeReferences> GetReplaceReferencesWidget() const { return ReplaceReferencesWidget; }
+	UE_DEPRECATED(5.0, "GetSCSEditor has been deprecated. Use GetSubobjectEditor instead.")
+	TSharedPtr<SSCSEditor> GetSCSEditor() const { return nullptr; }
+	
+	TSharedPtr<SSubobjectEditor> GetSubobjectEditor() const { return SubobjectEditor; }
+	TSharedPtr<SSCSEditorViewport> GetSubobjectViewport() const { return SubobjectViewport; }
+	TSharedPtr<SMyBlueprint> GetMyBlueprintWidget() const { return MyBlueprintWidget; }
+	TSharedPtr<SReplaceNodeReferences> GetReplaceReferencesWidget() const { return ReplaceReferencesWidget; }
 
 	/**
 	 * Provides access to the preview actor.
@@ -472,20 +484,30 @@ public:
 	/** Called when the graph editor tab is backgrounded */
 	virtual void OnGraphEditorBackgrounded(const TSharedRef<SGraphEditor>& InGraphEditor);
 
-	/** Enable/disable the SCS editor preview viewport */
-	void EnableSCSPreview(bool bEnable);
+	/** Enable/disable the subobject editor preview viewport */
+	void EnableSubobjectPreview(bool bEnable);
 
-	/** Refresh the preview viewport to reflect changes in the SCS */
-	void UpdateSCSPreview(bool bUpdateNow = false);
+	/** Refresh the preview viewport to reflect changes in the subobject */
+	void UpdateSubobjectPreview(bool bUpdateNow = false);
 
-	/** Delegate invoked when the SCS editor needs to obtain the Actor context for editing */
-	AActor* GetSCSEditorActorContext() const;
+	///////////////////////////////////////////////////////
+	// Subobject Editor Callbacks
+	/** Delegate invoked when the Subobject editor needs to obtain the object context for editing */
+	UObject* GetSubobjectEditorObjectContext() const;
 
-	/** Delegate invoked when the selection is changed in the SCS editor widget */
-	virtual void OnSelectionUpdated(const TArray<TSharedPtr<class FSCSEditorTreeNode>>& SelectedNodes);
+	/** Delegate invoked when the selection is changed in the subobject editor widget */
+	virtual void OnSelectionUpdated(const TArray<TSharedPtr<class FSubobjectEditorTreeNode>>& SelectedNodes);
 
-	/** Delegate invoked when an item is double clicked in the SCS editor widget */
-	virtual void OnComponentDoubleClicked(TSharedPtr<class FSCSEditorTreeNode> Node);
+	/** Delegate invoked when an item is double clicked in the subobject editor widget */
+	virtual void OnComponentDoubleClicked(TSharedPtr<class FSubobjectEditorTreeNode> Node);
+	
+	/** Delegate invoked when the selection is changed in the subobject editor widget */
+	UE_DEPRECATED(5.0, "OnSelectionUpdated(const TArray<TSharedPtr<class FSCSEditorTreeNode>>&) has been deprecated. Use OnSelectionUpdated(const TArray<TSharedPtr<class FSubobjectEditorTreeNode>>&) instead.")
+	virtual void OnSelectionUpdated(const TArray<TSharedPtr<class FSCSEditorTreeNode>>& SelectedNodes) {}
+
+	/** Delegate invoked when an item is double clicked in the subobject editor widget */
+	UE_DEPRECATED(5.0, "OnComponentDoubleClicked(TSharedPtr<class FSCSEditorTreeNode>) has been deprecated. Use OnComponentDoubleClicked(TSharedPtr<class FSubobjectEditorTreeNode>) instead.")
+	virtual void OnComponentDoubleClicked(TSharedPtr<class FSCSEditorTreeNode> Node) {}
 
 	/** Pin visibility accessors */
 	void SetPinVisibility(SGraphEditor::EPinVisibility Visibility);
@@ -573,18 +595,24 @@ public:
 	/** Sets customizations for the BP editor details panel. */
 	void SetDetailsCustomization(TSharedPtr<class FDetailsViewObjectFilter> DetailsObjectFilter, TSharedPtr<class IDetailRootObjectCustomization> DetailsRootCustomization);
 
-	/** Sets SCS editor UI customization */
-	void SetSCSEditorUICustomization(TSharedPtr<class ISCSEditorUICustomization> SCSEditorUICustomization);
+	/** Sets subobject editor UI customization */
+	void SetSubobjectEditorUICustomization(TSharedPtr<class ISCSEditorUICustomization> SCSEditorUICustomization);
+
+	UE_DEPRECATED(5.0, "SetSCSEditorUICustomization has been deprecated. Use SetSubobjectEditorUICustomization instead.")
+	void SetSCSEditorUICustomization(TSharedPtr<class ISCSEditorUICustomization> SCSEditorUICustomization)
+	{
+		return SetSubobjectEditorUICustomization(SCSEditorUICustomization);
+	}
 
 	/**
-	 * Register a customization for interacting with the SCS editor
+	 * Register a customization for interacting with the subobject editor
 	 * @param	InComponentName			The name of the component to customize behavior for
 	 * @param	InCustomization			The customization instance to use
 	 */
 	void RegisterSCSEditorCustomization(const FName& InComponentName, TSharedPtr<class ISCSEditorCustomization> InCustomization);
 
 	/**
-	 * Unregister a previously registered customization for interacting with the SCS editor
+	 * Unregister a previously registered customization for interacting with the subobject editor
 	 * @param	InComponentName			The name of the component to customize behavior for
 	 */
 	void UnregisterSCSEditorCustomization(const FName& InComponentName);
@@ -603,11 +631,17 @@ public:
 	void OnDumpCachedIndexDataForBlueprint();
 
 	/**
-	 * Check to see if we can customize the SCS editor for the passed-in scene component
+	 * Check to see if we can customize the subobject editor for the passed-in scene component
 	 * @param	InComponentToCustomize	The component to check to see if a customization exists
-	 * @return an SCS editor customization instance, if one exists.
+	 * @return an subobject editor customization instance, if one exists.
 	 */
-	TSharedPtr<class ISCSEditorCustomization> CustomizeSCSEditor(USceneComponent* InComponentToCustomize) const;
+	TSharedPtr<class ISCSEditorCustomization> CustomizeSubobjectEditor(const USceneComponent* InComponentToCustomize) const;
+
+	UE_DEPRECATED(5.0, "CustomizeSCSEditor has been deprecated. Use CustomizeSubobjectEditor instead.")
+	TSharedPtr<class ISCSEditorCustomization> CustomizeSCSEditor(USceneComponent* InComponentToCustomize) const
+	{
+		return CustomizeSubobjectEditor(InComponentToCustomize);
+	}
 
 	/**
 	 * Returns the currently focused graph in the Blueprint editor
@@ -784,7 +818,7 @@ protected:
 	void OpenBlueprintDebugger();
 	bool CanOpenBlueprintDebugger() const;
 
-	// Utility helper to get the currently hovered pin in the currently visible graph, or NULL if there isn't one
+	// Utility helper to get the currently hovered pin in the currently visible graph, or nullptr if there isn't one
 	UEdGraphPin* GetCurrentlySelectedPin() const;
 
 	// UI Action functionality
@@ -1143,8 +1177,8 @@ private:
 	/** Function to check whether the give graph is a subgraph */
 	static bool IsASubGraph( const class UEdGraph* GraphPtr );
 
-	/** Creates the SCSEditor tree component view and the SCS Viewport. */
-	void CreateSCSEditors();
+	/** Creates the Subobject Editor tree component view and the Subobject Viewport. */
+	void CreateSubobjectEditors();
 
 	/** Callback when a token is clicked on in the compiler results log */
 	void OnLogTokenClicked(const TSharedRef<class IMessageToken>& Token);
@@ -1241,10 +1275,14 @@ protected:
 	TArray<UBlueprint*> StandardLibraries;
 
 	/** SCS editor */
+	UE_DEPRECATED(5.0, "SCSEditor has been deprecated. Use SubobjectEditor instead.")
 	TSharedPtr<class SSCSEditor> SCSEditor;
 
+	/** Subobject Editor */
+	TSharedPtr<SSubobjectEditor> SubobjectEditor;
+
 	/** Viewport widget */
-	TSharedPtr<class SSCSEditorViewport> SCSViewport;
+	TSharedPtr<SSCSEditorViewport> SubobjectViewport;
 
 	/** Node inspector widget */
 	TSharedPtr<class SKismetInspector> Inspector;
@@ -1413,8 +1451,8 @@ private:
 	/** analytics statistics for the Editor */ 
 	FAnalyticsStatistics AnalyticsStats;
 
-	/** Customizations for the SCS editor */
-	TMap< FName, TSharedPtr<ISCSEditorCustomization> > SCSEditorCustomizations;
+	/** Customizations for the Subobject editor */
+	TMap<FName, TSharedPtr<ISCSEditorCustomization>> SubobjectEditorCustomizations;
 
 	/** Whether the current project is C++ or blueprint based */
 	bool bCodeBasedProject;
diff --git a/Engine/Source/Editor/Kismet/Public/BlueprintEditorModes.h b/Engine/Source/Editor/Kismet/Public/BlueprintEditorModes.h
--- a/Engine/Source/Editor/Kismet/Public/BlueprintEditorModes.h
+++ b/Engine/Source/Editor/Kismet/Public/BlueprintEditorModes.h
@@ -97,7 +97,7 @@ protected:
 	// Set of spawnable tabs in this mode
 	FWorkflowAllowedTabSet BlueprintComponentsTabFactories;
 
-	TArray<TWeakObjectPtr<UActorComponent>> CachedComponentSelection;
+	TArray<TWeakObjectPtr<const UActorComponent>> CachedComponentSelection;
 };
 
 class KISMET_API FBlueprintInterfaceApplicationMode : public FApplicationMode
diff --git a/Engine/Source/Editor/Kismet/Public/BlueprintEditorModule.h b/Engine/Source/Editor/Kismet/Public/BlueprintEditorModule.h
--- a/Engine/Source/Editor/Kismet/Public/BlueprintEditorModule.h
+++ b/Engine/Source/Editor/Kismet/Public/BlueprintEditorModule.h
@@ -21,6 +21,7 @@ class UUserDefinedStruct;
 class IDetailCustomization;
 class FKismetCompilerContext;
 struct FBlueprintDebugger;
+class FSubobjectEditorTreeNode;
 
 /** Delegate used to customize variable display */
 DECLARE_DELEGATE_RetVal_OneParam(TSharedPtr<IDetailCustomization>, FOnGetVariableCustomizationInstance, TSharedPtr<IBlueprintEditor> /*BlueprintEditor*/);
@@ -34,6 +35,7 @@ namespace ERefreshBlueprintEditorReason
 	enum Type
 	{
 		BlueprintCompiled,
+		PostUndo,
 		UnknownReason
 	};
 }
@@ -84,18 +86,22 @@ public:
 
 	virtual bool GetBoundsForSelectedNodes(class FSlateRect& Rect, float Padding ) = 0;
 
-	/** Util to get the currently selected SCS editor tree Nodes */
-	virtual TArray<TSharedPtr<class FSCSEditorTreeNode> >  GetSelectedSCSEditorTreeNodes() const = 0;
+	/** Util to get the currently selected Subobject editor tree Nodes */
+	virtual TArray<TSharedPtr<FSubobjectEditorTreeNode>> GetSelectedSubobjectEditorTreeNodes() const = 0;
 
 	/** Get number of currently selected nodes in the SCS editor tree */
 	virtual int32 GetNumberOfSelectedNodes() const = 0;
 
 	/** Find and select a specific SCS editor tree node associated with the given component */
-	virtual TSharedPtr<class FSCSEditorTreeNode> FindAndSelectSCSEditorTreeNode(const class UActorComponent* InComponent, bool IsCntrlDown) = 0;
+	virtual TSharedPtr<FSubobjectEditorTreeNode> FindAndSelectSubobjectEditorTreeNode(const class UActorComponent* InComponent, bool IsCntrlDown) = 0;
 
 	/** Used to track node create/delete events for Analytics */
 	virtual void AnalyticsTrackNodeEvent( UBlueprint* Blueprint, UEdGraphNode *GraphNode, bool bNodeDelete = false ) const = 0;
 
+	UE_DEPRECATED(5.0, "GetSelectedSCSEditorTreeNodes has been deprecated. Use GetSelectedSubobjectEditorTreeNodes instead.")
+	virtual TArray<TSharedPtr<class FSCSEditorTreeNode> >  GetSelectedSCSEditorTreeNodes() const = 0;
+	UE_DEPRECATED(5.0, "FindAndSelectSCSEditorTreeNode has been deprecated. Use FindAndSelectSubobjectEditorTreeNode instead.")
+	virtual TSharedPtr<class FSCSEditorTreeNode> FindAndSelectSCSEditorTreeNode(const class UActorComponent* InComponent, bool IsCntrlDown) = 0;
 };
 
 DECLARE_DELEGATE_RetVal_OneParam(TSharedRef<class ISCSEditorCustomization>, FSCSEditorCustomizationBuilder, TSharedRef< IBlueprintEditor > /* InBlueprintEditor */);
@@ -171,7 +177,7 @@ public:
 	virtual void SetDetailsCustomization(TSharedPtr<class FDetailsViewObjectFilter> InDetailsObjectFilter, TSharedPtr<class IDetailRootObjectCustomization> InDetailsRootCustomization);
 
 	/** Sets SCS editor UI customization */
-	virtual void SetSCSEditorUICustomization(TSharedPtr<class ISCSEditorUICustomization> InSCSEditorUICustomization);
+	virtual void SetSubobjectEditorUICustomization(TSharedPtr<class ISCSEditorUICustomization> InSCSEditorUICustomization);
 
 	/**
 	 * Register a customization for interacting with the SCS editor 
diff --git a/Engine/Source/Editor/Kismet/Public/ISCSEditorCustomization.h b/Engine/Source/Editor/Kismet/Public/ISCSEditorCustomization.h
--- a/Engine/Source/Editor/Kismet/Public/ISCSEditorCustomization.h
+++ b/Engine/Source/Editor/Kismet/Public/ISCSEditorCustomization.h
@@ -6,6 +6,8 @@
 #include "InputCoreTypes.h"
 #include "EditorViewportClient.h"
 
+class USceneComponent;
+
 /**
  * Customization class that allows per-component type handling of various SCS editor functionality.
  */
@@ -37,7 +39,7 @@ public:
 	 * @param	InPivot					The povot point for any transformations
 	 * @return true if the operation was handled.
 	 */
-	virtual bool HandleViewportDrag(class USceneComponent* InSceneComponent, class USceneComponent* InComponentTemplate, const FVector& InDeltaTranslation, const FRotator& InDeltaRotation, const FVector& InDeltaScale, const FVector& InPivot) = 0;
+	virtual bool HandleViewportDrag(const USceneComponent* InSceneComponent, class USceneComponent* InComponentTemplate, const FVector& InDeltaTranslation, const FRotator& InDeltaRotation, const FVector& InDeltaScale, const FVector& InPivot) = 0;
 
 	/**
 	 * Get the widget location for this scene component
@@ -45,7 +47,7 @@ public:
 	 * @param	OutWidgetLocation		The output widget location
 	 * @return true if the operation was handled.
 	 */
-	virtual bool HandleGetWidgetLocation(class USceneComponent* InSceneComponent, FVector& OutWidgetLocation) = 0;
+	virtual bool HandleGetWidgetLocation(const USceneComponent* InSceneComponent, FVector& OutWidgetLocation) = 0;
 
 	/**
 	 * Get the widget transform for this scene component
@@ -53,5 +55,5 @@ public:
 	 * @param	OutWidgetLocation		The output widget transform
 	 * @return true if the operation was handled.
 	 */
-	virtual bool HandleGetWidgetTransform(class USceneComponent* InSceneComponent, FMatrix& OutWidgetTransform) = 0;
+	virtual bool HandleGetWidgetTransform(const USceneComponent* InSceneComponent, FMatrix& OutWidgetTransform) = 0;
 };
