Remove unused SCS editor inclues

#rnx

diff --git a/Engine/Plugins/MovieScene/ActorSequence/Source/ActorSequenceEditor/Private/ActorSequenceComponentCustomization.cpp b/Engine/Plugins/MovieScene/ActorSequence/Source/ActorSequenceEditor/Private/ActorSequenceComponentCustomization.cpp
--- a/Engine/Plugins/MovieScene/ActorSequence/Source/ActorSequenceEditor/Private/ActorSequenceComponentCustomization.cpp
+++ b/Engine/Plugins/MovieScene/ActorSequence/Source/ActorSequenceEditor/Private/ActorSequenceComponentCustomization.cpp
@@ -11,7 +11,6 @@
 #include "DetailCategoryBuilder.h"
 #include "DetailWidgetRow.h"
 #include "Widgets/Docking/SDockTab.h"
-#include "SSCSEditor.h"
 #include "BlueprintEditorTabs.h"
 #include "ScopedTransaction.h"
 #include "ISequencerModule.h"
diff --git a/Engine/Source/Editor/Kismet/Private/BlueprintEditorModes.cpp b/Engine/Source/Editor/Kismet/Private/BlueprintEditorModes.cpp
--- a/Engine/Source/Editor/Kismet/Private/BlueprintEditorModes.cpp
+++ b/Engine/Source/Editor/Kismet/Private/BlueprintEditorModes.cpp
@@ -5,7 +5,6 @@
 #include "ToolMenus.h"
 
 // Core kismet tabs
-#include "SSCSEditor.h"
 #include "SSCSEditorViewport.h"
 #include "SKismetInspector.h"
 #include "SMyBlueprint.h"
diff --git a/Engine/Source/Editor/Kismet/Private/SMyBlueprint.cpp b/Engine/Source/Editor/Kismet/Private/SMyBlueprint.cpp
--- a/Engine/Source/Editor/Kismet/Private/SMyBlueprint.cpp
+++ b/Engine/Source/Editor/Kismet/Private/SMyBlueprint.cpp
@@ -29,7 +29,7 @@
 #include "HAL/PlatformApplicationMisc.h"
 #include "DetailLayoutBuilder.h"
 #include "SKismetInspector.h"
-#include "SSCSEditor.h"
+#include "SSubobjectBlueprintEditor.h"
 #include "GraphEditorDragDropAction.h"
 #include "BPFunctionDragDropAction.h"
 #include "BPVariableDragDropAction.h"
@@ -2189,7 +2189,7 @@ TSharedPtr<SWidget> SMyBlueprint::OnContextMenuOpening()
 					// of events appropriate to the component.
 					MenuBuilder.AddSubMenu(	LOCTEXT("AddEventSubMenu", "Add Event"), 
 											LOCTEXT("AddEventSubMenu_ToolTip", "Add Event"), 
-											FNewMenuDelegate::CreateStatic(	&SSCSEditor::BuildMenuEventsSection,
+											FNewMenuDelegate::CreateStatic(	&SSubobjectBlueprintEditor::BuildMenuEventsSection,
 												BlueprintEditor->GetBlueprintObj(), ComponentProperty->PropertyClass, 
 												FCanExecuteAction::CreateRaw(this, &SMyBlueprint::IsEditingMode),
 												FGetSelectedObjectsDelegate::CreateSP(this, &SMyBlueprint::GetSelectedItemsForContextMenu)));
diff --git a/Engine/Source/Editor/LevelEditor/Private/SInViewportDetails.cpp b/Engine/Source/Editor/LevelEditor/Private/SInViewportDetails.cpp
--- a/Engine/Source/Editor/LevelEditor/Private/SInViewportDetails.cpp
+++ b/Engine/Source/Editor/LevelEditor/Private/SInViewportDetails.cpp
@@ -20,7 +20,6 @@
 #include "UnrealEdGlobals.h"
 #include "LevelEditor.h"
 #include "Kismet2/KismetEditorUtilities.h"
-#include "SSCSEditor.h"
 #include "PropertyEditorModule.h"
 #include "IDetailsView.h"
 #include "LevelEditorGenericDetails.h"
