mark USCSEditorExtensionContext::GetSCSEditor as deprecated

#rnx

diff --git a/Engine/Source/Editor/Kismet/Public/SCSEditorExtensionContext.h b/Engine/Source/Editor/Kismet/Public/SCSEditorExtensionContext.h
--- a/Engine/Source/Editor/Kismet/Public/SCSEditorExtensionContext.h
+++ b/Engine/Source/Editor/Kismet/Public/SCSEditorExtensionContext.h
@@ -17,7 +17,8 @@ class KISMET_API USCSEditorExtensionContext : public UObject
 	GENERATED_BODY()
 
 public:
-
+	
+	UE_DEPRECATED(5.0, "GetSCSEditor has been deprecated. Use GetSubobjectEditor instead.")
 	const TWeakPtr<SSCSEditor>& GetSCSEditor() const { return SCSEditor; }
 	
 	const TWeakPtr<SSubobjectEditor>& GetSubobjectEditor() const { return SubobjectEditor; }
