Fix crash when using the arrow key to navigate the Subobject editor and you are at the bottom of the tree.

#rnx

diff --git a/Engine/Source/Editor/Kismet/Private/BlueprintEditor.cpp b/Engine/Source/Editor/Kismet/Private/BlueprintEditor.cpp
--- a/Engine/Source/Editor/Kismet/Private/BlueprintEditor.cpp
+++ b/Engine/Source/Editor/Kismet/Private/BlueprintEditor.cpp
@@ -940,9 +940,9 @@ void FBlueprintEditor::OnSelectionUpdated(const TArray<TSharedPtr<FSubobjectEdit
 		InspectorObjects.Empty(SelectedNodes.Num());
 		for (FSubobjectEditorTreeNodePtrType NodePtr : SelectedNodes)
 		{
-			if (NodePtr.IsValid())
+			const FSubobjectData* NodeData = NodePtr ? NodePtr->GetDataSource() : nullptr;
+			if (NodeData)
 			{
-				const FSubobjectData* NodeData = NodePtr->GetDataSource();
 				if (const AActor* Actor = NodeData->GetObject<AActor>())
 				{
 					if (const AActor* DefaultActor = NodeData->GetObjectForBlueprint<AActor>(GetBlueprintObj()))
