added inherited variables to blueprint Replace References menu
additionally added the name of the variable type to the list so it's easier to tell when you're doing a risky replace

diff --git a/Engine/Source/Editor/Kismet/Private/SReplaceNodeReferences.cpp b/Engine/Source/Editor/Kismet/Private/SReplaceNodeReferences.cpp
--- a/Engine/Source/Editor/Kismet/Private/SReplaceNodeReferences.cpp
+++ b/Engine/Source/Editor/Kismet/Private/SReplaceNodeReferences.cpp
@@ -17,6 +17,7 @@
 #include "ScopedTransaction.h"
 #include "ReplaceNodeReferencesHelper.h"
 #include "Algo/RemoveIf.h"
+#include "Styling/StyleColors.h"
 
 #define LOCTEXT_NAMESPACE "SNodeVariableReferences"
 
@@ -75,7 +76,19 @@ public:
 			.VAlign(VAlign_Center)
 			[
 				SNew(STextBlock)
-					.Text(FText::FromString(VariableReference.GetMemberName().ToString()))
+				.Text(GetDisplayTitle())
+				.ToolTipText(TooltipWarning)
+			]
+			// type of the variable
+			+ SHorizontalBox::Slot()
+			.VAlign(VAlign_Center)
+			.HAlign(HAlign_Right)
+			.Padding(FMargin(1.0f, 0.0f, 0.0f, 0.0f))
+			[
+				SNew(STextBlock)
+				.Text(GetPinClass() ? GetPinClass()->GetDisplayNameText() : FText::GetEmpty())
+				.ToolTipText(TooltipWarning)
+				.ColorAndOpacity(GetVariableTypeColor())
 			];
 	}
 
@@ -113,12 +126,33 @@ public:
 	}
 	// End of FTargetReplaceReferences interface
 
+	/** gets the UClass of the object being referenced by PinType */
+	UClass *GetPinClass() const
+	{
+		if (PinType.PinCategory == UEdGraphSchema_K2::PC_Object && PinType.PinSubCategoryObject != nullptr)
+		{
+			return Cast<UClass>(PinType.PinSubCategoryObject);
+		}
+		return nullptr;
+	}
+
+	/** gets the color that the pin class will be displayed as in the menu */
+	FSlateColor GetVariableTypeColor() const
+	{
+		const FSlateColor TypeColor = TooltipWarning.IsEmpty() ? FSlateColor(EStyleColor::AccentGreen) : FSlateColor(EStyleColor::AccentYellow);
+		const FLinearColor LinearColor = TypeColor.GetSpecifiedColor();
+		return FSlateColor(LinearColor.Desaturate(.2f));
+	}
+
 public:
 	/** Variable reference for this item */
 	FMemberReference VariableReference;
 
 	/** Pin type representing the FProperty of this item */
 	FEdGraphPinType PinType;
+
+	/** Used to warn if the user is going to commit a potentially destructive action */
+	FText TooltipWarning = FText::GetEmpty();
 };
 
 void SReplaceNodeReferences::Construct(const FArguments& InArgs, TSharedPtr<class FBlueprintEditor> InBlueprintEditor)
@@ -488,7 +522,33 @@ void SReplaceNodeReferences::GatherAllAvailableBlueprintVariables(UClass* InTarg
 			FEdGraphPinType Type;
 			K2Schema->ConvertPropertyToPinType(Property, VariableItem->PinType);
 
-			if (!bForTarget || VariableItem->PinType == SourcePinType)
+			// check whether this is a valid replacement using inheritance
+			bool bValidInheritance = false;
+			if (SourcePinType.PinCategory == UEdGraphSchema_K2::PC_Object &&
+				SourcePinType.PinSubCategoryObject != nullptr)
+			{
+				const UClass* VariableItemClass = VariableItem->GetPinClass();
+				const UClass* SourceClass = Cast<UClass>(SourcePinType.PinSubCategoryObject);
+
+				// ensure the casting succeeded
+				if (VariableItemClass && SourceClass)
+				{
+					if (VariableItemClass->IsChildOf(SourceClass) || SourceClass->IsChildOf(VariableItemClass))
+					{
+						bValidInheritance = true;
+						if (VariableItemClass != SourceClass)
+						{
+							VariableItem->TooltipWarning = FText::Format(
+								LOCTEXT("TooltipWarning", "Warning: Replacing a reference of type '{0}' with a reference of type '{1}' may break connections"),
+								SourceClass->GetDisplayNameText(),
+								VariableItemClass->GetDisplayNameText()
+							);
+						}
+					}
+				}
+			}
+
+			if (!bForTarget || bValidInheritance || VariableItem->PinType == SourcePinType)
 			{
 				CategoryReference->Children.Add(VariableItem);
 
