Fix crash on SBlueprintPalette when GetFocusGraph is null

diff --git a/Engine/Source/Editor/Kismet/Private/SBlueprintPalette.cpp b/Engine/Source/Editor/Kismet/Private/SBlueprintPalette.cpp
--- a/Engine/Source/Editor/Kismet/Private/SBlueprintPalette.cpp
+++ b/Engine/Source/Editor/Kismet/Private/SBlueprintPalette.cpp
@@ -763,7 +763,10 @@ private:
 		const UEdGraphSchema* Schema = GetDefault<UEdGraphSchema_K2>();
 		if (BlueprintEditorPtr.IsValid())
 		{
-			Schema = BlueprintEditorPtr.Pin()->GetFocusedGraph()->GetSchema();
+			if (BlueprintEditorPtr.Pin()->GetFocusedGraph())
+			{
+				Schema = BlueprintEditorPtr.Pin()->GetFocusedGraph()->GetSchema();
+			}
 		}
 		
 		this->ChildSlot
