Fixes to Blueprint Diff menu and nomad main menu registration to fix automated tests and Mac support.

diff --git a/Engine/Source/Editor/EditorStyle/Private/StarshipStyle.cpp b/Engine/Source/Editor/EditorStyle/Private/StarshipStyle.cpp
--- a/Engine/Source/Editor/EditorStyle/Private/StarshipStyle.cpp
+++ b/Engine/Source/Editor/EditorStyle/Private/StarshipStyle.cpp
@@ -5344,6 +5344,7 @@ void FStarshipEditorStyle::FStyle::SetupPersonaStyle()
 
 	//Blueprint Diff
 	{
+		Set("BlueprintDiff.ToolbarIcon", new IMAGE_BRUSH_SVG("Starship/Blueprints/BlueprintDiff", Icon16x16));
 		Set("BlueprintDif.HasGraph", new IMAGE_BRUSH_SVG("Starship/Blueprints/blueprint_Dif_has_graph", Icon10x10));
 		Set("BlueprintDif.MissingGraph", new IMAGE_BRUSH_SVG("Starship/Blueprints/blueprint_Dif_missing_graph", Icon8x8));
 		Set("BlueprintDif.NextDiff", new IMAGE_BRUSH_SVG("/Starship/Blueprints/diff_next", Icon16x16));
diff --git a/Engine/Source/Editor/Kismet/Private/SBlueprintEditorToolbar.cpp b/Engine/Source/Editor/Kismet/Private/SBlueprintEditorToolbar.cpp
--- a/Engine/Source/Editor/Kismet/Private/SBlueprintEditorToolbar.cpp
+++ b/Engine/Source/Editor/Kismet/Private/SBlueprintEditorToolbar.cpp
@@ -79,11 +79,6 @@
 		Section.AddMenuEntry( FBlueprintEditorCommands::Get().CompileBlueprint );
 		Section.AddMenuEntry( FBlueprintEditorCommands::Get().RefreshAllNodes );
 		Section.AddMenuEntry( FBlueprintEditorCommands::Get().ReparentBlueprint );
-		Section.AddSubMenu(
-			"Diff",
-			LOCTEXT("Diff", "Diff"),
-			LOCTEXT("BlueprintEditorDiffToolTip", "Diff against previous revisions"),
-			FNewToolMenuWidget::CreateStatic(&FKismet2Menu::MakeDiffMenu));
 		Section.AddMenuEntry(FBlueprintEditorCommands::Get().BeginBlueprintMerge);
 	}
 
@@ -235,111 +230,6 @@
 	}
 }
 
-/** Delegate called to diff a specific revision with the current */
-static void OnDiffRevisionPicked(FRevisionInfo const& RevisionInfo, TWeakObjectPtr<UBlueprint> BlueprintObj)
-{
-	if (BlueprintObj.IsValid())
-	{
-		bool const bIsLevelScriptBlueprint = FBlueprintEditorUtils::IsLevelScriptBlueprint(BlueprintObj.Get());
-		FString const Filename = SourceControlHelpers::PackageFilename(bIsLevelScriptBlueprint ? BlueprintObj.Get()->GetOuter()->GetPathName() : BlueprintObj.Get()->GetPathName());
 
-		ISourceControlProvider& SourceControlProvider = ISourceControlModule::Get().GetProvider();
-
-		// Get the SCC state
-		FSourceControlStatePtr SourceControlState = SourceControlProvider.GetState(Filename, EStateCacheUsage::Use);
-		if (SourceControlState.IsValid())
-		{
-			for (int32 HistoryIndex = 0; HistoryIndex < SourceControlState->GetHistorySize(); HistoryIndex++)
-			{
-				TSharedPtr<ISourceControlRevision, ESPMode::ThreadSafe> Revision = SourceControlState->GetHistoryItem(HistoryIndex);
-				check(Revision.IsValid());
-				if (Revision->GetRevision() == RevisionInfo.Revision)
-				{
-					// Get the revision of this package from source control
-					FString PreviousTempPkgName;
-					if (Revision->Get(PreviousTempPkgName))
-					{
-						// Try and load that package
-						UPackage* PreviousTempPkg = LoadPackage(NULL, *PreviousTempPkgName, LOAD_ForDiff|LOAD_DisableCompileOnLoad);
-
-						if (PreviousTempPkg != NULL)
-						{
-							UObject* PreviousAsset = NULL;
-
-							// If its a levelscript blueprint, find the previous levelscript blueprint in the map
-							if (bIsLevelScriptBlueprint)
-							{
-								TArray<UObject *> ObjectsInOuter;
-								GetObjectsWithOuter(PreviousTempPkg, ObjectsInOuter);
-
-								// Look for the level script blueprint for this package
-								for (int32 Index = 0; Index < ObjectsInOuter.Num(); Index++)
-								{
-									UObject* Obj = ObjectsInOuter[Index];
-									if (ULevelScriptBlueprint* ObjAsBlueprint = Cast<ULevelScriptBlueprint>(Obj))
-									{
-										PreviousAsset = ObjAsBlueprint;
-										break;
-									}
-								}
-							}
-							// otherwise its a normal Blueprint
-							else
-							{
-								FString PreviousAssetName = FPaths::GetBaseFilename(Filename, true);
-								PreviousAsset = FindObject<UObject>(PreviousTempPkg, *PreviousAssetName);
-							}
-
-							if (PreviousAsset != NULL)
-							{
-								FAssetToolsModule& AssetToolsModule = FModuleManager::LoadModuleChecked<FAssetToolsModule>(TEXT("AssetTools"));
-								FRevisionInfo OldRevision = { Revision->GetRevision(), Revision->GetCheckInIdentifier(), Revision->GetDate() };
-								FRevisionInfo CurrentRevision = { TEXT(""), Revision->GetCheckInIdentifier(), Revision->GetDate() };
-								AssetToolsModule.Get().DiffAssets(PreviousAsset, BlueprintObj.Get(), OldRevision, CurrentRevision);
-							}
-						}
-						else
-						{
-							FMessageDialog::Open(EAppMsgType::Ok, NSLOCTEXT("SourceControl.HistoryWindow", "UnableToLoadAssets", "Unable to load assets to diff. Content may no longer be supported?"));
-						}
-					}
-					break;
-				}
-			}
-		}
-	}
-}
-
-TSharedRef<SWidget> FKismet2Menu::MakeDiffMenu(const FToolMenuContext& InToolMenuContext)
-{
-	if (ISourceControlModule::Get().IsEnabled() && ISourceControlModule::Get().GetProvider().IsAvailable())
-	{
-		UBlueprintEditorToolMenuContext* Context = InToolMenuContext.FindContext<UBlueprintEditorToolMenuContext>();
-		UBlueprint* BlueprintObj = Context ? Context->GetBlueprintObj() : nullptr;
-		if(BlueprintObj)
-		{
-			TWeakObjectPtr<UBlueprint> BlueprintPtr = BlueprintObj;
-			// Add our async SCC task widget
-			return SNew(SBlueprintRevisionMenu, BlueprintObj)
-				.OnRevisionSelected_Static(&OnDiffRevisionPicked, BlueprintPtr);
-		}
-		else
-		{
-			// if BlueprintObj is null then this means that multiple blueprints are selected
-			FMenuBuilder MenuBuilder(true, NULL);
-			MenuBuilder.AddMenuEntry( LOCTEXT("NoRevisionsForMultipleBlueprints", "Multiple blueprints selected"), 
-				FText(), FSlateIcon(), FUIAction() );
-			return MenuBuilder.MakeWidget();
-		}
-	}
-
-	FMenuBuilder MenuBuilder(true, NULL);
-	MenuBuilder.AddMenuEntry( LOCTEXT("SourceControlDisabled", "Source control is disabled"), 
-		FText(), FSlateIcon(), FUIAction() );
-	return MenuBuilder.MakeWidget();
-}
-
-
-
 //////////////////////////////////////////////////////////////////////////
 // FFullBlueprintEditorCommands
@@ -415,12 +305,12 @@
 
 void FBlueprintEditorToolbar::AddBlueprintGlobalOptionsToolbar(UToolMenu* InMenu, bool bRegisterViewport)
 {
-	FToolMenuSection& Section = InMenu->AddSection("Settings");
+	FToolMenuSection& Section = InMenu->FindOrAddSection("Settings");
 	Section.InsertPosition = FToolMenuInsert("Asset", EToolMenuInsertType::After);
 
 	Section.AddDynamicEntry("BlueprintGlobalOptions", FNewToolMenuSectionDelegate::CreateLambda([bRegisterViewport](FToolMenuSection& InSection)
 	{
-		UBlueprintEditorToolMenuContext* Context = InSection.FindContext<UBlueprintEditorToolMenuContext>();
+		const UBlueprintEditorToolMenuContext* Context = InSection.FindContext<UBlueprintEditorToolMenuContext>();
 		if (Context && Context->GetBlueprintObj())
 		{
 			const FFullBlueprintEditorCommands& Commands = FFullBlueprintEditorCommands::Get();
@@ -442,7 +332,7 @@
 
 	Section.AddDynamicEntry("CompileCommands", FNewToolMenuSectionDelegate::CreateLambda([](FToolMenuSection& InSection)
 	{
-		UBlueprintEditorToolMenuContext* Context = InSection.FindContext<UBlueprintEditorToolMenuContext>();
+		const UBlueprintEditorToolMenuContext* Context = InSection.FindContext<UBlueprintEditorToolMenuContext>();
 		if (Context && Context->BlueprintEditor.IsValid() && Context->GetBlueprintObj())
 		{
 			TSharedPtr<class FBlueprintEditorToolbar> BlueprintEditorToolbar = Context->BlueprintEditor.Pin()->GetToolbarBuilder();
@@ -468,6 +358,33 @@
 			}
 		}
 	}));
+
+	// We want the diff menu to be on any blueprint toolbar that also contains compile
+	FToolMenuSection& DiffSection = InMenu->AddSection("SourceControl");
+	Section.InsertPosition = FToolMenuInsert("Asset", EToolMenuInsertType::After);
+
+	Section.AddDynamicEntry("CompileCommands", FNewToolMenuSectionDelegate::CreateLambda([](FToolMenuSection& InSection)
+		{
+			const UBlueprintEditorToolMenuContext* Context = InSection.FindContext<UBlueprintEditorToolMenuContext>();
+			if (Context && Context->BlueprintEditor.IsValid() && Context->GetBlueprintObj())
+			{
+				TSharedPtr<class FBlueprintEditorToolbar> BlueprintEditorToolbar = Context->BlueprintEditor.Pin()->GetToolbarBuilder();
+				if (BlueprintEditorToolbar.IsValid())
+				{
+					InSection.InsertPosition = FToolMenuInsert();
+					FToolMenuEntry DiffEntry = FToolMenuEntry::InitComboButton(
+						"Diff",
+						FUIAction(),
+						FOnGetContent::CreateStatic(&FBlueprintEditorToolbar::MakeDiffMenu, Context),
+						LOCTEXT("Diff", "Diff"),
+						LOCTEXT("BlueprintEditorDiffToolTip", "Diff against previous revisions"),
+						FSlateIcon(FAppStyle::Get().GetStyleSetName(), "BlueprintDiff.ToolbarIcon")
+					);
+					DiffEntry.StyleNameOverride = "CalloutToolbar";
+					InSection.AddEntry(DiffEntry);
+				}
+			}
+		}));
 }
 
 void FBlueprintEditorToolbar::AddNewToolbar(UToolMenu* InMenu)
@@ -604,4 +521,109 @@
 	}
 }
 
+
+/** Delegate called to diff a specific revision with the current */
+static void OnDiffRevisionPicked(FRevisionInfo const& RevisionInfo, TWeakObjectPtr<UBlueprint> BlueprintObj)
+{
+	if (BlueprintObj.IsValid())
+	{
+		bool const bIsLevelScriptBlueprint = FBlueprintEditorUtils::IsLevelScriptBlueprint(BlueprintObj.Get());
+		FString const Filename = SourceControlHelpers::PackageFilename(bIsLevelScriptBlueprint ? BlueprintObj.Get()->GetOuter()->GetPathName() : BlueprintObj.Get()->GetPathName());
+
+		ISourceControlProvider& SourceControlProvider = ISourceControlModule::Get().GetProvider();
+
+		// Get the SCC state
+		FSourceControlStatePtr SourceControlState = SourceControlProvider.GetState(Filename, EStateCacheUsage::Use);
+		if (SourceControlState.IsValid())
+		{
+			for (int32 HistoryIndex = 0; HistoryIndex < SourceControlState->GetHistorySize(); HistoryIndex++)
+			{
+				TSharedPtr<ISourceControlRevision, ESPMode::ThreadSafe> Revision = SourceControlState->GetHistoryItem(HistoryIndex);
+				check(Revision.IsValid());
+				if (Revision->GetRevision() == RevisionInfo.Revision)
+				{
+					// Get the revision of this package from source control
+					FString PreviousTempPkgName;
+					if (Revision->Get(PreviousTempPkgName))
+					{
+						// Try and load that package
+						UPackage* PreviousTempPkg = LoadPackage(NULL, *PreviousTempPkgName, LOAD_ForDiff | LOAD_DisableCompileOnLoad);
+
+						if (PreviousTempPkg != NULL)
+						{
+							UObject* PreviousAsset = NULL;
+
+							// If its a levelscript blueprint, find the previous levelscript blueprint in the map
+							if (bIsLevelScriptBlueprint)
+							{
+								TArray<UObject*> ObjectsInOuter;
+								GetObjectsWithOuter(PreviousTempPkg, ObjectsInOuter);
+
+								// Look for the level script blueprint for this package
+								for (int32 Index = 0; Index < ObjectsInOuter.Num(); Index++)
+								{
+									UObject* Obj = ObjectsInOuter[Index];
+									if (ULevelScriptBlueprint* ObjAsBlueprint = Cast<ULevelScriptBlueprint>(Obj))
+									{
+										PreviousAsset = ObjAsBlueprint;
+										break;
+									}
+								}
+							}
+							// otherwise its a normal Blueprint
+							else
+							{
+								FString PreviousAssetName = FPaths::GetBaseFilename(Filename, true);
+								PreviousAsset = FindObject<UObject>(PreviousTempPkg, *PreviousAssetName);
+							}
+
+							if (PreviousAsset != NULL)
+							{
+								FAssetToolsModule& AssetToolsModule = FModuleManager::LoadModuleChecked<FAssetToolsModule>(TEXT("AssetTools"));
+								FRevisionInfo OldRevision = { Revision->GetRevision(), Revision->GetCheckInIdentifier(), Revision->GetDate() };
+								FRevisionInfo CurrentRevision = { TEXT(""), Revision->GetCheckInIdentifier(), Revision->GetDate() };
+								AssetToolsModule.Get().DiffAssets(PreviousAsset, BlueprintObj.Get(), OldRevision, CurrentRevision);
+							}
+						}
+						else
+						{
+							FMessageDialog::Open(EAppMsgType::Ok, NSLOCTEXT("SourceControl.HistoryWindow", "UnableToLoadAssets", "Unable to load assets to diff. Content may no longer be supported?"));
+						}
+					}
+					break;
+				}
+			}
+		}
+	}
+}
+
+TSharedRef<SWidget> FBlueprintEditorToolbar::MakeDiffMenu(const UBlueprintEditorToolMenuContext* InContext)
+{
+	if (ISourceControlModule::Get().IsEnabled() && ISourceControlModule::Get().GetProvider().IsAvailable())
+	{
+		UBlueprint* BlueprintObj =  InContext ? InContext->GetBlueprintObj() : nullptr;
+		if (BlueprintObj)
+		{
+			TWeakObjectPtr<UBlueprint> BlueprintPtr = BlueprintObj;
+			// Add our async SCC task widget
+			return SNew(SBlueprintRevisionMenu, BlueprintObj)
+				.OnRevisionSelected_Static(&OnDiffRevisionPicked, BlueprintPtr);
+		}
+		else
+		{
+			// if BlueprintObj is null then this means that multiple blueprints are selected
+			FMenuBuilder MenuBuilder(true, NULL);
+			MenuBuilder.AddMenuEntry(LOCTEXT("NoRevisionsForMultipleBlueprints", "Multiple blueprints selected"),
+				FText(), FSlateIcon(), FUIAction());
+			return MenuBuilder.MakeWidget();
+		}
+	}
+
+	FMenuBuilder MenuBuilder(true, NULL);
+	MenuBuilder.AddMenuEntry(LOCTEXT("SourceControlDisabled", "Source control is disabled"),
+		FText(), FSlateIcon(), FUIAction());
+	return MenuBuilder.MakeWidget();
+}
+
+
 #undef LOCTEXT_NAMESPACE
diff --git a/Engine/Source/Editor/Kismet/Public/SBlueprintEditorToolbar.h b/Engine/Source/Editor/Kismet/Public/SBlueprintEditorToolbar.h
--- a/Engine/Source/Editor/Kismet/Public/SBlueprintEditorToolbar.h
+++ b/Engine/Source/Editor/Kismet/Public/SBlueprintEditorToolbar.h
@@ -16,6 +16,7 @@ class FMenuBuilder;
 class FToolBarBuilder;
 class UToolMenu;
 struct FToolMenuContext;
+class UBlueprintEditorToolMenuContext;
 
 /**
  * Kismet menu
@@ -35,12 +36,6 @@ protected:
 	static void FillDebugMenu(UToolMenu* Menu);
 
 	static void FillDeveloperMenu(UToolMenu* Menu);
-
-private:
-	/** Diff current blueprint against the specified revision */
-	static void DiffAgainstRevision( class UBlueprint* Current, int32 OldRevision );
-
-	static TSharedRef<SWidget> MakeDiffMenu(const FToolMenuContext& ToolMenuContext);
 };
 
 
@@ -93,6 +88,11 @@ public:
 	/** Returns the current status as text for the blueprint being edited */
 	FText GetStatusTooltip() const;
 
+	/** Diff current blueprint against the specified revision */
+	static void DiffAgainstRevision(class UBlueprint* Current, int32 OldRevision);
+
+	static TSharedRef<SWidget> MakeDiffMenu(const UBlueprintEditorToolMenuContext* InContext);
+
 protected:
 	/** Pointer back to the blueprint editor tool that owns us */
 	TWeakPtr<FBlueprintEditor> BlueprintEditor;
diff --git a/Engine/Source/Editor/MainFrame/Private/Menus/MainMenu.cpp b/Engine/Source/Editor/MainFrame/Private/Menus/MainMenu.cpp
--- a/Engine/Source/Editor/MainFrame/Private/Menus/MainMenu.cpp
+++ b/Engine/Source/Editor/MainFrame/Private/Menus/MainMenu.cpp
@@ -540,7 +540,12 @@ void FMainMenu::RegisterExitMenuItems()
 
 void FMainMenu::RegisterNomadMainMenu()
 {
-	UToolMenu* Menu = UToolMenus::Get()->RegisterMenu("MainFrame.NomadMainMenu", "MainFrame.MainMenu");
+	UToolMenus* ToolMenus = UToolMenus::Get();
+	static const FName NomadMainMenuName("MainFrame.NomadMainMenu");
+	if (!ToolMenus->IsMenuRegistered(NomadMainMenuName))
+	{
+		UToolMenu* Menu = UToolMenus::Get()->RegisterMenu(NomadMainMenuName, "MainFrame.MainMenu");
+	}
 }
 
 
diff --git a/Engine/Source/Editor/UMGEditor/Private/WidgetBlueprintEditorToolbar.cpp b/Engine/Source/Editor/UMGEditor/Private/WidgetBlueprintEditorToolbar.cpp
--- a/Engine/Source/Editor/UMGEditor/Private/WidgetBlueprintEditorToolbar.cpp
+++ b/Engine/Source/Editor/UMGEditor/Private/WidgetBlueprintEditorToolbar.cpp
@@ -127,7 +127,7 @@ void FWidgetBlueprintEditorToolbar::FillWidgetBlueprintEditorModesToolbar(FToolB
 void FWidgetBlueprintEditorToolbar::AddWidgetReflector(UToolMenu* InMenu)
 {
 	FToolMenuSection& Section = InMenu->AddSection("WidgetTools");
-	Section.InsertPosition = FToolMenuInsert("Asset", EToolMenuInsertType::After);
+	Section.InsertPosition = FToolMenuInsert("SourceControl", EToolMenuInsertType::After);
 
 	Section.AddEntry(FToolMenuEntry::InitToolBarButton(
 		"OpenWidgetReflector",
