Fix graph node's details not being read-only when displayed in a read-only graph

diff --git a/Engine/Source/Editor/Kismet/Private/SKismetInspector.cpp b/Engine/Source/Editor/Kismet/Private/SKismetInspector.cpp
--- a/Engine/Source/Editor/Kismet/Private/SKismetInspector.cpp
+++ b/Engine/Source/Editor/Kismet/Private/SKismetInspector.cpp
@@ -904,13 +904,35 @@ bool SKismetInspector::IsPropertyEditingEnabled() const
 
 	for (const TWeakObjectPtr<UObject>& SelectedObject : SelectedObjects)
 	{
-		UActorComponent* Component = Cast<UActorComponent>(SelectedObject.Get());
-		if (Component && !CastChecked<UActorComponent>(Component->GetArchetype())->IsEditableWhenInherited() )
+		if (UActorComponent* Component = Cast<UActorComponent>(SelectedObject.Get()))
 		{
-			bIsEditable = false;
-			break;
+			if(!CastChecked<UActorComponent>(Component->GetArchetype())->IsEditableWhenInherited())
+			{
+				bIsEditable = false;
+				break;
+			}
+		}
+		else if(UEdGraph* EdGraph = Cast<UEdGraph>(SelectedObject.Get()))
+		{
+			if(!BlueprintEditorPtr.Pin()->IsEditable(EdGraph))
+			{
+				bIsEditable = false;
+				break;
+			}
+		}
+		else if(UEdGraphNode* EdGraphNode = Cast<UEdGraphNode>(SelectedObject.Get()))
+		{
+			if(UEdGraph* OuterGraph = EdGraphNode->GetGraph())
+			{
+				if(!BlueprintEditorPtr.Pin()->IsEditable(OuterGraph))
+				{
+					bIsEditable = false;
+					break;
+				}
+			}
 		}
 	}
+	
 	return bIsEditable && (!IsPropertyEditingEnabledDelegate.IsBound() || IsPropertyEditingEnabledDelegate.Execute());
 }
 
