Partially revert CL 16640876 to only cover nodes, not graphs

This broke BP interface graphs

diff --git a/Engine/Source/Editor/Kismet/Private/SKismetInspector.cpp b/Engine/Source/Editor/Kismet/Private/SKismetInspector.cpp
--- a/Engine/Source/Editor/Kismet/Private/SKismetInspector.cpp
+++ b/Engine/Source/Editor/Kismet/Private/SKismetInspector.cpp
@@ -912,14 +912,6 @@ bool SKismetInspector::IsPropertyEditingEnabled() const
 				break;
 			}
 		}
-		else if(UEdGraph* EdGraph = Cast<UEdGraph>(SelectedObject.Get()))
-		{
-			if(!BlueprintEditorPtr.Pin()->IsEditable(EdGraph))
-			{
-				bIsEditable = false;
-				break;
-			}
-		}
 		else if(UEdGraphNode* EdGraphNode = Cast<UEdGraphNode>(SelectedObject.Get()))
 		{
 			if(UEdGraph* OuterGraph = EdGraphNode->GetGraph())
