Fixing missing Compile button

diff --git a/Engine/Source/Editor/Kismet/Private/SBlueprintEditorToolbar.cpp b/Engine/Source/Editor/Kismet/Private/SBlueprintEditorToolbar.cpp
--- a/Engine/Source/Editor/Kismet/Private/SBlueprintEditorToolbar.cpp
+++ b/Engine/Source/Editor/Kismet/Private/SBlueprintEditorToolbar.cpp
@@ -361,9 +361,9 @@ void FBlueprintEditorToolbar::AddCompileToolbar(UToolMenu* InMenu)
 
 	// We want the diff menu to be on any blueprint toolbar that also contains compile
 	FToolMenuSection& DiffSection = InMenu->AddSection("SourceControl");
-	Section.InsertPosition = FToolMenuInsert("Asset", EToolMenuInsertType::After);
+	DiffSection.InsertPosition = FToolMenuInsert("Asset", EToolMenuInsertType::After);
 
-	Section.AddDynamicEntry("CompileCommands", FNewToolMenuSectionDelegate::CreateLambda([](FToolMenuSection& InSection)
+	DiffSection.AddDynamicEntry("SourceControlCommands", FNewToolMenuSectionDelegate::CreateLambda([](FToolMenuSection& InSection)
 		{
 			const UBlueprintEditorToolMenuContext* Context = InSection.FindContext<UBlueprintEditorToolMenuContext>();
 			if (Context && Context->BlueprintEditor.IsValid() && Context->GetBlueprintObj())
