Added ability to show external graphs in My Blueprint

Only one level of external graph is added to avoid circular references creating infinite recursion

diff --git a/Engine/Source/Editor/BlueprintGraph/Classes/K2Node_ExternalGraphInterface.h b/Engine/Source/Editor/BlueprintGraph/Classes/K2Node_ExternalGraphInterface.h
new file mode 100644
--- /dev/null
+++ b/Engine/Source/Editor/BlueprintGraph/Classes/K2Node_ExternalGraphInterface.h
@@ -0,0 +1,22 @@
+// Copyright Epic Games, Inc. All Rights Reserved.
+
+#pragma once
+
+#include "CoreMinimal.h"
+#include "UObject/Interface.h"
+#include "K2Node_ExternalGraphInterface.generated.h"
+
+UINTERFACE(meta=(CannotImplementInterfaceInBlueprint))
+class BLUEPRINTGRAPH_API UK2Node_ExternalGraphInterface : public UInterface
+{
+	GENERATED_BODY()
+};
+
+class BLUEPRINTGRAPH_API IK2Node_ExternalGraphInterface
+{
+	GENERATED_BODY()
+
+public:
+	// Get external graphs to display for this node
+	virtual TArray<UEdGraph*> GetExternalGraphs() const = 0;
+};
\ No newline at end of file
diff --git a/Engine/Source/Editor/Kismet/Private/SMyBlueprint.cpp b/Engine/Source/Editor/Kismet/Private/SMyBlueprint.cpp
--- a/Engine/Source/Editor/Kismet/Private/SMyBlueprint.cpp
+++ b/Engine/Source/Editor/Kismet/Private/SMyBlueprint.cpp
@@ -51,6 +51,7 @@
 #include "ReplaceNodeReferencesHelper.h"
 #include "Animation/AnimClassInterface.h"
 #include "BPGraphClipboardData.h"
+#include "K2Node_ExternalGraphInterface.h"
 #include "SEditorHeaderButton.h"
 
 #define LOCTEXT_NAMESPACE "MyBlueprint"
@@ -67,6 +68,7 @@ void FMyBlueprintCommands::RegisterCommands()
 {
 	UI_COMMAND( OpenGraph, "Open Graph", "Opens up this function, macro, or event graph's graph panel up.", EUserInterfaceActionType::Button, FInputChord() );
 	UI_COMMAND( OpenGraphInNewTab, "Open in New Tab", "Opens up this function, macro, or event graph's graph panel up in a new tab. Hold down Ctrl and double click for shortcut.", EUserInterfaceActionType::Button, FInputChord() );
+	UI_COMMAND( OpenExternalGraph, "Open External Graph", "Opens up this external graph's graph panel in its own asset editor", EUserInterfaceActionType::Button, FInputChord() );
 	UI_COMMAND( FocusNode, "Focus", "Focuses on the associated node", EUserInterfaceActionType::Button, FInputChord() );
 	UI_COMMAND( FocusNodeInNewTab, "Focus in New Tab", "Focuses on the associated node in a new tab", EUserInterfaceActionType::Button, FInputChord() );
 	UI_COMMAND( ImplementFunction, "Implement event", "Implements this overridable function as a new event.", EUserInterfaceActionType::Button, FInputChord() );
@@ -268,6 +270,11 @@ void SMyBlueprint::Construct(const FArguments& InArgs, TWeakPtr<FBlueprintEditor
 			FCanExecuteAction(), FIsActionChecked(),
 			FIsActionButtonVisible::CreateSP(this, &SMyBlueprint::CanOpenGraph) );
 
+		CommandList->MapAction( FMyBlueprintCommands::Get().OpenExternalGraph,
+			FExecuteAction::CreateSP(this, &SMyBlueprint::OnOpenExternalGraph),
+			FCanExecuteAction(), FIsActionChecked(),
+			FIsActionButtonVisible::CreateSP(this, &SMyBlueprint::CanOpenExternalGraph) );
+		
 		CommandList->MapAction( FMyBlueprintCommands::Get().FocusNode,
 			FExecuteAction::CreateSP(this, &SMyBlueprint::OnFocusNode),
 			FCanExecuteAction(), FIsActionChecked(),
@@ -1048,7 +1055,7 @@ TSharedRef<SWidget> SMyBlueprint::OnCreateWidgetForAction(FCreateWidgetForAction
 	return BlueprintEditorPtr.IsValid() ? SNew(SBlueprintPaletteItem, InCreateData, BlueprintEditorPtr.Pin()) : SNew(SBlueprintPaletteItem, InCreateData, GetBlueprintObj());
 }
 
-void SMyBlueprint::GetChildGraphs(UEdGraph* InEdGraph, int32 const SectionId, FGraphActionSort& SortList, const FText& ParentCategory)
+void SMyBlueprint::GetChildGraphs(UEdGraph* InEdGraph, int32 const SectionId, FGraphActionSort& SortList, const FText& ParentCategory) const
 {
 	check(InEdGraph);
 
@@ -1101,7 +1108,7 @@ void SMyBlueprint::GetChildGraphs(UEdGraph* InEdGraph, int32 const SectionId, FG
 	}
 }
 
-void SMyBlueprint::GetChildEvents(UEdGraph const* InEdGraph, int32 const SectionId, FGraphActionSort& SortList, const FText& ParentCategory) const
+void SMyBlueprint::GetChildEvents(UEdGraph const* InEdGraph, int32 const SectionId, FGraphActionSort& SortList, const FText& ParentCategory, bool bInAddChildGraphs) const
 {
 	if (!ensure(InEdGraph != NULL))
 	{
@@ -1127,11 +1134,36 @@ void SMyBlueprint::GetChildEvents(UEdGraph const* InEdGraph, int32 const Section
 
 	for (UEdGraphNode* GraphNode : InEdGraph->Nodes)
 	{
-		if (GraphNode && GraphNode->GetClass()->ImplementsInterface(UK2Node_EventNodeInterface::StaticClass()))
+		if (GraphNode)
 		{
-			TSharedPtr<FEdGraphSchemaAction> EventNodeAction = CastChecked<IK2Node_EventNodeInterface>(GraphNode)->GetEventNodeAction(ActionCategory);
-			EventNodeAction->SectionID = SectionId;
-			SortList.AddAction(EventNodeAction);
+			TSharedPtr<FEdGraphSchemaAction> EventNodeAction;
+			if(GraphNode->GetClass()->ImplementsInterface(UK2Node_EventNodeInterface::StaticClass()))
+			{
+				EventNodeAction = CastChecked<IK2Node_EventNodeInterface>(GraphNode)->GetEventNodeAction(ActionCategory);
+				EventNodeAction->SectionID = SectionId;
+				SortList.AddAction(EventNodeAction);
+			}
+
+			if(bInAddChildGraphs && GraphNode->GetClass()->ImplementsInterface(UK2Node_ExternalGraphInterface::StaticClass()))
+			{
+				TArray<UEdGraph*> ExternalGraphs = CastChecked<IK2Node_ExternalGraphInterface>(GraphNode)->GetExternalGraphs();
+				for(UEdGraph* ExternalGraph : ExternalGraphs)
+				{
+					FText ExternalGraphCategory;
+					if(EventNodeAction.IsValid())
+					{
+						ExternalGraphCategory = FText::Format(FText::FromString(TEXT("{0}|{1}")), ActionCategory, EventNodeAction->GetMenuDescription());
+					}
+					else
+					{
+						ExternalGraphCategory = ActionCategory;
+					}
+					
+					// Dont add child graphs, to avoid circular references creating infinite recursion
+					const bool bAddChildGraphs = false;
+					AddEventForFunctionGraph(ExternalGraph, SectionId, SortList, ExternalGraphCategory, bAddChildGraphs);
+				}
+			}
 		}
 	}
 }
@@ -1214,6 +1246,57 @@ EVisibility SMyBlueprint::GetLocalActionsListVisibility() const
 	return EVisibility::Collapsed;
 }
 
+void SMyBlueprint::AddEventForFunctionGraph(UEdGraph* InEdGraph, int32 const SectionId, FGraphActionSort& SortList, const FText& ParentCategory, bool bAddChildGraphs) const
+{
+	UBlueprint* BlueprintObj = GetBlueprintObj();
+	check(BlueprintObj);
+	
+	FGraphDisplayInfo DisplayInfo;
+	InEdGraph->GetSchema()->GetGraphDisplayInformation(*InEdGraph, DisplayInfo);
+
+	FText FunctionCategory = InEdGraph->GetSchema()->GetGraphCategory(InEdGraph);
+	if (FunctionCategory.IsEmpty() && BlueprintObj->SkeletonGeneratedClass != nullptr)
+	{
+		UFunction* Function = BlueprintObj->SkeletonGeneratedClass->FindFunctionByName(InEdGraph->GetFName());
+		if (Function != nullptr)
+		{
+			FunctionCategory = FObjectEditorUtils::GetCategoryText(Function);
+		}
+	}
+
+	// Default, so place in 'non' category
+	if (FunctionCategory.EqualTo(FText::FromString(BlueprintObj->GetName())) || FunctionCategory.EqualTo(UEdGraphSchema_K2::VR_DefaultCategory))
+	{
+		FunctionCategory = FText::GetEmpty();
+	}
+
+	FText ActionCategory;
+	if (!ParentCategory.IsEmpty())
+	{
+		ActionCategory = FText::Format(FText::FromString(TEXT("{0}|{1}")), ParentCategory, FunctionCategory);
+	}
+	else
+	{
+		ActionCategory = MoveTemp(FunctionCategory);
+	}
+	
+	//@TODO: Should be a bit more generic (or the AnimGraph shouldn't be stored as a FunctionGraph...)
+	const bool bIsConstructionScript = InEdGraph->GetFName() == UEdGraphSchema_K2::FN_UserConstructionScript;
+
+	TSharedPtr<FEdGraphSchemaAction_K2Graph> NewFuncAction = MakeShareable(new FEdGraphSchemaAction_K2Graph(EEdGraphSchemaAction_K2Graph::Function, ActionCategory, DisplayInfo.PlainName, DisplayInfo.Tooltip, bIsConstructionScript ? 2 : 1, SectionId));
+	NewFuncAction->FuncName = InEdGraph->GetFName();
+	NewFuncAction->EdGraph = InEdGraph;
+
+	const FString UserCategoryName = FEditorCategoryUtils::GetCategoryDisplayString(ActionCategory.ToString());
+	SortList.AddAction(UserCategoryName, NewFuncAction);
+
+	if(bAddChildGraphs)
+	{
+		GetChildGraphs(InEdGraph, NewFuncAction->GetSectionID(), SortList, ActionCategory);
+	}
+	GetChildEvents(InEdGraph, NewFuncAction->GetSectionID(), SortList, ActionCategory, bAddChildGraphs);
+}
+
 void SMyBlueprint::CollectAllActions(FGraphActionListBuilderBase& OutAllActions)
 {
 	const UEdGraphSchema_K2* K2Schema = GetDefault<UEdGraphSchema_K2>();
@@ -1331,28 +1414,6 @@ void SMyBlueprint::CollectAllActions(FGraphActionListBuilderBase& OutAllActions)
 	{
 		check(Graph);
 
-		FGraphDisplayInfo DisplayInfo;
-		Graph->GetSchema()->GetGraphDisplayInformation(*Graph, DisplayInfo);
-
-		FText FunctionCategory = Graph->GetSchema()->GetGraphCategory(Graph);
-		if (FunctionCategory.IsEmpty() && BlueprintObj->SkeletonGeneratedClass != nullptr)
-		{
-			UFunction* Function = BlueprintObj->SkeletonGeneratedClass->FindFunctionByName(Graph->GetFName());
-			if (Function != nullptr)
-			{
-				FunctionCategory = FObjectEditorUtils::GetCategoryText(Function);
-			}
-		}
-
-		// Default, so place in 'non' category
-		if (FunctionCategory.EqualTo(FText::FromString(BlueprintObj->GetName())) || FunctionCategory.EqualTo(UEdGraphSchema_K2::VR_DefaultCategory))
-		{
-			FunctionCategory = FText::GetEmpty();
-		}
-
-		//@TODO: Should be a bit more generic (or the AnimGraph shouldn't be stored as a FunctionGraph...)
-		const bool bIsConstructionScript = Graph->GetFName() == UEdGraphSchema_K2::FN_UserConstructionScript;
-		
 		int32 SectionID = NodeSectionID::FUNCTION;
 
 		if(Graph->IsA<UAnimationGraph>())
@@ -1368,15 +1429,7 @@ void SMyBlueprint::CollectAllActions(FGraphActionListBuilderBase& OutAllActions)
 			}
 		}
 
-		TSharedPtr<FEdGraphSchemaAction_K2Graph> NewFuncAction = MakeShareable(new FEdGraphSchemaAction_K2Graph(EEdGraphSchemaAction_K2Graph::Function, FunctionCategory, DisplayInfo.PlainName, DisplayInfo.Tooltip, bIsConstructionScript ? 2 : 1, SectionID));
-		NewFuncAction->FuncName = Graph->GetFName();
-		NewFuncAction->EdGraph = Graph;
-
-		const FString UserCategoryName = FEditorCategoryUtils::GetCategoryDisplayString(FunctionCategory.ToString());
-		SortList.AddAction(UserCategoryName, NewFuncAction);
-
-		GetChildGraphs(Graph, NewFuncAction->GetSectionID(), SortList, FunctionCategory);
-		GetChildEvents(Graph, NewFuncAction->GetSectionID(), SortList, FunctionCategory);
+		AddEventForFunctionGraph(Graph, SectionID, SortList, FText::GetEmpty(), true);
 
 		ImplementedFunctionCache.Add(Graph->GetFName());
 	}
@@ -1509,6 +1562,7 @@ void SMyBlueprint::CollectAllActions(FGraphActionListBuilderBase& OutAllActions)
 					if(bIsAnimFunction && NewFuncAction->EdGraph)
 					{
 						GetChildGraphs(NewFuncAction->EdGraph, NewFuncAction->GetSectionID(), SortList, FunctionCategory);
+						GetChildEvents(NewFuncAction->EdGraph, NewFuncAction->GetSectionID(), SortList, FunctionCategory);
 					}
 				}
 			}
@@ -2204,6 +2258,7 @@ TSharedPtr<SWidget> SMyBlueprint::OnContextMenuOpening()
 		{
 			MenuBuilder.AddMenuEntry(FMyBlueprintCommands::Get().OpenGraph);
 			MenuBuilder.AddMenuEntry(FMyBlueprintCommands::Get().OpenGraphInNewTab);
+			MenuBuilder.AddMenuEntry(FMyBlueprintCommands::Get().OpenExternalGraph);
 			MenuBuilder.AddMenuEntry(FMyBlueprintCommands::Get().FocusNode);
 			MenuBuilder.AddMenuEntry(FMyBlueprintCommands::Get().FocusNodeInNewTab);
 			MenuBuilder.AddMenuEntry(FGenericCommands::Get().Rename, NAME_None, LOCTEXT("Rename", "Rename"), LOCTEXT("Rename_Tooltip", "Renames this function or variable from blueprint.") );
@@ -2377,7 +2432,14 @@ bool SMyBlueprint::CanOpenGraph() const
 	return (bGraph || bDelegate) && BlueprintEditorPtr.IsValid();
 }
 
-void SMyBlueprint::OpenGraph(FDocumentTracker::EOpenDocumentCause InCause)
+bool SMyBlueprint::CanOpenExternalGraph() const 
+{
+	const FEdGraphSchemaAction_K2Graph* GraphAction = SelectionAsGraph();
+	const bool bGraph = GraphAction && GraphAction->EdGraph;
+	return CanOpenGraph() && bGraph && !BlueprintEditorPtr.Pin()->IsGraphInCurrentBlueprint(GraphAction->EdGraph);
+}
+
+void SMyBlueprint::OpenGraph(FDocumentTracker::EOpenDocumentCause InCause, bool bOpenExternalGraphInNewEditor)
 {
 	UEdGraph* GraphToOpen = nullptr;
 
@@ -2405,14 +2467,28 @@ void SMyBlueprint::OpenGraph(FDocumentTracker::EOpenDocumentCause InCause)
 	
 	if (GraphToOpen)
 	{
-		BlueprintEditorPtr.Pin()->OpenDocument(GraphToOpen, InCause);
+		if(bOpenExternalGraphInNewEditor && !BlueprintEditorPtr.Pin()->IsGraphInCurrentBlueprint(GraphToOpen))
+		{
+			if(UBlueprint* OtherBlueprint = FBlueprintEditorUtils::FindBlueprintForGraph(GraphToOpen))
+			{
+				GEditor->GetEditorSubsystem<UAssetEditorSubsystem>()->OpenEditorForAsset(OtherBlueprint);
+				if(IBlueprintEditor* OtherBlueprintEditor = static_cast<IBlueprintEditor*>(GEditor->GetEditorSubsystem<UAssetEditorSubsystem>()->FindEditorForAsset(OtherBlueprint, true)))
+				{
+					OtherBlueprintEditor->JumpToHyperlink(GraphToOpen, false);
+				}
+			}
+		}
+		else
+		{
+			BlueprintEditorPtr.Pin()->OpenDocument(GraphToOpen, InCause);
+		}
 	}
 }
 
 
 void SMyBlueprint::OnOpenGraph()
 {
-	OpenGraph(FDocumentTracker::OpenNewDocument);	
+	OpenGraph(FDocumentTracker::OpenNewDocument);
 }
 
 void SMyBlueprint::OnOpenGraphInNewTab()
@@ -2420,6 +2496,11 @@ void SMyBlueprint::OnOpenGraphInNewTab()
 	OpenGraph(FDocumentTracker::ForceOpenNewDocument);	
 }
 
+void SMyBlueprint::OnOpenExternalGraph()
+{
+	OpenGraph(FDocumentTracker::OpenNewDocument, true);
+}
+
 bool SMyBlueprint::CanFocusOnNode() const
 {
 	FEdGraphSchemaAction_K2Event const* const EventAction = SelectionAsEvent();
diff --git a/Engine/Source/Editor/Kismet/Private/SMyBlueprint.h b/Engine/Source/Editor/Kismet/Private/SMyBlueprint.h
--- a/Engine/Source/Editor/Kismet/Private/SMyBlueprint.h
+++ b/Engine/Source/Editor/Kismet/Private/SMyBlueprint.h
@@ -39,6 +39,7 @@ public:
 	// Basic operations
 	TSharedPtr<FUICommandInfo> OpenGraph;
 	TSharedPtr<FUICommandInfo> OpenGraphInNewTab;
+	TSharedPtr<FUICommandInfo> OpenExternalGraph;
 	TSharedPtr<FUICommandInfo> FocusNode;
 	TSharedPtr<FUICommandInfo> FocusNodeInNewTab;
 	TSharedPtr<FUICommandInfo> ImplementFunction;
@@ -155,8 +156,9 @@ private:
 	/** Callback used to populate all actions list in SGraphActionMenu */
 	void CollectAllActions(FGraphActionListBuilderBase& OutAllActions);
 	void CollectStaticSections(TArray<int32>& StaticSectionIDs);
-	void GetChildGraphs(UEdGraph* EdGraph, int32 const SectionId, FGraphActionSort& SortList, const FText& ParentCategory = FText::GetEmpty());
-	void GetChildEvents(UEdGraph const* EdGraph, int32 const SectionId, FGraphActionSort& SortList, const FText& ParentCategory = FText::GetEmpty()) const;
+	void AddEventForFunctionGraph(UEdGraph* EdGraph, int32 const SectionId, FGraphActionSort& SortList, const FText& ParentCategory = FText::GetEmpty(), bool bAddChildGraphs = true) const;
+	void GetChildGraphs(UEdGraph* EdGraph, int32 const SectionId, FGraphActionSort& SortList, const FText& ParentCategory = FText::GetEmpty()) const;
+	void GetChildEvents(UEdGraph const* EdGraph, int32 const SectionId, FGraphActionSort& SortList, const FText& ParentCategory = FText::GetEmpty(), bool bInAddChildGraphs = true) const;
 	void GetLocalVariables(FGraphActionSort& SortList) const;
 	
 	/** Handles the visibility of the local action list */
@@ -213,7 +215,7 @@ private:
 	bool GetShowAccessSpecifier() const;
 
 	/** Helper function to open the selected graph */
-	void OpenGraph(FDocumentTracker::EOpenDocumentCause InCause);
+	void OpenGraph(FDocumentTracker::EOpenDocumentCause InCause, bool bOpenExternalGraphInNewEditor = false);
 
 	/**
 	* Check if the override of a given function is most likely desired as a blueprint function 
@@ -228,7 +230,9 @@ private:
 	/** Callbacks for commands */
 	void OnOpenGraph();
 	void OnOpenGraphInNewTab();
+	void OnOpenExternalGraph();
 	bool CanOpenGraph() const;
+	bool CanOpenExternalGraph() const;
 	bool CanFocusOnNode() const;
 	void OnFocusNode();
 	void OnFocusNodeInNewTab();
