Added details customization for picking FMemberReferences

Currently only supports function references when exposed via the details panel. Initial use case is for function bindings on anim nodes. For the customization to apply the FMemberReference must specify the "FunctionReference" metadata.
Prototypes for functions are exposed via metadata "PrototypeFunction" using the full package path to the function.

diff --git a/Engine/Plugins/Runtime/PropertyAccess/Source/PropertyAccessEditor/Private/SPropertyBinding.cpp b/Engine/Plugins/Runtime/PropertyAccess/Source/PropertyAccessEditor/Private/SPropertyBinding.cpp
--- a/Engine/Plugins/Runtime/PropertyAccess/Source/PropertyAccessEditor/Private/SPropertyBinding.cpp
+++ b/Engine/Plugins/Runtime/PropertyAccess/Source/PropertyAccessEditor/Private/SPropertyBinding.cpp
@@ -458,52 +458,55 @@ void SPropertyBinding::FillPropertyMenu(FMenuBuilder& MenuBuilder, UStruct* InOw
 	//---------------------------------------
 	// Function Bindings
 
-	if ( UClass* OwnerClass = Cast<UClass>(InOwnerStruct) )
+	if (Args.bAllowFunctionBindings)
 	{
-		static FName FunctionIcon(TEXT("GraphEditor.Function_16x"));
-
-		MenuBuilder.BeginSection("Functions", LOCTEXT("Functions", "Functions"));
+		if ( UClass* OwnerClass = Cast<UClass>(InOwnerStruct) )
 		{
-			ForEachBindableFunction(OwnerClass, [this, &InBindingChain, &MenuBuilder, &MakeFunctionWidget] (const FFunctionInfo& Info) 
-			{
-				TArray<TSharedPtr<FBindingChainElement>> NewBindingChain(InBindingChain);
-				NewBindingChain.Emplace(MakeShared<FBindingChainElement>(Info.Function));
+			static FName FunctionIcon(TEXT("GraphEditor.Function_16x"));
 
-				FProperty* ReturnProperty = Info.Function->GetReturnProperty();
-				// We can get here if we accept non-leaf UObject functions, so if so we need to check the return value for compatibility
-				if(!Args.bAllowUObjectFunctions || Args.OnCanBindProperty.Execute(ReturnProperty))
+			MenuBuilder.BeginSection("Functions", LOCTEXT("Functions", "Functions"));
+			{
+				ForEachBindableFunction(OwnerClass, [this, &InBindingChain, &MenuBuilder, &MakeFunctionWidget] (const FFunctionInfo& Info) 
 				{
-					MenuBuilder.AddMenuEntry(
-						FUIAction(FExecuteAction::CreateSP(this, &SPropertyBinding::HandleAddBinding, NewBindingChain)),
-						MakeFunctionWidget(Info));
-				}
+					TArray<TSharedPtr<FBindingChainElement>> NewBindingChain(InBindingChain);
+					NewBindingChain.Emplace(MakeShared<FBindingChainElement>(Info.Function));
 
-				// Only show bindable subobjects, structs and variables if we're generating pure bindings.
-				if(Args.bGeneratePureBindings)
-				{
-					if(FObjectPropertyBase* ObjectPropertyBase = CastField<FObjectPropertyBase>(ReturnProperty))
+					FProperty* ReturnProperty = Info.Function->GetReturnProperty();
+					// We can get here if we accept non-leaf UObject functions, so if so we need to check the return value for compatibility
+					if(!Args.bAllowUObjectFunctions || Args.OnCanBindProperty.Execute(ReturnProperty))
 					{
-						MenuBuilder.AddSubMenu(
-							MakeFunctionWidget(Info),
-							FNewMenuDelegate::CreateSP(this, &SPropertyBinding::FillPropertyMenu, static_cast<UStruct*>(ObjectPropertyBase->PropertyClass), NewBindingChain));
+						MenuBuilder.AddMenuEntry(
+							FUIAction(FExecuteAction::CreateSP(this, &SPropertyBinding::HandleAddBinding, NewBindingChain)),
+							MakeFunctionWidget(Info));
 					}
-					else if(FStructProperty* StructProperty = CastField<FStructProperty>(ReturnProperty))
+
+					// Only show bindable subobjects, structs and variables if we're generating pure bindings.
+					if(Args.bGeneratePureBindings)
 					{
-						MenuBuilder.AddSubMenu(
-							MakeFunctionWidget(Info),
-							FNewMenuDelegate::CreateSP(this, &SPropertyBinding::FillPropertyMenu, static_cast<UStruct*>(StructProperty->Struct), NewBindingChain));
+						if(FObjectPropertyBase* ObjectPropertyBase = CastField<FObjectPropertyBase>(ReturnProperty))
+						{
+							MenuBuilder.AddSubMenu(
+								MakeFunctionWidget(Info),
+								FNewMenuDelegate::CreateSP(this, &SPropertyBinding::FillPropertyMenu, static_cast<UStruct*>(ObjectPropertyBase->PropertyClass), NewBindingChain));
+						}
+						else if(FStructProperty* StructProperty = CastField<FStructProperty>(ReturnProperty))
+						{
+							MenuBuilder.AddSubMenu(
+								MakeFunctionWidget(Info),
+								FNewMenuDelegate::CreateSP(this, &SPropertyBinding::FillPropertyMenu, static_cast<UStruct*>(StructProperty->Struct), NewBindingChain));
+						}
 					}
-				}
-			});
+				});
+			}
+			MenuBuilder.EndSection(); //Functions
 		}
-		MenuBuilder.EndSection(); //Functions
 	}
 
 	//---------------------------------------
 	// Property Bindings
 
 	// Only show bindable subobjects and variables if we're generating pure bindings.
-	if ( Args.bGeneratePureBindings )
+	if ( Args.bGeneratePureBindings && Args.bAllowPropertyBindings )
 	{
 		FProperty* BindingProperty = nullptr;
 		if(Args.BindableSignature != nullptr)
@@ -744,6 +747,8 @@ void SPropertyBinding::HandleCreateAndAddBinding()
 	UFunction* Function = Blueprint->SkeletonGeneratedClass->FindFunctionByName(FunctionGraph->GetFName());
 	check(Function);
 
+	Args.OnNewFunctionBindingCreated.ExecuteIfBound(FunctionGraph, Function);
+	
 	// Add the binding to the blueprint
 	TArray<TSharedPtr<FBindingChainElement>> BindingPath;
 	BindingPath.Emplace(MakeShared<FBindingChainElement>(Function));
diff --git a/Engine/Source/Editor/Kismet/Private/BlueprintDetailsCustomization.cpp b/Engine/Source/Editor/Kismet/Private/BlueprintDetailsCustomization.cpp
--- a/Engine/Source/Editor/Kismet/Private/BlueprintDetailsCustomization.cpp
+++ b/Engine/Source/Editor/Kismet/Private/BlueprintDetailsCustomization.cpp
@@ -82,6 +82,7 @@
 #include "UObject/TextProperty.h"
 #include "Subsystems/AssetEditorSubsystem.h"
 #include "SupportedRangeTypes.h"	// StructsSupportingRangeVisibility
+#include "IPropertyAccessEditor.h"
 
 #define LOCTEXT_NAMESPACE "BlueprintDetailsCustomization"
 
@@ -6891,5 +6892,197 @@ TSharedRef<SWidget> FBlueprintDocumentationDetails::GenerateExcerptList()
 		];
 }
 
+void FBlueprintMemberReferenceDetails::CustomizeHeader(TSharedRef<IPropertyHandle> InStructPropertyHandle, FDetailWidgetRow& InHeaderRow, IPropertyTypeCustomizationUtils& InStructCustomizationUtils)
+{
+	if(IModularFeatures::Get().IsModularFeatureAvailable("PropertyAccessEditor"))
+	{
+		UBlueprint* Blueprint = MyBlueprint.Pin()->GetBlueprintObj();
+		if(Blueprint == nullptr)
+		{
+			return;
+		}
+
+		const bool bFunctionReference = InStructPropertyHandle->HasMetaData("FunctionReference");
+		const bool bPropertyReference = InStructPropertyHandle->HasMetaData("PropertyReference");
+		if(bPropertyReference || !bFunctionReference)
+		{
+			// Only function references are supported right now
+			return;
+		}
+		
+		// Try to find the prototype function
+		const FString& PrototypeFunctionName = InStructPropertyHandle->GetMetaData("PrototypeFunction");
+		UFunction* PrototypeFunction = PrototypeFunctionName.IsEmpty() ? nullptr : FindObject<UFunction>(nullptr, *PrototypeFunctionName);
+
+		auto OnGenerateBindingName = []() -> FString
+		{
+			return TEXT("NewFunction");
+		};
+		
+		auto OnCanBindProperty = [](FProperty* InProperty)
+		{
+			return true;
+		};
+		
+		auto OnCanBindFunction = [PrototypeFunction](UFunction* InFunction)
+		{
+			if(PrototypeFunction != nullptr)
+			{
+				return PrototypeFunction->IsSignatureCompatibleWith(InFunction)
+					&& FBlueprintEditorUtils::HasFunctionBlueprintThreadSafeMetaData(PrototypeFunction) == FBlueprintEditorUtils::HasFunctionBlueprintThreadSafeMetaData(InFunction);
+			}
+			
+			return false;
+		};
+
+		auto OnAddBinding = [InStructPropertyHandle, Blueprint](FName InPropertyName, const TArray<FBindingChainElement>& InBindingChain)
+		{
+			void* StructData = nullptr;
+			const FPropertyAccess::Result Result = InStructPropertyHandle->GetValueData(StructData);
+			if(Result == FPropertyAccess::Success)
+			{
+				InStructPropertyHandle->NotifyPreChange();
+				
+				check(StructData);
+				FMemberReference* MemberReference = static_cast<FMemberReference*>(StructData);
+				MemberReference->SetFromField<UFunction>(InBindingChain[0].Field.Get<UFunction>(), true);
+
+				InStructPropertyHandle->NotifyPostChange(EPropertyChangeType::ValueSet);
+			}
+		};
+
+		auto OnRemoveBinding = [InStructPropertyHandle](FName InPropertyName)
+		{
+			
+			
+			void* StructData = nullptr;
+			const FPropertyAccess::Result Result = InStructPropertyHandle->GetValueData(StructData);
+			if(Result == FPropertyAccess::Success)
+			{
+				InStructPropertyHandle->NotifyPreChange();
+				
+				check(StructData);
+				FMemberReference* MemberReference = static_cast<FMemberReference*>(StructData);
+				FGuid EmptyGuid;
+				MemberReference->SetSelfMember(NAME_None, EmptyGuid);
+
+				InStructPropertyHandle->NotifyPostChange(EPropertyChangeType::ValueSet);
+			}
+		};
+
+		auto CanRemoveBinding = [InStructPropertyHandle](FName InPropertyName)
+		{
+			void* StructData = nullptr;
+			const FPropertyAccess::Result Result = InStructPropertyHandle->GetValueData(StructData);
+			if(Result == FPropertyAccess::Success)
+			{
+				check(StructData);
+				FMemberReference* MemberReference = static_cast<FMemberReference*>(StructData);
+				return MemberReference->GetMemberName() != NAME_None;
+			}
+			
+			return false;
+		};
+
+		auto OnNewFunctionBindingCreated = [PrototypeFunction](UEdGraph* InFunctionGraph, UFunction* InFunction)
+		{
+			// Ensure newly created function is thread safe 
+			if(PrototypeFunction && FBlueprintEditorUtils::HasFunctionBlueprintThreadSafeMetaData(PrototypeFunction))
+			{
+				TArray<UK2Node_FunctionEntry*> EntryNodes;
+				InFunctionGraph->GetNodesOfClass<UK2Node_FunctionEntry>(EntryNodes);
+				if(EntryNodes.Num() > 0)
+				{
+					EntryNodes[0]->MetaData.bThreadSafe = true;
+				}
+			}
+		};
+
+		auto CurrentBindingText = [bFunctionReference, Blueprint, InStructPropertyHandle]()
+		{
+			void* StructData = nullptr;
+			const FPropertyAccess::Result Result = InStructPropertyHandle->GetValueData(StructData);
+			if(Result == FPropertyAccess::Success)
+			{
+				check(StructData);
+				FMemberReference* MemberReference = static_cast<FMemberReference*>(StructData);
+				if(bFunctionReference)
+				{
+					UFunction* Function = MemberReference->ResolveMember<UFunction>(Blueprint->SkeletonGeneratedClass);
+					if(Function)
+					{
+						return FText::FromName(Function->GetFName());
+					}
+					else
+					{
+						return FText::FromName(MemberReference->GetMemberName());
+					}
+				}
+			}
+			else if(Result == FPropertyAccess::MultipleValues)
+			{
+				return LOCTEXT("MultipleValues", "Multiple Values");
+			}
+
+			return FText::GetEmpty();
+		};
+
+		auto CurrentBindingToolTipText = [InStructPropertyHandle]()
+		{
+			void* StructData = nullptr;
+			const FPropertyAccess::Result Result = InStructPropertyHandle->GetValueData(StructData);
+			if(Result == FPropertyAccess::Success)
+			{
+				check(StructData);
+				FMemberReference* MemberReference = static_cast<FMemberReference*>(StructData);
+				return FText::FromName(MemberReference->GetMemberName());
+			}
+			else if(Result == FPropertyAccess::MultipleValues)
+			{
+				return LOCTEXT("MultipleValues", "Multiple Values");
+			}
+
+			return FText::GetEmpty();
+		};
+	
+		FPropertyBindingWidgetArgs Args;
+		Args.BindableSignature = PrototypeFunction;
+		Args.OnGenerateBindingName = FOnGenerateBindingName::CreateLambda(OnGenerateBindingName);
+		Args.OnCanBindProperty = FOnCanBindProperty::CreateLambda(OnCanBindProperty);
+		Args.OnCanBindFunction = FOnCanBindFunction::CreateLambda(OnCanBindFunction);
+		Args.OnCanBindToClass = FOnCanBindToClass::CreateLambda([](UClass* InClass){ return true; });
+		Args.OnAddBinding = FOnAddBinding::CreateLambda(OnAddBinding);
+		Args.OnRemoveBinding = FOnRemoveBinding::CreateLambda(OnRemoveBinding);
+		Args.OnCanRemoveBinding = FOnCanRemoveBinding::CreateLambda(CanRemoveBinding);
+		Args.OnNewFunctionBindingCreated = FOnNewFunctionBindingCreated::CreateLambda(OnNewFunctionBindingCreated);
+		Args.CurrentBindingText = MakeAttributeLambda(CurrentBindingText);
+		Args.CurrentBindingToolTipText = MakeAttributeLambda(CurrentBindingToolTipText);
+		Args.CurrentBindingImage = FEditorStyle::GetBrush("GraphEditor.Function_16x");
+		Args.CurrentBindingColor = FEditorStyle::GetSlateColor("Colors.Foreground").GetSpecifiedColor();
+		Args.bGeneratePureBindings = false;
+		Args.bAllowFunctionBindings = bFunctionReference;
+		Args.bAllowPropertyBindings = false;
+		Args.bAllowNewBindings = true;
+		Args.bAllowArrayElementBindings = false;
+		Args.bAllowUObjectFunctions = false;
+		Args.bAllowStructFunctions = false;
+		Args.bAllowStructMemberBindings = false;
+
+		IPropertyAccessEditor& PropertyAccessEditor = IModularFeatures::Get().GetModularFeature<IPropertyAccessEditor>("PropertyAccessEditor");
+		InHeaderRow
+		.NameContent()
+		[
+			InStructPropertyHandle->CreatePropertyNameWidget()
+		]
+		.ValueContent()
+		[
+			SNew(SBox)
+			.MaxDesiredWidth(200.0f)
+			[
+				PropertyAccessEditor.MakePropertyBindingWidget(Blueprint, Args)
+			]
+		];
+	}
+}
 
 #undef LOCTEXT_NAMESPACE
diff --git a/Engine/Source/Editor/Kismet/Private/BlueprintDetailsCustomization.h b/Engine/Source/Editor/Kismet/Private/BlueprintDetailsCustomization.h
--- a/Engine/Source/Editor/Kismet/Private/BlueprintDetailsCustomization.h
+++ b/Engine/Source/Editor/Kismet/Private/BlueprintDetailsCustomization.h
@@ -1048,3 +1048,25 @@ private:
 	TArray<TSharedPtr<FString>> ExcerptList;
 
 };
+
+class FBlueprintMemberReferenceDetails : public IPropertyTypeCustomization
+{
+public:
+	static TSharedRef<IPropertyTypeCustomization> MakeInstance(TWeakPtr<SMyBlueprint> InMyBlueprint)
+	{
+		return MakeShareable(new FBlueprintMemberReferenceDetails(InMyBlueprint));
+	}
+
+private:
+	FBlueprintMemberReferenceDetails(TWeakPtr<SMyBlueprint> InMyBlueprint)
+		: MyBlueprint(InMyBlueprint)
+	{
+	}
+	
+	// IPropertyTypeCustomization interface
+	virtual void CustomizeHeader(TSharedRef<IPropertyHandle> InStructPropertyHandle, FDetailWidgetRow& InHeaderRow, IPropertyTypeCustomizationUtils& InStructCustomizationUtils) override;
+	virtual void CustomizeChildren(TSharedRef<IPropertyHandle> InStructPropertyHandle, IDetailChildrenBuilder& InChildBuilder, IPropertyTypeCustomizationUtils& InStructCustomizationUtils) override {};
+
+private:
+	TWeakPtr<SMyBlueprint> MyBlueprint;
+};
diff --git a/Engine/Source/Editor/Kismet/Private/SKismetInspector.cpp b/Engine/Source/Editor/Kismet/Private/SKismetInspector.cpp
--- a/Engine/Source/Editor/Kismet/Private/SKismetInspector.cpp
+++ b/Engine/Source/Editor/Kismet/Private/SKismetInspector.cpp
@@ -376,6 +376,8 @@ void SKismetInspector::Construct(const FArguments& InArgs)
 		PropertyView->RegisterInstancedCustomPropertyLayout(UPropertyWrapper::StaticClass(), LayoutVariableDetails);
 		PropertyView->RegisterInstancedCustomPropertyLayout(UK2Node_VariableGet::StaticClass(), LayoutVariableDetails);
 		PropertyView->RegisterInstancedCustomPropertyLayout(UK2Node_VariableSet::StaticClass(), LayoutVariableDetails);
+
+		PropertyView->RegisterInstancedCustomPropertyTypeLayout(FMemberReference::StaticStruct()->GetFName(), FOnGetPropertyTypeCustomizationInstance::CreateStatic(&FBlueprintMemberReferenceDetails::MakeInstance, MyBlueprint));
 	}
 
 	if (Kismet2.IsValid() && Kismet2->IsEditingSingleBlueprint())
diff --git a/Engine/Source/Editor/UnrealEd/Public/IPropertyAccessEditor.h b/Engine/Source/Editor/UnrealEd/Public/IPropertyAccessEditor.h
--- a/Engine/Source/Editor/UnrealEd/Public/IPropertyAccessEditor.h
+++ b/Engine/Source/Editor/UnrealEd/Public/IPropertyAccessEditor.h
@@ -118,6 +118,9 @@ DECLARE_DELEGATE_OneParam(FOnRemoveBinding, FName /*InPropertyName*/);
 /** Delegate called to see if we can remove remove a binding (ie. if it exists) */
 DECLARE_DELEGATE_RetVal_OneParam(bool, FOnCanRemoveBinding, FName /*InPropertyName*/);
 
+/** Delegate called once a new function binding has been created */
+DECLARE_DELEGATE_TwoParams(FOnNewFunctionBindingCreated, UEdGraph* /*InFunctionGraph*/, UFunction* /*InFunction*/);
+
 /** Setup arguments structure for a property binding widget */
 struct FPropertyBindingWidgetArgs
 {
@@ -157,6 +160,9 @@ struct FPropertyBindingWidgetArgs
 	/** Delegate called to see if we can remove remove a binding (ie. if it exists) */
 	FOnCanRemoveBinding OnCanRemoveBinding;
 
+	/** Delegate called once a new function binding has been created */
+	FOnNewFunctionBindingCreated OnNewFunctionBindingCreated;
+	
 	/** The current binding's text label */
 	TAttribute<FText> CurrentBindingText;
 
@@ -175,6 +181,12 @@ struct FPropertyBindingWidgetArgs
 	/** Whether to generate pure bindings */
 	bool bGeneratePureBindings = true;
 
+	/** Whether to allow function bindings */
+	bool bAllowFunctionBindings = true;
+
+	/** Whether to allow property bindings */
+	bool bAllowPropertyBindings = true;	
+	
 	/** Whether to allow array element bindings */
 	bool bAllowArrayElementBindings = false;
 
