Fix: Call In Editor checkbox shows up on macros (and cannot be checked) but should not show up at all (it's irrelevant to macros)

diff --git a/Engine/Source/Editor/Kismet/Private/BlueprintDetailsCustomization.cpp b/Engine/Source/Editor/Kismet/Private/BlueprintDetailsCustomization.cpp
--- a/Engine/Source/Editor/Kismet/Private/BlueprintDetailsCustomization.cpp
+++ b/Engine/Source/Editor/Kismet/Private/BlueprintDetailsCustomization.cpp
@@ -3678,7 +3678,8 @@ void FBlueprintGraphActionDetails::CustomizeDetails( IDetailLayoutBuilder& Detai
 				]
 			];
 		}
-		const bool bShowCallInEditor = bIsCustomEvent || FBlueprintEditorUtils::IsEditorUtilityBlueprint( GetBlueprintObj() ) || (FunctionEntryNode && FunctionEntryNode->IsEditable());
+		
+		const bool bShowCallInEditor = bIsCustomEvent || bIsFunctionGraph;
 		if( bShowCallInEditor )
 		{
 			Category.AddCustomRow( LOCTEXT( "EditorCallable", "Call In Editor" ) )
@@ -4023,6 +4024,10 @@ void FBlueprintGraphActionDetails::OnEditorCallableEventModified( const ECheckBo
 			EntryPoint->MetaData.bCallInEditor = bCallInEditor;
 			FBlueprintEditorUtils::MarkBlueprintAsModified( EntryPoint->GetBlueprint() );
 		}
+		else
+		{
+			checkf(false, TEXT("Only Events and Functions are Callable In Editor"));
+		}
 	}
 }
 
