Fix crash in Blueprint diff tool

diff --git a/Engine/Source/Editor/Kismet/Private/SKismetInspector.cpp b/Engine/Source/Editor/Kismet/Private/SKismetInspector.cpp
--- a/Engine/Source/Editor/Kismet/Private/SKismetInspector.cpp
+++ b/Engine/Source/Editor/Kismet/Private/SKismetInspector.cpp
@@ -41,6 +41,7 @@
 #include "BlueprintDetailsCustomization.h"
 #include "K2Node_BitmaskLiteral.h"
 #include "BitmaskLiteralDetails.h"
+#include "BlueprintMemberReferenceCustomization.h"
 #include "FormatTextDetails.h"
 #include "Kismet2/BlueprintEditorUtils.h"
 
@@ -918,7 +919,7 @@ bool SKismetInspector::IsPropertyEditingEnabled() const
 		{
 			if(UEdGraph* OuterGraph = EdGraphNode->GetGraph())
 			{
-				if(!BlueprintEditorPtr.Pin()->IsEditable(OuterGraph))
+				if(BlueprintEditorPtr.IsValid() && !BlueprintEditorPtr.Pin()->IsEditable(OuterGraph))
 				{
 					bIsEditable = false;
 					break;
