Fix CIS issue

Local changes got included in my last submit

diff --git a/Engine/Source/Editor/Kismet/Private/SKismetInspector.cpp b/Engine/Source/Editor/Kismet/Private/SKismetInspector.cpp
--- a/Engine/Source/Editor/Kismet/Private/SKismetInspector.cpp
+++ b/Engine/Source/Editor/Kismet/Private/SKismetInspector.cpp
@@ -41,7 +41,6 @@
 #include "BlueprintDetailsCustomization.h"
 #include "K2Node_BitmaskLiteral.h"
 #include "BitmaskLiteralDetails.h"
-#include "BlueprintMemberReferenceCustomization.h"
 #include "FormatTextDetails.h"
 #include "Kismet2/BlueprintEditorUtils.h"
 
