Add imported namespace filtering for class pickers within the BP editor context; currently disabled by default (experimental feature).

Import filtering can now be dynamically enabled (via CVar) on all pin type selection and class viewer widgets, including those used for reparenting and interface add operations.

diff --git a/Engine/Source/Editor/Kismet/Private/BlueprintDetailsCustomization.cpp b/Engine/Source/Editor/Kismet/Private/BlueprintDetailsCustomization.cpp
--- a/Engine/Source/Editor/Kismet/Private/BlueprintDetailsCustomization.cpp
+++ b/Engine/Source/Editor/Kismet/Private/BlueprintDetailsCustomization.cpp
@@ -290,9 +290,9 @@ void FBlueprintVarActionDetails::CustomizeDetails( IDetailLayoutBuilder& DetailL
 	TSharedPtr<SToolTip> VarTypeTooltip = IDocumentation::Get()->CreateToolTip(LOCTEXT("VarTypeTooltip", "The type of the variable."), NULL, DocLink, TEXT("VariableType"));
 
 	TSharedPtr<IPinTypeSelectorFilter> CustomPinTypeFilter;
-	if (BlueprintEditor.IsValid() && GetDefault<UBlueprintEditorSettings>()->bEnableNamespaceFilteringFeatures)
+	if (BlueprintEditor.IsValid())
 	{
-		CustomPinTypeFilter = BlueprintEditor.Pin()->GetOrCreateNamespaceHelperForBlueprint(GetBlueprintObj())->GetPinTypeSelectorFilter();
+		CustomPinTypeFilter = BlueprintEditor.Pin()->GetImportedPinTypeSelectorFilter();
 	}
 	
 	const UEdGraphSchema* Schema = GetDefault<UEdGraphSchema_K2>();
@@ -2823,9 +2823,9 @@ void FBlueprintGraphArgumentLayout::GenerateHeaderRowContent( FDetailWidgetRow&
 		if (MyBlueprintPtr.IsValid())
 		{
 			TSharedPtr<FBlueprintEditor> BlueprintEditorPtr = MyBlueprintPtr->GetBlueprintEditor().Pin();
-			if (BlueprintEditorPtr.IsValid() && GetDefault<UBlueprintEditorSettings>()->bEnableNamespaceFilteringFeatures)
+			if (BlueprintEditorPtr.IsValid())
 			{
-				CustomPinTypeFilter = BlueprintEditorPtr->GetOrCreateNamespaceHelperForBlueprint(MyBlueprintPtr->GetBlueprintObj())->GetPinTypeSelectorFilter();
+				CustomPinTypeFilter = BlueprintEditorPtr->GetImportedPinTypeSelectorFilter();
 			}
 		}
 	}
diff --git a/Engine/Source/Editor/Kismet/Private/BlueprintEditor.cpp b/Engine/Source/Editor/Kismet/Private/BlueprintEditor.cpp
--- a/Engine/Source/Editor/Kismet/Private/BlueprintEditor.cpp
+++ b/Engine/Source/Editor/Kismet/Private/BlueprintEditor.cpp
@@ -135,6 +135,7 @@
 
 #include "AssetRegistryModule.h"
 #include "BlueprintEditorTabFactories.h"
+#include "ClassViewerFilter.h"
 #include "SPinTypeSelector.h"
 #include "Animation/AnimBlueprint.h"
 #include "AnimStateConduitNode.h"
@@ -229,6 +230,200 @@ namespace BlueprintEditorImpl
 {
 	static const float InstructionFadeDuration = 0.5f;
 
+	/** Consolidates the class viewer filters for each edited BP asset into a single filter option */
+	class FImportedClassViewerFilter : public IClassViewerFilter, public TSharedFromThis<FImportedClassViewerFilter>
+	{
+	public:
+		FImportedClassViewerFilter(const TArray<TSharedRef<FBlueprintNamespaceHelper>>& InNamespaceHelpers)
+		{
+			bIsFilterEnabled = true;
+
+			ClassViewerFilters.Reserve(InNamespaceHelpers.Num());
+			for (const TSharedRef<FBlueprintNamespaceHelper>& NamespaceHelper : InNamespaceHelpers)
+			{
+				TSharedPtr<IClassViewerFilter> Filter = NamespaceHelper->GetClassViewerFilter();
+				if (Filter.IsValid())
+				{
+					ClassViewerFilters.Add(Filter.ToSharedRef());
+				}
+			}
+		}
+
+		// IClassViewerFilter interface
+		virtual bool IsClassAllowed(const FClassViewerInitializationOptions& InInitOptions, const UClass* InClass, TSharedRef<FClassViewerFilterFuncs> InFilterFuncs) override
+		{
+			if (!GetDefault<UBlueprintEditorSettings>()->bEnableNamespaceFilteringFeatures)
+			{
+				return true;
+			}
+
+			if(bIsFilterEnabled)
+			{
+				for (const TSharedRef<IClassViewerFilter>& Filter : ClassViewerFilters)
+				{
+					if (!Filter->IsClassAllowed(InInitOptions, InClass, InFilterFuncs))
+					{
+						return false;
+					}
+				}
+			}
+
+			return true;
+		}
+
+		virtual bool IsUnloadedClassAllowed(const FClassViewerInitializationOptions& InInitOptions, const TSharedRef< const IUnloadedBlueprintData > InBlueprint, TSharedRef< FClassViewerFilterFuncs > InFilterFuncs) override
+		{
+			if (!GetDefault<UBlueprintEditorSettings>()->bEnableNamespaceFilteringFeatures)
+			{
+				return true;
+			}
+
+			if (bIsFilterEnabled)
+			{
+				for (const TSharedRef<IClassViewerFilter>& Filter : ClassViewerFilters)
+				{
+					if (!Filter->IsUnloadedClassAllowed(InInitOptions, InBlueprint, InFilterFuncs))
+					{
+						return false;
+					}
+				}
+			}
+
+			return true;
+		}
+
+		virtual void GetFilterOptions(TArray<TSharedRef<FClassViewerFilterOption>>& OutFilterOptions)
+		{
+			if (!GetDefault<UBlueprintEditorSettings>()->bEnableNamespaceFilteringFeatures)
+			{
+				return;
+			}
+
+			TSharedRef<FClassViewerFilterOption> FilterOption = MakeShared<FClassViewerFilterOption>();
+			FilterOption->bEnabled = bIsFilterEnabled;
+			FilterOption->LabelText = LOCTEXT("ClassViewerNamespaceFilterMenuOptionLabel", "Show Only Imported Types");
+			FilterOption->ToolTipText = LOCTEXT("ClassViewerNamespaceFilterMenuOptionToolTip", "Don't include non-imported class types.");
+			FilterOption->OnOptionChanged = FOnClassViewerFilterOptionChanged::CreateSP(this, &FImportedClassViewerFilter::OnFilterOptionChanged);
+
+			OutFilterOptions.Add(FilterOption);
+		}
+
+	protected:
+		void OnFilterOptionChanged(bool bIsEnabled)
+		{
+			bIsFilterEnabled = bIsEnabled;
+		}
+
+	private:
+		/** Class viewer filters (one per edited BP asset). */
+		TArray<TSharedRef<IClassViewerFilter>> ClassViewerFilters;
+
+		/** Whether or not the filter is enabled. */
+		bool bIsFilterEnabled;
+	};
+
+	/** Consolidates the pin type selector filters for each edited BP asset into a singular construct */
+	class FImportedPinTypeSelectorFilter : public IPinTypeSelectorFilter, public TSharedFromThis<FImportedPinTypeSelectorFilter>
+	{
+		DECLARE_MULTICAST_DELEGATE(FOnFilterChanged);
+	
+	public:
+		FImportedPinTypeSelectorFilter(const TArray<TSharedRef<FBlueprintNamespaceHelper>>& InNamespaceHelpers)
+		{
+			bIsFilterEnabled = true;
+
+			PinTypeSelectorFilters.Reserve(InNamespaceHelpers.Num());
+			for (const TSharedRef<FBlueprintNamespaceHelper>& NamespaceHelper : InNamespaceHelpers)
+			{
+				TSharedPtr<IPinTypeSelectorFilter> Filter = NamespaceHelper->GetPinTypeSelectorFilter();
+				if (Filter.IsValid())
+				{
+					PinTypeSelectorFilters.Add(Filter.ToSharedRef());
+				}
+			}
+		}
+
+		// IPinTypeSelectorFilter interface
+		virtual bool ShouldShowPinTypeTreeItem(FPinTypeTreeItem InItem) const override
+		{
+			if (!GetDefault<UBlueprintEditorSettings>()->bEnableNamespaceFilteringFeatures)
+			{
+				return true;
+			}
+
+			if (bIsFilterEnabled)
+			{
+				for (const TSharedRef<IPinTypeSelectorFilter>& Filter : PinTypeSelectorFilters)
+				{
+					if (!Filter->ShouldShowPinTypeTreeItem(InItem))
+					{
+						return false;
+					}
+				}
+			}
+
+			return true;
+		}
+
+		virtual FDelegateHandle RegisterOnFilterChanged(FSimpleDelegate InOnFilterChanged) override
+		{
+			return OnFilterChanged.Add(InOnFilterChanged);
+		}
+
+		virtual void UnregisterOnFilterChanged(FDelegateHandle InDelegateHandle) override
+		{
+			OnFilterChanged.Remove(InDelegateHandle);
+		}
+
+		virtual TSharedPtr<SWidget> GetFilterOptionsWidget() override
+		{
+			if (!GetDefault<UBlueprintEditorSettings>()->bEnableNamespaceFilteringFeatures)
+			{
+				return nullptr;
+			}
+
+			if (!FilterOptionsWidget.IsValid())
+			{
+				SAssignNew(FilterOptionsWidget, SCheckBox)
+					.IsChecked(this, &FImportedPinTypeSelectorFilter::IsFilterToggleChecked)
+					.OnCheckStateChanged(this, &FImportedPinTypeSelectorFilter::OnToggleFilter)
+					[
+						SNew(STextBlock)
+						.Text(LOCTEXT("PinTypeNamespaceFilterToggleOptionLabel", "Hide Non-Imported Types"))
+					];
+			}
+
+			return FilterOptionsWidget;
+		}
+
+	protected:
+		ECheckBoxState IsFilterToggleChecked() const
+		{
+			return bIsFilterEnabled ? ECheckBoxState::Checked : ECheckBoxState::Unchecked;
+		}
+
+		void OnToggleFilter(ECheckBoxState NewState)
+		{
+			bIsFilterEnabled = (NewState == ECheckBoxState::Checked);
+
+			// Notify any listeners that the filter has been changed.
+			OnFilterChanged.Broadcast();
+		}
+
+	private:
+		/** Pin type filters (one per edited BP asset). */
+		TArray<TSharedRef<IPinTypeSelectorFilter>> PinTypeSelectorFilters;
+
+		/** Cached filter options widget. */
+		TSharedPtr<SWidget> FilterOptionsWidget;
+
+		/** Delegate that's called whenever filter options are changed. */
+		FOnFilterChanged OnFilterChanged;
+
+		/** Whether or not the filter is enabled. */
+		bool bIsFilterEnabled;
+	};
+
 	/**
 	 * Utility function that will check to see if the specified graph has any 
 	 * nodes other than those that come default, pre-placed, in the graph.
@@ -1644,6 +1839,16 @@ void FBlueprintEditor::CommonInitialization(const TArray<UBlueprint*>& InitBluep
 		DocumentManager->RegisterDocumentFactory(GraphEditorFactory);
 	}
 
+	// Create imported namespace type filters for value editing.
+	TArray<TSharedRef<FBlueprintNamespaceHelper>> LocalNamespaceHelpersArray;
+	LocalNamespaceHelpersArray.Reserve(InitBlueprints.Num());
+	for (const UBlueprint* BP : InitBlueprints)
+	{
+		LocalNamespaceHelpersArray.Add(GetOrCreateNamespaceHelperForBlueprint(BP));
+	}
+	ImportedClassViewerFilter = MakeShared<BlueprintEditorImpl::FImportedClassViewerFilter>(LocalNamespaceHelpersArray);
+	ImportedPinTypeSelectorFilter = MakeShared<BlueprintEditorImpl::FImportedPinTypeSelectorFilter>(LocalNamespaceHelpersArray);
+
 	// Make sure we know when tabs become active to update details tab
 	OnActiveTabChangedDelegateHandle = FGlobalTabmanager::Get()->OnActiveTabChanged_Subscribe( FOnActiveTabChanged::FDelegate::CreateRaw(this, &FBlueprintEditor::OnActiveTabChanged) );
 
diff --git a/Engine/Source/Editor/Kismet/Private/BlueprintNamespaceHelper.cpp b/Engine/Source/Editor/Kismet/Private/BlueprintNamespaceHelper.cpp
--- a/Engine/Source/Editor/Kismet/Private/BlueprintNamespaceHelper.cpp
+++ b/Engine/Source/Editor/Kismet/Private/BlueprintNamespaceHelper.cpp
@@ -12,6 +12,7 @@
 #include "SPinTypeSelector.h"
 #include "Widgets/Input/SCheckBox.h"
 #include "Widgets/Text/STextBlock.h"
+#include "ClassViewerFilter.h"
 
 #define LOCTEXT_NAMESPACE "BlueprintNamespaceHelper"
 
@@ -91,46 +92,55 @@ static FAutoConsoleVariableSink CVarUpdateNamespaceFeatureSettingsSink(
 
 // ---
 
-class FPinTypeSelectorNamespaceFilter : public IPinTypeSelectorFilter, public TSharedFromThis<FPinTypeSelectorNamespaceFilter>
+class FClassViewerNamespaceFilter : public IClassViewerFilter
 {
-	DECLARE_MULTICAST_DELEGATE(FOnFilterChanged);
-
 public:
-	FPinTypeSelectorNamespaceFilter(const FBlueprintNamespaceHelper* InNamespaceHelper)
+	FClassViewerNamespaceFilter(const FBlueprintNamespaceHelper* InNamespaceHelper)
 		: CachedNamespaceHelper(InNamespaceHelper)
-		, bIsFilterEnabled(true)
 	{
 	}
 
-	virtual FDelegateHandle RegisterOnFilterChanged(FSimpleDelegate InOnFilterChanged) override
+	virtual bool IsClassAllowed(const FClassViewerInitializationOptions& InInitOptions, const UClass* InClass, TSharedRef<FClassViewerFilterFuncs> InFilterFuncs) override
 	{
-		return OnFilterChanged.Add(InOnFilterChanged);
-	}
+		if (!CachedNamespaceHelper)
+		{
+			return true;
+		}
 
-	virtual void UnregisterOnFilterChanged(FDelegateHandle InDelegateHandle) override
-	{
-		OnFilterChanged.Remove(InDelegateHandle);
+		return CachedNamespaceHelper->IsImportedObject(InClass);
 	}
 
-	virtual TSharedPtr<SWidget> GetFilterOptionsWidget() override
+	virtual bool IsUnloadedClassAllowed(const FClassViewerInitializationOptions& InInitOptions, const TSharedRef<const IUnloadedBlueprintData> InBlueprint, TSharedRef<FClassViewerFilterFuncs> InFilterFuncs) override
 	{
-		if (!FilterOptionsWidget.IsValid())
+		if (!CachedNamespaceHelper)
 		{
-			SAssignNew(FilterOptionsWidget, SCheckBox)
-			.IsChecked(this, &FPinTypeSelectorNamespaceFilter::IsFilterToggleChecked)
-			.OnCheckStateChanged(this, &FPinTypeSelectorNamespaceFilter::OnToggleFilter)
-			[
-				SNew(STextBlock)
-				.Text(LOCTEXT("PinTypeNamespaceFilterToggleOptionLabel", "Hide Non-Imported Types"))
-			];
+			return true;
 		}
 
-		return FilterOptionsWidget;
+		FSoftObjectPath ClassPath(InBlueprint->GetClassPath());
+		return CachedNamespaceHelper->IsImportedObject(ClassPath);
+	}
+
+private:
+	/** Associated namespace helper object. */
+	const FBlueprintNamespaceHelper* CachedNamespaceHelper;
+};
+
+// ---
+
+class FPinTypeSelectorNamespaceFilter : public IPinTypeSelectorFilter
+{
+	DECLARE_MULTICAST_DELEGATE(FOnFilterChanged);
+
+public:
+	FPinTypeSelectorNamespaceFilter(const FBlueprintNamespaceHelper* InNamespaceHelper)
+		: CachedNamespaceHelper(InNamespaceHelper)
+	{
 	}
 
 	virtual bool ShouldShowPinTypeTreeItem(FPinTypeTreeItem InItem) const override
 	{
-		if (!CachedNamespaceHelper || !InItem.IsValid() || !bIsFilterEnabled)
+		if (!CachedNamespaceHelper)
 		{
 			return true;
 		}
@@ -156,32 +166,9 @@ public:
 		return true;
 	}
 
-protected:
-	ECheckBoxState IsFilterToggleChecked() const
-	{
-		return bIsFilterEnabled ? ECheckBoxState::Checked : ECheckBoxState::Unchecked;
-	}
-
-	void OnToggleFilter(ECheckBoxState NewState)
-	{
-		bIsFilterEnabled = (NewState == ECheckBoxState::Checked);
-
-		// Notify any listeners that the filter has been changed.
-		OnFilterChanged.Broadcast();
-	}
-
 private:
 	/** Associated namespace helper object. */
 	const FBlueprintNamespaceHelper* CachedNamespaceHelper;
-
-	/** Cached filter options widget. */
-	TSharedPtr<SWidget> FilterOptionsWidget;
-
-	/** Delegate that's called whenever filter options are changed. */
-	FOnFilterChanged OnFilterChanged;
-
-	/** Whether or not the filter is enabled. */
-	bool bIsFilterEnabled;
 };
 
 // ---
@@ -218,6 +205,7 @@ FBlueprintNamespaceHelper::FBlueprintNamespaceHelper(const UBlueprint* InBluepri
 		}
 	}
 
+	ClassViewerFilter = MakeShared<FClassViewerNamespaceFilter>(this);
 	PinTypeSelectorFilter = MakeShared<FPinTypeSelectorNamespaceFilter>(this);
 }
 
diff --git a/Engine/Source/Editor/Kismet/Private/SBlueprintPalette.cpp b/Engine/Source/Editor/Kismet/Private/SBlueprintPalette.cpp
--- a/Engine/Source/Editor/Kismet/Private/SBlueprintPalette.cpp
+++ b/Engine/Source/Editor/Kismet/Private/SBlueprintPalette.cpp
@@ -755,9 +755,9 @@ private:
 	void ConstructInternal(const FArguments& InArgs)
 	{
 		TSharedPtr<IPinTypeSelectorFilter> CustomPinTypeFilter;
-		if (BlueprintEditorPtr.IsValid() && GetDefault<UBlueprintEditorSettings>()->bEnableNamespaceFilteringFeatures)
+		if (BlueprintEditorPtr.IsValid())
 		{
-			CustomPinTypeFilter = BlueprintEditorPtr.Pin()->GetOrCreateNamespaceHelperForBlueprint(BlueprintObj)->GetPinTypeSelectorFilter();
+			CustomPinTypeFilter = BlueprintEditorPtr.Pin()->GetImportedPinTypeSelectorFilter();
 		}
 
 		const UEdGraphSchema* Schema = GetDefault<UEdGraphSchema_K2>();
diff --git a/Engine/Source/Editor/Kismet/Private/SKismetInspector.cpp b/Engine/Source/Editor/Kismet/Private/SKismetInspector.cpp
--- a/Engine/Source/Editor/Kismet/Private/SKismetInspector.cpp
+++ b/Engine/Source/Editor/Kismet/Private/SKismetInspector.cpp
@@ -43,6 +43,7 @@
 #include "BitmaskLiteralDetails.h"
 #include "FormatTextDetails.h"
 #include "Kismet2/BlueprintEditorUtils.h"
+#include "ClassViewerFilter.h"
 
 #define LOCTEXT_NAMESPACE "KismetInspector"
 
@@ -349,6 +350,14 @@ void SKismetInspector::Construct(const FArguments& InArgs)
 	DetailsViewArgs.bHideSelectionTip = true;
 	DetailsViewArgs.NotifyHook = NotifyHook;
 	DetailsViewArgs.ViewIdentifier = InArgs._ViewIdentifier;
+	if (Kismet2.IsValid())
+	{
+		TSharedPtr<IClassViewerFilter> ImportedClassFilter = Kismet2->GetImportedClassViewerFilter();
+		if (ImportedClassFilter.IsValid())
+		{
+			DetailsViewArgs.ClassViewerFilters.Add(ImportedClassFilter.ToSharedRef());
+		}
+	}
 
 	PropertyView = EditModule.CreateDetailView( DetailsViewArgs );
 		
diff --git a/Engine/Source/Editor/Kismet/Public/BlueprintEditor.h b/Engine/Source/Editor/Kismet/Public/BlueprintEditor.h
--- a/Engine/Source/Editor/Kismet/Public/BlueprintEditor.h
+++ b/Engine/Source/Editor/Kismet/Public/BlueprintEditor.h
@@ -241,6 +241,8 @@ public:
 	virtual int32 GetNumberOfSelectedNodes() const override;
 	virtual void AnalyticsTrackNodeEvent(UBlueprint* Blueprint, UEdGraphNode *GraphNode, bool bNodeDelete = false) const override;
 	void AnalyticsTrackCompileEvent(UBlueprint* Blueprint, int32 NumErrors, int32 NumWarnings) const;
+	virtual TSharedPtr<class IClassViewerFilter> GetImportedClassViewerFilter() const override { return ImportedClassViewerFilter; }
+	virtual TSharedPtr<class IPinTypeSelectorFilter> GetImportedPinTypeSelectorFilter() const override { return ImportedPinTypeSelectorFilter; }
 	//~ End IBlueprintEditor Interface
 
 	//~ Begin FTickableEditorObject Interface
@@ -1347,6 +1349,12 @@ protected:
 	/** Cached set of instanced namespace helper objects */
 	TMap<TWeakObjectPtr<const UBlueprint>, TSharedRef<FBlueprintNamespaceHelper>> CachedNamespaceHelpers;
 
+	/** Filter used to restrict class viewer widgets in the editor context to imported namespaces only */
+	TSharedPtr<class IClassViewerFilter> ImportedClassViewerFilter;
+
+	/** Filter used to restrict pin type selector widgets in the editor context to imported namespaces only */
+	TSharedPtr<class IPinTypeSelectorFilter> ImportedPinTypeSelectorFilter;
+
 	FOnSetPinVisibility OnSetPinVisibility;
 
 	/** Has someone requested a deferred update of the saved document state? */
diff --git a/Engine/Source/Editor/Kismet/Public/BlueprintEditorModule.h b/Engine/Source/Editor/Kismet/Public/BlueprintEditorModule.h
--- a/Engine/Source/Editor/Kismet/Public/BlueprintEditorModule.h
+++ b/Engine/Source/Editor/Kismet/Public/BlueprintEditorModule.h
@@ -98,6 +98,12 @@ public:
 	/** Used to track node create/delete events for Analytics */
 	virtual void AnalyticsTrackNodeEvent( UBlueprint* Blueprint, UEdGraphNode *GraphNode, bool bNodeDelete = false ) const = 0;
 
+	/** Return the class viewer filter associated with the current set of imported namespaces within this editor context. Default is NULL (no filter). */
+	virtual TSharedPtr<class IClassViewerFilter> GetImportedClassViewerFilter() const { return nullptr; }
+
+	/** Return the pin type selector filter associated with the current set of imported namespaces within this editor context. Default is NULL (no filter). */
+	virtual TSharedPtr<class IPinTypeSelectorFilter> GetImportedPinTypeSelectorFilter() const { return nullptr; }
+
 	UE_DEPRECATED(5.0, "GetSelectedSCSEditorTreeNodes has been deprecated. Use GetSelectedSubobjectEditorTreeNodes instead.")
 	virtual TArray<TSharedPtr<class FSCSEditorTreeNode> >  GetSelectedSCSEditorTreeNodes() const = 0;
 	UE_DEPRECATED(5.0, "FindAndSelectSCSEditorTreeNode has been deprecated. Use FindAndSelectSubobjectEditorTreeNode instead.")
diff --git a/Engine/Source/Editor/Kismet/Public/BlueprintNamespaceHelper.h b/Engine/Source/Editor/Kismet/Public/BlueprintNamespaceHelper.h
--- a/Engine/Source/Editor/Kismet/Public/BlueprintNamespaceHelper.h
+++ b/Engine/Source/Editor/Kismet/Public/BlueprintNamespaceHelper.h
@@ -4,6 +4,7 @@
 
 #include "CoreMinimal.h"
 
+class IClassViewerFilter;
 class IPinTypeSelectorFilter;
 
 // Utility class for extracting namespace info from a specified Blueprint context.
@@ -39,6 +40,11 @@ public:
 	bool IsImportedObject(const UObject* InObject) const;
 	bool IsImportedObject(const FSoftObjectPath& InObjectPath) const;
 
+	TSharedPtr<IClassViewerFilter> GetClassViewerFilter() const
+	{
+		return ClassViewerFilter;
+	}
+
 	TSharedPtr<IPinTypeSelectorFilter> GetPinTypeSelectorFilter() const
 	{
 		return PinTypeSelectorFilter;
@@ -48,6 +54,9 @@ private:
 	// Complete list of all fully-qualified namespace path identifiers for the associated Blueprint.
 	TSet<FString> FullyQualifiedListOfNamespaces;
 
+	// For use with the class viewer widget in order to filter class type items by namespace.
+	TSharedPtr<IClassViewerFilter> ClassViewerFilter;
+
 	// For use with the pin type selector widget in order to filter pin type items by namespace.
 	TSharedPtr<IPinTypeSelectorFilter> PinTypeSelectorFilter;
 };
\ No newline at end of file
diff --git a/Engine/Source/Editor/UnrealEd/Private/Kismet2/BlueprintEditorUtils.cpp b/Engine/Source/Editor/UnrealEd/Private/Kismet2/BlueprintEditorUtils.cpp
--- a/Engine/Source/Editor/UnrealEd/Private/Kismet2/BlueprintEditorUtils.cpp
+++ b/Engine/Source/Editor/UnrealEd/Private/Kismet2/BlueprintEditorUtils.cpp
@@ -8377,6 +8377,18 @@ TSharedRef<SWidget> FBlueprintEditorUtils::ConstructBlueprintParentClassPicker(
 	for ( UBlueprint* Blueprint : Blueprints )
 	{
 		Blueprint->GetReparentingRules(Filter->AllowedChildrenOfClasses, Filter->DisallowedChildrenOfClasses);
+
+		// Include a class viewer filter for imported namespaces if the class picker is being hosted in an editor context.
+		TSharedPtr<IToolkit> AssetEditor = FToolkitManager::Get().FindEditorForAsset(Blueprint);
+		if (AssetEditor.IsValid() && AssetEditor->IsBlueprintEditor())
+		{
+			TSharedPtr<IBlueprintEditor> BlueprintEditor = StaticCastSharedPtr<IBlueprintEditor>(AssetEditor);
+			TSharedPtr<IClassViewerFilter> ImportedClassViewerFilter = BlueprintEditor->GetImportedClassViewerFilter();
+			if (ImportedClassViewerFilter.IsValid())
+			{
+				Options.ClassFilters.AddUnique(ImportedClassViewerFilter.ToSharedRef());
+			}
+		}
 	}
 
 	if(bIsActor)
@@ -8559,6 +8571,18 @@ TSharedRef<SWidget> FBlueprintEditorUtils::ConstructBlueprintInterfaceClassPicke
 			const FBPInterfaceDescription& CurrentInterface = *it;
 			Filter->DisallowedClasses.Add(CurrentInterface.Interface);
 		}
+
+		// Include a class viewer filter for imported namespaces if the class picker is being hosted in an editor context
+		TSharedPtr<IToolkit> AssetEditor = FToolkitManager::Get().FindEditorForAsset(Blueprint);
+		if (AssetEditor.IsValid() && AssetEditor->IsBlueprintEditor())
+		{
+			TSharedPtr<IBlueprintEditor> BlueprintEditor = StaticCastSharedPtr<IBlueprintEditor>(AssetEditor);
+			TSharedPtr<IClassViewerFilter> ImportedClassViewerFilter = BlueprintEditor->GetImportedClassViewerFilter();
+			if (ImportedClassViewerFilter.IsValid())
+			{
+				Options.ClassFilters.AddUnique(ImportedClassViewerFilter.ToSharedRef());
+			}
+		}
 	}
 
 	// never allow parenting to children of itself
