Fix: Double Click an interface function in BP does nothing

diff --git a/Engine/Source/Editor/Kismet/Private/SMyBlueprint.cpp b/Engine/Source/Editor/Kismet/Private/SMyBlueprint.cpp
--- a/Engine/Source/Editor/Kismet/Private/SMyBlueprint.cpp
+++ b/Engine/Source/Editor/Kismet/Private/SMyBlueprint.cpp
@@ -2065,6 +2065,26 @@ void SMyBlueprint::ExecuteAction(TSharedPtr<FEdGraphSchemaAction> InAction)
 			{
 				BlueprintEditorPtr.Pin()->JumpToHyperlink(GraphAction->EdGraph);
 			}
+			else if(IsAnInterfaceEvent(GraphAction))
+			{
+				// Focus it's node in the event graph
+				UFunction* OverrideFunc = nullptr;
+				UClass* const OverrideFuncClass = FBlueprintEditorUtils::GetOverrideFunctionClass(BlueprintObj, GraphAction->FuncName, &OverrideFunc);
+				if (OverrideFunc)
+				{
+					FName EventName = OverrideFunc->GetFName();
+					// check if event has been implemented
+					if (UK2Node_Event* ExistingNode = FBlueprintEditorUtils::FindOverrideForFunction(BlueprintObj, OverrideFuncClass, EventName))
+					{
+						FKismetEditorUtilities::BringKismetToFocusAttentionOnObject(ExistingNode);
+					}
+					else
+					{
+						// if there isn't an associated node, make one and focus it
+						ImplementFunction(GraphAction);
+					}
+				}
+			}
 		}
 		if (InAction->GetTypeId() == FEdGraphSchemaAction_K2Delegate::StaticGetTypeId())
 		{
@@ -2835,6 +2855,11 @@ bool SMyBlueprint::IsEditingMode() const
 	return BlueprintEditorSPtr.IsValid() && BlueprintEditorSPtr->InEditingMode();
 }
 
+bool SMyBlueprint::IsAnInterfaceEvent(FEdGraphSchemaAction_K2Graph* InAction)
+{
+	return InAction->GraphType == EEdGraphSchemaAction_K2Graph::Interface && !InAction->EdGraph;
+}
+
 void SMyBlueprint::OnDeleteDelegate(FEdGraphSchemaAction_K2Delegate* InDelegateAction)
 {
 	UEdGraph* GraphToActOn = InDelegateAction->EdGraph;
diff --git a/Engine/Source/Editor/Kismet/Private/SMyBlueprint.h b/Engine/Source/Editor/Kismet/Private/SMyBlueprint.h
--- a/Engine/Source/Editor/Kismet/Private/SMyBlueprint.h
+++ b/Engine/Source/Editor/Kismet/Private/SMyBlueprint.h
@@ -300,6 +300,9 @@ private:
 	/** Helper function indicating whehter we're in editing mode, and can modify the target blueprint */
 	bool IsEditingMode() const;
 
+	/** Determine whether an FEdGraphSchemaAction is associated with an event */
+	static bool IsAnInterfaceEvent(FEdGraphSchemaAction_K2Graph* InAction);
+
 private:
 	/** List of UI Commands for this scope */
 	TSharedPtr<FUICommandList> CommandList;
