Fix adding variable in blueprint causes crash

diff --git a/Engine/Source/Editor/Kismet/Private/BlueprintDetailsCustomization.cpp b/Engine/Source/Editor/Kismet/Private/BlueprintDetailsCustomization.cpp
--- a/Engine/Source/Editor/Kismet/Private/BlueprintDetailsCustomization.cpp
+++ b/Engine/Source/Editor/Kismet/Private/BlueprintDetailsCustomization.cpp
@@ -298,7 +298,10 @@ void FBlueprintVarActionDetails::CustomizeDetails( IDetailLayoutBuilder& DetailL
 	const UEdGraphSchema* Schema = GetDefault<UEdGraphSchema_K2>();
 	if (BlueprintEditor.IsValid())
 	{
-		Schema = BlueprintEditor.Pin()->GetFocusedGraph()->GetSchema();
+		if (BlueprintEditor.Pin()->GetFocusedGraph())
+		{
+			Schema = BlueprintEditor.Pin()->GetFocusedGraph()->GetSchema();
+		}
 	}
 
 	Category.AddCustomRow(LOCTEXT("VariableTypeLabel", "Variable Type"))
