Add warning message when the category of a variable coincides with the name of the blueprint

diff --git a/Engine/Source/Editor/Kismet/Private/BlueprintDetailsCustomization.cpp b/Engine/Source/Editor/Kismet/Private/BlueprintDetailsCustomization.cpp
--- a/Engine/Source/Editor/Kismet/Private/BlueprintDetailsCustomization.cpp
+++ b/Engine/Source/Editor/Kismet/Private/BlueprintDetailsCustomization.cpp
@@ -539,6 +539,20 @@ void FBlueprintVarActionDetails::CustomizeDetails( IDetailLayoutBuilder& DetailL
 				SNew(SEditableTextBox)
 					.Text(this, &FBlueprintVarActionDetails::OnGetCategoryText)
 					.OnTextCommitted(this, &FBlueprintVarActionDetails::OnCategoryTextCommitted, CachedVariableName )
+					.OnVerifyTextChanged_Lambda([&](const FText& InNewText, FText& OutErrorMessage) -> bool
+					{
+						if (InNewText.IsEmpty())
+						{
+							OutErrorMessage = OutErrorMessage = LOCTEXT("CategoryEmpty", "Cannot add a category with an empty string.");
+							return false;
+						}
+						if (InNewText.EqualTo(FText::FromString(GetBlueprintObj()->GetName())))
+						{
+							OutErrorMessage = OutErrorMessage = LOCTEXT("CategoryEqualsBlueprintName", "Cannot add a category with the same name as the blueprint.");
+							return false;
+						}
+						return true;
+					})
 					.ToolTip(CategoryTooltip)
 					.SelectAllTextWhenFocused(true)
 					.RevertTextOnEscape(true)
@@ -5990,6 +6004,20 @@ void FBlueprintComponentDetails::CustomizeDetails(IDetailLayoutBuilder& DetailLa
 					SNew(SEditableTextBox)
 					.Text(this, &FBlueprintComponentDetails::OnGetVariableCategoryText)
 					.OnTextCommitted(this, &FBlueprintComponentDetails::OnVariableCategoryTextCommitted, CachedNodePtr->GetVariableName())
+					.OnVerifyTextChanged_Lambda([&](const FText& InNewText, FText& OutErrorMessage) -> bool
+					{
+						if (InNewText.IsEmpty())
+						{
+							OutErrorMessage = OutErrorMessage = LOCTEXT("CategoryEmpty", "Cannot add a category with an empty string.");
+							return false;
+						}
+						if (InNewText.EqualTo(FText::FromString(GetBlueprintObj()->GetName())))
+						{
+							OutErrorMessage = OutErrorMessage = LOCTEXT("CategoryEqualsBlueprintName", "Cannot add a category with the same name as the blueprint.");
+							return false;
+						}
+						return true;
+					})
 					.ToolTipText(CategoryTooltip)
 					.SelectAllTextWhenFocused(true)
 					.RevertTextOnEscape(true)
