Uppercase/lowercase categories with same name

diff --git a/Engine/Source/Editor/Kismet/Private/BlueprintDetailsCustomization.cpp b/Engine/Source/Editor/Kismet/Private/BlueprintDetailsCustomization.cpp
--- a/Engine/Source/Editor/Kismet/Private/BlueprintDetailsCustomization.cpp
+++ b/Engine/Source/Editor/Kismet/Private/BlueprintDetailsCustomization.cpp
@@ -66,6 +66,7 @@
 #include "Modules/ModuleManager.h"
 #include "ISequencerModule.h"
 #include "AnimatedPropertyKey.h"
+#include "EditorCategoryUtils.h"
 
 #include "PropertyCustomizationHelpers.h"
 
@@ -541,14 +542,15 @@ void FBlueprintVarActionDetails::CustomizeDetails( IDetailLayoutBuilder& DetailL
 					.OnTextCommitted(this, &FBlueprintVarActionDetails::OnCategoryTextCommitted, CachedVariableName )
 					.OnVerifyTextChanged_Lambda([&](const FText& InNewText, FText& OutErrorMessage) -> bool
 					{
-						if (InNewText.IsEmpty())
+						const FText NewText = FEditorCategoryUtils::GetCategoryDisplayString(InNewText);
+						if (NewText.IsEmpty())
 						{
-							OutErrorMessage = OutErrorMessage = LOCTEXT("CategoryEmpty", "Cannot add a category with an empty string.");
+							OutErrorMessage = LOCTEXT("CategoryEmpty", "Cannot add a category with an empty string.");
 							return false;
 						}
-						if (InNewText.EqualTo(FText::FromString(GetBlueprintObj()->GetName())))
+						if (NewText.EqualTo(FText::FromString(GetBlueprintObj()->GetName())))
 						{
-							OutErrorMessage = OutErrorMessage = LOCTEXT("CategoryEqualsBlueprintName", "Cannot add a category with the same name as the blueprint.");
+							OutErrorMessage = LOCTEXT("CategoryEqualsBlueprintName", "Cannot add a category with the same name as the blueprint.");
 							return false;
 						}
 						return true;
@@ -1672,8 +1674,8 @@ void FBlueprintVarActionDetails::OnCategoryTextCommitted(const FText& NewText, E
 {
 	if (InTextCommit == ETextCommit::OnEnter || InTextCommit == ETextCommit::OnUserMovedFocus)
 	{
-		// Remove excess whitespace and prevent categories with just spaces
-		FText CategoryName = FText::TrimPrecedingAndTrailing(NewText);
+		// Sanitize category name
+		FText CategoryName = FEditorCategoryUtils::GetCategoryDisplayString(NewText);
 
 		FBlueprintEditorUtils::SetBlueprintVariableCategory(GetBlueprintObj(), VarName, GetLocalVariableScope(CachedVariableProperty.Get()), CategoryName);
 		check(MyBlueprint.IsValid());
@@ -6006,14 +6008,15 @@ void FBlueprintComponentDetails::CustomizeDetails(IDetailLayoutBuilder& DetailLa
 					.OnTextCommitted(this, &FBlueprintComponentDetails::OnVariableCategoryTextCommitted, CachedNodePtr->GetVariableName())
 					.OnVerifyTextChanged_Lambda([&](const FText& InNewText, FText& OutErrorMessage) -> bool
 					{
-						if (InNewText.IsEmpty())
+						const FText NewText = FEditorCategoryUtils::GetCategoryDisplayString(InNewText);
+						if (NewText.IsEmpty())
 						{
-							OutErrorMessage = OutErrorMessage = LOCTEXT("CategoryEmpty", "Cannot add a category with an empty string.");
+							OutErrorMessage = LOCTEXT("CategoryEmpty", "Cannot add a category with an empty string.");
 							return false;
 						}
-						if (InNewText.EqualTo(FText::FromString(GetBlueprintObj()->GetName())))
+						if (NewText.EqualTo(FText::FromString(GetBlueprintObj()->GetName())))
 						{
-							OutErrorMessage = OutErrorMessage = LOCTEXT("CategoryEqualsBlueprintName", "Cannot add a category with the same name as the blueprint.");
+							OutErrorMessage = LOCTEXT("CategoryEqualsBlueprintName", "Cannot add a category with the same name as the blueprint.");
 							return false;
 						}
 						return true;
