Fix: Able to override blueprint declared events multiple times
made it so that the function would no longer appear in the list after being overriden. This was the default behavior already but didn't work if the blueprint's parent was another blueprint

diff --git a/Engine/Source/Editor/Kismet/Private/SMyBlueprint.cpp b/Engine/Source/Editor/Kismet/Private/SMyBlueprint.cpp
--- a/Engine/Source/Editor/Kismet/Private/SMyBlueprint.cpp
+++ b/Engine/Source/Editor/Kismet/Private/SMyBlueprint.cpp
@@ -1469,11 +1469,19 @@ void SMyBlueprint::CollectAllActions(FGraphActionListBuilderBase& OutAllActions)
 		const UFunction* Function = *FunctionIt;
 		const FName FunctionName = Function->GetFName();
 
+		UClass *OuterClass = CastChecked<UClass>(Function->GetOuter());
+		// ignore skeleton classes and convert them into their "authoritative" types so they
+		// can be found in the graph
+		if(UBlueprintGeneratedClass *GeneratedOuterClass = Cast<UBlueprintGeneratedClass>(OuterClass))
+		{
+			OuterClass = GeneratedOuterClass->GetAuthoritativeClass();
+		}
+
 		if (    UEdGraphSchema_K2::CanKismetOverrideFunction(Function) 
 			 && !OverridableFunctionNames.Contains(FunctionName) 
 			 && !ImplementedFunctionCache.Contains(FunctionName) 
 			 && !FObjectEditorUtils::IsFunctionHiddenFromClass(Function, ParentClass)
-			 && !FBlueprintEditorUtils::FindOverrideForFunction(BlueprintObj, CastChecked<UClass>(Function->GetOuter()), Function->GetFName()) )
+			 && !FBlueprintEditorUtils::FindOverrideForFunction(BlueprintObj, OuterClass, Function->GetFName()) )
 		{
 			FText FunctionTooltip = FText::FromString(UK2Node_CallFunction::GetDefaultTooltipForFunction(Function));
 			FText FunctionDesc = K2Schema->GetFriendlySignatureName(Function);
