Fix a crash when navigating with the arrow keys in the blueprint editor

diff --git a/Engine/Source/Editor/Kismet/Private/SCSEditorViewportClient.cpp b/Engine/Source/Editor/Kismet/Private/SCSEditorViewportClient.cpp
--- a/Engine/Source/Editor/Kismet/Private/SCSEditorViewportClient.cpp
+++ b/Engine/Source/Editor/Kismet/Private/SCSEditorViewportClient.cpp
@@ -667,7 +667,7 @@ FVector FSCSEditorViewportClient::GetWidgetLocation() const
 			// Use the last selected item for the widget location
 			const FSubobjectData* Data = SelectedNodes[0]->GetDataSource();
 
-			const USceneComponent* SceneComp = Cast<USceneComponent>(Data->FindComponentInstanceInActor(PreviewActor));
+			const USceneComponent* SceneComp = Data ? Cast<USceneComponent>(Data->FindComponentInstanceInActor(PreviewActor)) : nullptr;
 			if(SceneComp)
 			{
 				TSharedPtr<ISCSEditorCustomization> Customization = BlueprintEditorPtr.Pin()->CustomizeSubobjectEditor(SceneComp);
@@ -708,7 +708,7 @@ FMatrix FSCSEditorViewportClient::GetWidgetCoordSystem() const
 			{
 				const FSubobjectEditorTreeNodePtrType SelectedNode = SelectedNodes.Last();
 				const FSubobjectData* Data = SelectedNode->GetDataSource();
-				const USceneComponent* SceneComp = SelectedNode.IsValid() ? Cast<USceneComponent>(Data->FindComponentInstanceInActor(PreviewActor)) : nullptr;
+				const USceneComponent* SceneComp = Data ? Cast<USceneComponent>(Data->FindComponentInstanceInActor(PreviewActor)) : nullptr;
 				if(SceneComp)
 				{
 					TSharedPtr<ISCSEditorCustomization> Customization = BlueprintEditor->CustomizeSubobjectEditor(SceneComp);
