Fix for blueprint class default values not initially propagating to child BP classes when created via "Promote to Variable" operation.  We now detect this at the 1st normal BP compilation and initialize child BP values then.

diff --git a/Engine/Source/Editor/Kismet/Private/BlueprintCompilationManager.cpp b/Engine/Source/Editor/Kismet/Private/BlueprintCompilationManager.cpp
--- a/Engine/Source/Editor/Kismet/Private/BlueprintCompilationManager.cpp
+++ b/Engine/Source/Editor/Kismet/Private/BlueprintCompilationManager.cpp
@@ -461,6 +461,8 @@ struct FCompilerData
 	FKismetCompilerOptions InternalOptions;
 	TSharedPtr<FBlueprintCompileReinstancer> Reinstancer;
 	TArray<FSkeletonFixupData> SkeletonFixupData;
+	/** variables that are new to the generated class and will need their default set (can occur when a new variable is added to a BP's ancestor class) */
+	TArray<FBPVariableDescription> NewDefaultVariables;	
 
 	ECompilationManagerJobType JobType;
 	bool bPackageWasDirty;
@@ -1085,6 +1087,34 @@ void FBlueprintCompilationManagerImpl::FlushCompilationQueueImpl(bool bSuppressB
 				CurrentlyCompilingBPs.RemoveAll(HasNoReferencesToChangedFunctions);
 			}
 		}
+		
+		// Detect any variable-based properties that are not in the old generated class, save them for after reinstancing. This can occur 
+		//    when a new variable is introduced in an ancestor class, and we'll need to use its default as our generated class's initial value.
+		{
+			DECLARE_SCOPE_HIERARCHICAL_COUNTER(DetectNewDefaultVariables)
+
+			for (FCompilerData& CompilerData : CurrentlyCompilingBPs)
+			{
+				if (CompilerData.JobType == ECompilationManagerJobType::Normal &&
+					!CompilerData.BP->bIsRegeneratingOnLoad)
+				{
+					const UClass* ParentClass = CompilerData.BP->ParentClass;
+					while (const UBlueprint* ParentBP = UBlueprint::GetBlueprintFromClass(ParentClass))
+					{
+						for (const FBPVariableDescription& ParentBPVarDesc : ParentBP->NewVariables)
+						{
+							if (!CompilerData.BP->GeneratedClass->FindPropertyByName(ParentBPVarDesc.VarName))
+							{
+								CompilerData.NewDefaultVariables.Add(ParentBPVarDesc);
+							}
+						}
+
+						ParentClass = ParentBP->ParentClass;
+					}
+				}
+			}
+		}
+
 
 		// STAGE IX: Reconstruct nodes and replace deprecated nodes, then broadcast 'precompile
 		for (FCompilerData& CompilerData : CurrentlyCompilingBPs)
@@ -1423,6 +1453,32 @@ void FBlueprintCompilationManagerImpl::FlushCompilationQueueImpl(bool bSuppressB
 
 			// We purposefully do not remove the OldCDOs yet, need to keep them in memory past first GC
 		}
+
+		// Set default values on any newly-introduced variables (from ancestor BPs)
+		{
+			DECLARE_SCOPE_HIERARCHICAL_COUNTER(SetNewDefaultVariables)
+
+			for (FCompilerData& CompilerData : CurrentlyCompilingBPs)
+			{
+				if (!CompilerData.NewDefaultVariables.Num())
+				{
+					continue;
+				}
+
+				const UBlueprintGeneratedClass* GenClass = Cast<UBlueprintGeneratedClass>(CompilerData.BP->GeneratedClass);
+
+				if (GenClass && GenClass->ClassDefaultObject)
+				{
+					for (const FBPVariableDescription& NewInheritedVar : CompilerData.NewDefaultVariables)
+					{
+						if (const FProperty* MatchingProperty = GenClass->FindPropertyByName(NewInheritedVar.VarName))
+						{
+							FBlueprintEditorUtils::PropertyValueFromString(MatchingProperty, NewInheritedVar.DefaultValue, reinterpret_cast<uint8*>(GenClass->ClassDefaultObject));
+						}
+					}
+				}
+			}
+		}
 		
 		// STAGE XV: POST CDO COMPILED
 		for (FCompilerData& CompilerData : CurrentlyCompilingBPs)
