Tooltip field for event dispatchers no longer works and should be renamed to 'Description' to match custom functions

diff --git a/Engine/Source/Editor/Kismet/Private/BlueprintDetailsCustomization.cpp b/Engine/Source/Editor/Kismet/Private/BlueprintDetailsCustomization.cpp
--- a/Engine/Source/Editor/Kismet/Private/BlueprintDetailsCustomization.cpp
+++ b/Engine/Source/Editor/Kismet/Private/BlueprintDetailsCustomization.cpp
@@ -235,7 +235,18 @@ FBlueprintVarActionDetails::~FBlueprintVarActionDetails()
 BEGIN_SLATE_FUNCTION_BUILD_OPTIMIZATION
 void FBlueprintVarActionDetails::CustomizeDetails( IDetailLayoutBuilder& DetailLayout )
 {
+	// update the property being edited
+	PropertyBeingEdited = nullptr;
+	TArray<TWeakObjectPtr<UObject>> ObjectsBeingEdited;
 	DetailLayout.GetObjectsBeingCustomized(ObjectsBeingEdited);
+	for(TWeakObjectPtr<UObject> Object : ObjectsBeingEdited)
+	{
+		if(UPropertyWrapper* Wrapper = Cast<UPropertyWrapper>(Object.Get()))
+		{
+			PropertyBeingEdited = Wrapper->GetProperty();
+			break;
+		}
+	}
 
 	CachedVariableProperty = SelectionAsProperty();
 
@@ -327,6 +338,30 @@ void FBlueprintVarActionDetails::CustomizeDetails( IDetailLayoutBuilder& DetailL
 			.CustomFilter(CustomPinTypeFilter)
 		];
 
+	TSharedPtr<SToolTip> ToolTipTooltip = IDocumentation::Get()->CreateToolTip(LOCTEXT("VarToolTipTooltip", "Extra information about this variable, shown when cursor is over it."), NULL, DocLink, TEXT("Description"));
+
+	Category.AddCustomRow( LOCTEXT("IsVariableToolTipLabel", "Description") )
+	.Visibility(TAttribute<EVisibility>(this, &FBlueprintVarActionDetails::IsTooltipEditVisible))
+	.NameContent()
+	[
+		SNew(STextBlock)
+		.Text( LOCTEXT("IsVariableToolTipLabel", "Description") )
+		.ToolTip(ToolTipTooltip)
+		.Font( DetailFontInfo )
+	]
+	.ValueContent()
+	.MinDesiredWidth(250.f)
+	.MaxDesiredWidth(250.f)
+	[
+		SNew(SMultiLineEditableTextBox)
+		.Text( this, &FBlueprintVarActionDetails::OnGetTooltipText )
+		.ToolTipText( this, &FBlueprintVarActionDetails::OnGetTooltipText )
+		.OnTextCommitted( this, &FBlueprintVarActionDetails::OnTooltipTextCommitted, CachedVariableName )
+		.IsEnabled(IsVariableInBlueprint())
+		.Font( DetailFontInfo )
+		.ModiferKeyForNewLine(EModifierKey::Shift)
+	];
+
 	TSharedPtr<SToolTip> EditableTooltip = IDocumentation::Get()->CreateToolTip(LOCTEXT("VarEditableTooltip", "Whether this variable is publicly editable on instances of this Blueprint."), NULL, DocLink, TEXT("Editable"));
 
 	Category.AddCustomRow( LOCTEXT("IsVariableEditableLabel", "Instance Editable") )
@@ -367,30 +402,6 @@ void FBlueprintVarActionDetails::CustomizeDetails( IDetailLayoutBuilder& DetailL
 		.ToolTip(ReadOnlyTooltip)
 	];
 
-	TSharedPtr<SToolTip> ToolTipTooltip = IDocumentation::Get()->CreateToolTip(LOCTEXT("VarToolTipTooltip", "Extra information about this variable, shown when cursor is over it."), NULL, DocLink, TEXT("Tooltip"));
-
-	Category.AddCustomRow( LOCTEXT("IsVariableToolTipLabel", "Tooltip") )
-	.Visibility(TAttribute<EVisibility>(this, &FBlueprintVarActionDetails::IsTooltipEditVisible))
-	.NameContent()
-	[
-		SNew(STextBlock)
-		.Text( LOCTEXT("IsVariableToolTipLabel", "Tooltip") )
-		.ToolTip(ToolTipTooltip)
-		.Font( DetailFontInfo )
-	]
-	.ValueContent()
-	.MinDesiredWidth(250.f)
-	.MaxDesiredWidth(250.f)
-	[
-		SNew(SMultiLineEditableTextBox)
-		.Text( this, &FBlueprintVarActionDetails::OnGetTooltipText )
-		.ToolTipText( this, &FBlueprintVarActionDetails::OnGetTooltipText )
-		.OnTextCommitted( this, &FBlueprintVarActionDetails::OnTooltipTextCommitted, CachedVariableName )
-		.IsEnabled(IsVariableInBlueprint())
-		.Font( DetailFontInfo )
-		.ModiferKeyForNewLine(EModifierKey::Shift)
-	];
-
 	TSharedPtr<SToolTip> Widget3DTooltip = IDocumentation::Get()->CreateToolTip(LOCTEXT("VariableWidget3D_Tooltip", "When true, allows the user to tweak the vector variable by using a 3D transform widget in the viewport (usable when varible is public/enabled)."), NULL, DocLink, TEXT("Widget3D"));
 
 	Category.AddCustomRow( LOCTEXT("VariableWidget3D_Prompt", "Show 3D Widget") )
@@ -1558,20 +1569,6 @@ void FBlueprintVarActionDetails::PopulateCategories(SMyBlueprint* MyBlueprint, T
 	}
 }
 
-FProperty* FBlueprintVarActionDetails::CustomizedObjectAsProperty() const
-{
-	if(ObjectsBeingEdited.Num() == 1)
-	{
-		UPropertyWrapper* Wrapper = Cast<UPropertyWrapper>(ObjectsBeingEdited[0].Get());
-		if (Wrapper)
-		{
-			return Wrapper->GetProperty();
-		}
-	}
-
-	return nullptr;
-}
-
 UK2Node_Variable* FBlueprintVarActionDetails::EdGraphSelectionAsVar() const
 {
 	TWeakPtr<FBlueprintEditor> BlueprintEditor = MyBlueprint.Pin()->GetBlueprintEditor();
@@ -1612,12 +1609,7 @@ FProperty* FBlueprintVarActionDetails::SelectionAsProperty() const
 	{
 		return GraphVar->GetPropertyForVariable();
 	}
-	FProperty* Property = CustomizedObjectAsProperty();
-	if(Property)
-	{
-		return Property;
-	}
-	return NULL;
+	return PropertyBeingEdited;
 }
 
 FName FBlueprintVarActionDetails::GetVariableName() const
@@ -1637,10 +1629,9 @@ FName FBlueprintVarActionDetails::GetVariableName() const
 	{
 		return GraphVar->GetVarName();
 	}
-	FProperty* Property = CustomizedObjectAsProperty();
-	if(Property)
+	if(PropertyBeingEdited)
 	{
-		return Property->GetFName();
+		return PropertyBeingEdited->GetFName();
 	}
 	return NAME_None;
 }
diff --git a/Engine/Source/Editor/Kismet/Private/BlueprintDetailsCustomization.h b/Engine/Source/Editor/Kismet/Private/BlueprintDetailsCustomization.h
--- a/Engine/Source/Editor/Kismet/Private/BlueprintDetailsCustomization.h
+++ b/Engine/Source/Editor/Kismet/Private/BlueprintDetailsCustomization.h
@@ -124,7 +124,6 @@ private:
 	FEdGraphSchemaAction_K2Var* MyBlueprintSelectionAsVar() const {return MyBlueprint.Pin()->SelectionAsVar();}
 	FEdGraphSchemaAction_K2LocalVar* MyBlueprintSelectionAsLocalVar() const {return MyBlueprint.Pin()->SelectionAsLocalVar();}
 	UK2Node_Variable* EdGraphSelectionAsVar() const;
-	FProperty* CustomizedObjectAsProperty() const;
 	FProperty* SelectionAsProperty() const;
 	FName GetVariableName() const;
 
@@ -305,8 +304,8 @@ private:
 	/** External detail customizations */
 	TArray<TSharedPtr<IDetailCustomization>> ExternalDetailCustomizations;
 
-	/** Array of nodes were were constructed to represent */
-	TArray< TWeakObjectPtr<UObject> > ObjectsBeingEdited;
+	/** If a property is being edited right now, this will cache it */
+	FProperty* PropertyBeingEdited;
 };
 
 class FBaseBlueprintGraphActionDetails : public IDetailCustomization
