Fixed crash when creating a child of a blueprint that hasn't been fully compiled yet via the Content Browser (introduced in CL 16759714)

#rnx

diff --git a/Engine/Source/Editor/Kismet/Private/BlueprintCompilationManager.cpp b/Engine/Source/Editor/Kismet/Private/BlueprintCompilationManager.cpp
--- a/Engine/Source/Editor/Kismet/Private/BlueprintCompilationManager.cpp
+++ b/Engine/Source/Editor/Kismet/Private/BlueprintCompilationManager.cpp
@@ -1096,7 +1096,8 @@ void FBlueprintCompilationManagerImpl::FlushCompilationQueueImpl(bool bSuppressB
 			for (FCompilerData& CompilerData : CurrentlyCompilingBPs)
 			{
 				if (CompilerData.JobType == ECompilationManagerJobType::Normal &&
-					!CompilerData.BP->bIsRegeneratingOnLoad)
+					!CompilerData.BP->bIsRegeneratingOnLoad &&
+					 CompilerData.BP->GeneratedClass)
 				{
 					const UClass* ParentClass = CompilerData.BP->ParentClass;
 					while (const UBlueprint* ParentBP = UBlueprint::GetBlueprintFromClass(ParentClass))
