Fix: The 'const' checkbox should not be visible in the Details tab for Blueprint Library functions.

diff --git a/Engine/Source/Editor/BlueprintGraph/Classes/K2Node_FunctionEntry.h b/Engine/Source/Editor/BlueprintGraph/Classes/K2Node_FunctionEntry.h
--- a/Engine/Source/Editor/BlueprintGraph/Classes/K2Node_FunctionEntry.h
+++ b/Engine/Source/Editor/BlueprintGraph/Classes/K2Node_FunctionEntry.h
@@ -117,6 +117,18 @@ class UK2Node_FunctionEntry : public UK2Node_FunctionTerminator
 	{
 		ExtraFlags &= ~InFlags;
 	}
+	
+	/** Used to safely check whether the passed in flag is set. */
+	BLUEPRINTGRAPH_API bool HasAnyExtraFlags(int32 FlagsToCheck) const
+	{
+		return (ExtraFlags & FlagsToCheck) != 0 || FlagsToCheck == ~0;
+	}
+
+	/** Used to safely check whether all of the passed in flags are set. */
+	BLUEPRINTGRAPH_API bool HasAllExtraFlags(int32 FlagsToCheck) const
+	{
+		return ((ExtraFlags & FlagsToCheck) == FlagsToCheck);
+	}
 
 protected:
 	/** Copies data from any local variables matching properties in VariableStruct into the VariableStructData */
diff --git a/Engine/Source/Editor/BlueprintGraph/Private/K2Node_FunctionEntry.cpp b/Engine/Source/Editor/BlueprintGraph/Private/K2Node_FunctionEntry.cpp
--- a/Engine/Source/Editor/BlueprintGraph/Private/K2Node_FunctionEntry.cpp
+++ b/Engine/Source/Editor/BlueprintGraph/Private/K2Node_FunctionEntry.cpp
@@ -233,6 +233,13 @@ void UK2Node_FunctionEntry::PostLoad()
 
 		UpdateLoadedDefaultValues();
 	}
+
+	// fix deprecated state.
+	if (HasAllExtraFlags(FUNC_Const | FUNC_Static))
+	{
+		// static functions can't be marked const
+		SetExtraFlags(GetFunctionFlags() & ~FUNC_Const);
+	}
 }
 
 void UK2Node_FunctionEntry::Serialize(FArchive& Ar)
diff --git a/Engine/Source/Editor/Kismet/Private/BlueprintDetailsCustomization.cpp b/Engine/Source/Editor/Kismet/Private/BlueprintDetailsCustomization.cpp
--- a/Engine/Source/Editor/Kismet/Private/BlueprintDetailsCustomization.cpp
+++ b/Engine/Source/Editor/Kismet/Private/BlueprintDetailsCustomization.cpp
@@ -5046,15 +5046,16 @@ ECheckBoxState FBlueprintGraphActionDetails::GetIsPureFunction() const
 
 bool FBlueprintGraphActionDetails::IsConstFunctionVisible() const
 {
-	bool bSupportedType = false;
-	bool bIsEditable = false;
-	UK2Node_EditablePinBase* FunctionEntryNode = FunctionEntryNodePtr.Get();
-	if(FunctionEntryNode)
+	bool bVisible = false;
+	UFunction* Function = FindFunction();
+	UK2Node_FunctionEntry* EntryNode = Cast<UK2Node_FunctionEntry>(FunctionEntryNodePtr.Get());
+	if(Function && EntryNode)
 	{
-		bSupportedType = FunctionEntryNode->IsA<UK2Node_FunctionEntry>();
-		bIsEditable = FunctionEntryNode->IsEditable();
+		const bool bIsStatic = EntryNode->GetFunctionFlags() & FUNC_Static;
+		const bool bIsEditable = EntryNode->IsEditable();
+		bVisible = bIsEditable && !bIsStatic;
 	}
-	return bSupportedType && bIsEditable;
+	return bVisible;
 }
 
 void FBlueprintGraphActionDetails::OnIsConstFunctionModified( const ECheckBoxState NewCheckedState )
diff --git a/Engine/Source/Runtime/CoreUObject/Private/UObject/Class.cpp b/Engine/Source/Runtime/CoreUObject/Private/UObject/Class.cpp
--- a/Engine/Source/Runtime/CoreUObject/Private/UObject/Class.cpp
+++ b/Engine/Source/Runtime/CoreUObject/Private/UObject/Class.cpp
@@ -5829,6 +5829,13 @@ void UFunction::PostLoad()
 	{
 		OwningClass->ClassFlags &= ~CLASS_ReplicationDataIsSetUp;
 	}
+	
+	// fix deprecated state. 
+	if (HasAllFunctionFlags(FUNC_Const | FUNC_Static))
+	{
+		// static functions can't be marked const
+		FunctionFlags &= ~FUNC_Const;
+	}
 }
 
 FProperty* UFunction::GetReturnProperty() const
