Replace the last use of SSCSEditor in the SCSDiff panel with the Subobject Editor

#rnx

diff --git a/Engine/Source/Editor/Kismet/Private/SCSDiff.cpp b/Engine/Source/Editor/Kismet/Private/SCSDiff.cpp
--- a/Engine/Source/Editor/Kismet/Private/SCSDiff.cpp
+++ b/Engine/Source/Editor/Kismet/Private/SCSDiff.cpp
@@ -5,11 +5,11 @@
 #include "GameFramework/Actor.h"
 #include "SKismetInspector.h"
 #include "SSCSEditor.h"
+#include "SSubobjectEditor.h"
+#include "SSubobjectBlueprintEditor.h"
 #include "IDetailsView.h"
 #include "Kismet2/BlueprintEditorUtils.h"
 
-#include <vector>
-
 FSCSDiff::FSCSDiff(const UBlueprint* InBlueprint)
 {
 	if (!FBlueprintEditorUtils::SupportsConstructionScript(InBlueprint) || InBlueprint->SimpleConstructionScript == NULL)
@@ -30,13 +30,12 @@ FSCSDiff::FSCSDiff(const UBlueprint* InBlueprint)
 		.Orientation(Orient_Vertical)
 		+ SSplitter::Slot()
 		[
-			SAssignNew(SCSEditor, SSCSEditor)
-				.ActorContext(InBlueprint->GeneratedClass->GetDefaultObject<AActor>())
+			SAssignNew(SubobjectEditor, SSubobjectBlueprintEditor)
+				.ObjectContext(InBlueprint->GeneratedClass->GetDefaultObject<AActor>())
 				.AllowEditing(false)
 				.HideComponentClassCombo(true)
-				.OnSelectionUpdated(SSCSEditor::FOnSelectionUpdated::CreateRaw(this, &FSCSDiff::OnSCSEditorUpdateSelectionFromNodes))
-				.OnHighlightPropertyInDetailsView(SSCSEditor::FOnHighlightPropertyInDetailsView::CreateRaw(this, &FSCSDiff::OnSCSEditorHighlightPropertyInDetailsView))
-				.IsDiffing(true)
+				.OnSelectionUpdated(SSubobjectEditor::FOnSelectionUpdated::CreateRaw(this, &FSCSDiff::OnSCSEditorUpdateSelectionFromNodes))
+				.OnHighlightPropertyInDetailsView(SSubobjectBlueprintEditor::FOnHighlightPropertyInDetailsView::CreateRaw(this, &FSCSDiff::OnSCSEditorHighlightPropertyInDetailsView))
 		]
 		+ SSplitter::Slot()
 		[
@@ -46,10 +45,10 @@ FSCSDiff::FSCSDiff(const UBlueprint* InBlueprint)
 
 void FSCSDiff::HighlightProperty(FName VarName, FPropertySoftPath Property)
 {
-	if (SCSEditor.IsValid())
+	if (SubobjectEditor.IsValid())
 	{
 		check(VarName != FName());
-		SCSEditor->HighlightTreeNode(VarName, FPropertyPath());
+		SubobjectEditor->HighlightTreeNode(VarName, FPropertyPath());
 	}
 }
 
@@ -58,39 +57,46 @@ TSharedRef< SWidget > FSCSDiff::TreeWidget()
 	return ContainerWidget.ToSharedRef();
 }
 
-void GetDisplayedHierarchyRecursive(UBlueprint* Blueprint, TArray< int32 >& TreeAddress, const FSCSEditorTreeNode& Node, TArray< FSCSResolvedIdentifier >& OutResult)
+void GetDisplayedHierarchyRecursive(UBlueprint* Blueprint, TArray<int32>& TreeAddress, const FSubobjectDataHandle& Node, TArray<FSCSResolvedIdentifier>& OutResult)
 {
-	FSCSIdentifier Identifier = { Node.GetVariableName(), TreeAddress };
-	FSCSResolvedIdentifier ResolvedIdentifier = { Identifier, Node.GetOrCreateEditableComponentTemplate(Blueprint) };
+	FSubobjectData* NodeData = Node.GetData();
+	if (!NodeData)
+	{
+		return;
+	}
+
+	FSCSIdentifier Identifier = { NodeData->GetVariableName(), TreeAddress };
+	FSCSResolvedIdentifier ResolvedIdentifier = { Identifier, NodeData->GetObjectForBlueprint<UActorComponent>(Blueprint) };
 	OutResult.Push(ResolvedIdentifier);
-	const auto& Children = Node.GetChildren();
+	
+	const TArray<FSubobjectDataHandle>& Children = NodeData->GetChildrenHandles();
 	for (int32 Iter = 0; Iter != Children.Num(); ++Iter)
 	{
 		TreeAddress.Push(Iter);
-		GetDisplayedHierarchyRecursive(Blueprint, TreeAddress, *Children[Iter], OutResult);
+		GetDisplayedHierarchyRecursive(Blueprint, TreeAddress, Children[Iter], OutResult);
 		TreeAddress.Pop();
 	}
 }
 
-TArray< FSCSResolvedIdentifier > FSCSDiff::GetDisplayedHierarchy() const
+TArray<FSCSResolvedIdentifier> FSCSDiff::GetDisplayedHierarchy() const
 {
 	TArray< FSCSResolvedIdentifier > Ret;
 
-	if( SCSEditor.IsValid() && SCSEditor->GetActorNode().IsValid())
+	if(SubobjectEditor.IsValid() && SubobjectEditor->GetRootNodes().Num() > 0)
 	{
-		const TArray<FSCSEditorTreeNodePtrType>& RootNodes = SCSEditor->GetActorNode()->GetComponentNodes();
+		const TArray<FSubobjectEditorTreeNodePtrType>& RootNodes = SubobjectEditor->GetRootNodes();
 		for (int32 Iter = 0; Iter != RootNodes.Num(); ++Iter)
 		{
 			TArray< int32 > TreeAddress;
 			TreeAddress.Push(Iter);
-			GetDisplayedHierarchyRecursive(Blueprint, TreeAddress, *RootNodes[Iter], Ret);
+			GetDisplayedHierarchyRecursive(Blueprint, TreeAddress, RootNodes[Iter]->GetDataHandle(), Ret);
 		}
 	}
 
 	return Ret;
 }
 
-void FSCSDiff::OnSCSEditorUpdateSelectionFromNodes(const TArray<FSCSEditorTreeNodePtrType>& SelectedNodes)
+void FSCSDiff::OnSCSEditorUpdateSelectionFromNodes(const TArray<FSubobjectEditorTreeNodePtrType>& SelectedNodes)
 {
 	FText InspectorTitle = FText::GetEmpty();
 	TArray<UObject*> InspectorObjects;
@@ -98,10 +104,10 @@ void FSCSDiff::OnSCSEditorUpdateSelectionFromNodes(const TArray<FSCSEditorTreeNo
 	for (auto NodeIt = SelectedNodes.CreateConstIterator(); NodeIt; ++NodeIt)
 	{
 		auto NodePtr = *NodeIt;
-		if(NodePtr.IsValid() && NodePtr->CanEdit())
+		if(NodePtr.IsValid() && NodePtr->GetDataSource()->CanEdit())
 		{
 			InspectorTitle = FText::FromString(NodePtr->GetDisplayString());
-			InspectorObjects.Add(NodePtr->GetEditableObjectForBlueprint<UObject>(Blueprint));
+			InspectorObjects.Add(const_cast<UObject*>(NodePtr->GetDataSource()->GetObjectForBlueprint(Blueprint)));
 		}
 	}
 
diff --git a/Engine/Source/Editor/Kismet/Public/SCSDiff.h b/Engine/Source/Editor/Kismet/Public/SCSDiff.h
--- a/Engine/Source/Editor/Kismet/Public/SCSDiff.h
+++ b/Engine/Source/Editor/Kismet/Public/SCSDiff.h
@@ -6,8 +6,10 @@
 #include "DiffUtils.h"
 
 class FSCSEditorTreeNode;
+class FSubobjectEditorTreeNode;
 class SKismetInspector;
 class SSCSEditor;
+class SSubobjectBlueprintEditor;
 
 /** Struct to support diffing the component tree for a blueprint */
 class KISMET_API FSCSDiff
@@ -23,12 +25,12 @@ public:
 	const UBlueprint* GetBlueprint() const { return Blueprint; }
 
 protected:
-	void OnSCSEditorUpdateSelectionFromNodes(const TArray< TSharedPtr<class FSCSEditorTreeNode> >& SelectedNodes);
+	void OnSCSEditorUpdateSelectionFromNodes(const TArray<TSharedPtr<FSubobjectEditorTreeNode>>& SelectedNodes);
 	void OnSCSEditorHighlightPropertyInDetailsView(const class FPropertyPath& InPropertyPath);
 
 private:
 	TSharedPtr< class SWidget > ContainerWidget;
-	TSharedPtr< class SSCSEditor > SCSEditor;
+	TSharedPtr< class SSubobjectBlueprintEditor > SubobjectEditor;
 	TSharedPtr< class SKismetInspector > Inspector;
 
 	/** Blueprint we are inspecting */
