Fix to unblock async Blueprint asset indexing when a global search is hosted inside a non-global search tab.

diff --git a/Engine/Source/Editor/Kismet/Private/FindInBlueprintManager.cpp b/Engine/Source/Editor/Kismet/Private/FindInBlueprintManager.cpp
--- a/Engine/Source/Editor/Kismet/Private/FindInBlueprintManager.cpp
+++ b/Engine/Source/Editor/Kismet/Private/FindInBlueprintManager.cpp
@@ -1589,11 +1589,11 @@ public:
 		{
 			// Don't utilize the task graph if any of the following conditions hold TRUE:
 			// a) The application has throttled the tick rate.
-			// b) No global search tabs are currently open and visible.
+			// b) No global search tabs are currently open and visible, or there are no async search queries that are otherwise pending completion.
 			// c) The initial asset discovery phase has not yet been completed.
 			// d) Multiprocessing has been explicitly disabled for this operation.
 			return FSlateThrottleManager::Get().IsAllowingExpensiveTasks()
-				&& FFindInBlueprintSearchManager::Get().IsGlobalFindResultsOpen()
+				&& (FFindInBlueprintSearchManager::Get().IsGlobalFindResultsOpen() || FFindInBlueprintSearchManager::Get().IsAsyncSearchQueryInProgress())
 				&& !FFindInBlueprintSearchManager::Get().IsAssetDiscoveryInProgress()
 				&& !EnumHasAnyFlags(CacheParams.OpFlags, EFiBCacheOpFlags::ExecuteOnSingleThread);
 		}
@@ -3328,6 +3328,12 @@ bool FFindInBlueprintSearchManager::IsAssetDiscoveryInProgress() const
 	return GIsRunning && AssetRegistryModule && AssetRegistryModule->Get().IsLoadingAssets();
 }
 
+bool FFindInBlueprintSearchManager::IsAsyncSearchQueryInProgress() const
+{
+	// Note: Not using ActiveSearchCounter here, as that's used to block on pause and thus can be decremented during an active search.
+	return ActiveSearchQueries.Num() > 0;
+}
+
 TSharedPtr< FJsonObject > FFindInBlueprintSearchManager::ConvertJsonStringToObject(FSearchDataVersionInfo InVersionInfo, FString InJsonString, TMap<int32, FText>& OutFTextLookupTable)
 {
 	/** The searchable data is more complicated than a Json string, the Json being the main searchable body that is parsed. Below is a diagram of the full data:
@@ -3621,8 +3627,8 @@ bool FFindInBlueprintSearchManager::IsTickable() const
 	const bool bHasPendingAssets = PendingAssets.Num() > 0;
 	const bool bNeedsFirstIndex = bHasFirstSearchOccurred && AssetsToIndexOnFirstSearch.Num() > 0;
 
-	// Tick only if we have an active caching operation or if a search has occured before we're ready or if we have pending assets and an open FiB context
-	return IsCacheInProgress() || bNeedsFirstIndex || (bHasPendingAssets && IsGlobalFindResultsOpen());
+	// Tick only if we have an active caching operation or if a search has occured before we're ready or if we have pending assets and an open FiB context or an active async search query
+	return IsCacheInProgress() || bNeedsFirstIndex || (bHasPendingAssets && (IsGlobalFindResultsOpen() || IsAsyncSearchQueryInProgress()));
 }
 
 TStatId FFindInBlueprintSearchManager::GetStatId() const
diff --git a/Engine/Source/Editor/Kismet/Public/FindInBlueprintManager.h b/Engine/Source/Editor/Kismet/Public/FindInBlueprintManager.h
--- a/Engine/Source/Editor/Kismet/Public/FindInBlueprintManager.h
+++ b/Engine/Source/Editor/Kismet/Public/FindInBlueprintManager.h
@@ -610,6 +610,9 @@ public:
 	/** Returns TRUE if we're still inside the initial asset discovery and registration stage */
 	bool IsAssetDiscoveryInProgress() const;
 
+	/** Returns TRUE if there are one or more active asynchronous search queries */
+	bool IsAsyncSearchQueryInProgress() const;
+
 	/** Returns a weak reference to the widget that initiated the current caching operation */
 	TWeakPtr<SFindInBlueprints> GetSourceCachingWidget() const { return SourceCachingWidget; }
 
