Merge the �Debug/Blueprint Debugger� and the �Window/Debug� windows together and revamp KismetDebuggingView

diff --git a/Engine/Source/Editor/Kismet/Private/BlueprintDebugger.cpp b/Engine/Source/Editor/Kismet/Private/BlueprintDebugger.cpp
--- a/Engine/Source/Editor/Kismet/Private/BlueprintDebugger.cpp
+++ b/Engine/Source/Editor/Kismet/Private/BlueprintDebugger.cpp
@@ -10,7 +10,6 @@
 #include "Framework/Docking/TabManager.h"
 #include "Framework/MultiBox/MultiBoxBuilder.h"
 #include "Layout/WidgetPath.h"
-#include "WatchPointViewer.h"
 #include "Widgets/Docking/SDockTab.h"
 #include "WorkspaceMenuStructure.h"
 #include "WorkspaceMenuStructureModule.h"
@@ -35,15 +34,13 @@ struct FBlueprintDebuggerCommands : public TCommands<FBlueprintDebuggerCommands>
 	// End of TCommand<> interface
 
 	TSharedPtr<FUICommandInfo> ShowCallStackViewer;
-	TSharedPtr<FUICommandInfo> ShowWatchViewer;
 	TSharedPtr<FUICommandInfo> ShowExecutionTrace;
 };
 
 void FBlueprintDebuggerCommands::RegisterCommands()
 {
 	UI_COMMAND(ShowCallStackViewer, "Call Stack", "Toggles visibility of the Call Stack window", EUserInterfaceActionType::Check, FInputChord());
-	UI_COMMAND(ShowWatchViewer, "Watches", "Toggles visibility of the Watches window", EUserInterfaceActionType::Check, FInputChord());
-	UI_COMMAND(ShowExecutionTrace, "Execution Flow", "Toggles visibility of the Execution Flow window", EUserInterfaceActionType::Check, FInputChord());
+	UI_COMMAND(ShowExecutionTrace, "Data Flow", "Toggles visibility of the Data Flow window", EUserInterfaceActionType::Check, FInputChord());
 }
 
 struct FBlueprintDebuggerImpl
@@ -124,7 +121,6 @@ TSharedRef<SDockTab> FBlueprintDebuggerImpl::CreateBluprintDebuggerTab(const FSp
 
 	const FName ExecutionFlowTabName = FName(TEXT("ExecutionFlowApp"));
 	const FName CallStackTabName = CallStackViewer::GetTabName();
-	const FName WatchViewerTabName = WatchViewer::GetTabName();
 
 	TWeakPtr<FTabManager> DebuggingToolsTabManagerWeak = DebuggingToolsTabManager;
 	// On tab close will save the layout if the debugging window itself is closed,
@@ -150,22 +146,23 @@ TSharedRef<SDockTab> FBlueprintDebuggerImpl::CreateBluprintDebuggerTab(const FSp
 			FOnSpawnTab::CreateStatic(
 				[](const FSpawnTabArgs&)->TSharedRef<SDockTab>
 				{
+				const TSharedPtr<SKismetDebuggingView> KismetDebuggingView = SNew(SKismetDebuggingView)
+					.BlueprintToWatch(nullptr);
 				return SNew(SDockTab)
 					.TabRole(ETabRole::PanelTab)
-					.Label(NSLOCTEXT("BlueprintExecutionFlow", "TabTitle", "Execution Flow"))
+					.Label_Raw(KismetDebuggingView.Get(), &SKismetDebuggingView::GetTabLabel)
 					[
-						SNew(SKismetDebuggingView)
+						KismetDebuggingView.ToSharedRef()
 					];
 				}
 			)
 		)
-		.SetDisplayName(NSLOCTEXT("BlueprintDebugger", "ExecutionFlowTabTitle", "Blueprint Execution Flow"))
-		.SetTooltipText(NSLOCTEXT("BlueprintDebugger", "ExecutionFlowTooltipText", "Open the Blueprint Execution Flow tab."));
+		.SetDisplayName(NSLOCTEXT("BlueprintDebugger", "ExecutionFlowTabTitle", "Blueprint Data Flow"))
+		.SetTooltipText(NSLOCTEXT("BlueprintDebugger", "ExecutionFlowTooltipText", "Open the Blueprint Data Flow tab."));
 
-		CallStackViewer::RegisterTabSpawner(*DebuggingToolsTabManager);
-		WatchViewer::RegisterTabSpawner(*DebuggingToolsTabManager);
+		CallStackViewer::RegisterTabSpawner(*DebuggingToolsTabManager); 
 
-		BlueprintDebuggerLayout = FTabManager::NewLayout("Standalone_BlueprintDebugger_Layout_v1")
+		BlueprintDebuggerLayout = FTabManager::NewLayout("Standalone_BlueprintDebugger_Layout_v2")
 		->AddArea
 		(
 			FTabManager::NewPrimaryArea()
@@ -176,7 +173,6 @@ TSharedRef<SDockTab> FBlueprintDebuggerImpl::CreateBluprintDebuggerTab(const FSp
 				->SetSizeCoefficient(.4f)
 				->SetHideTabWell(true)
 				->AddTab(CallStackTabName, ETabState::OpenedTab)
-				->AddTab(WatchViewerTabName, ETabState::OpenedTab)
 				->AddTab(ExecutionFlowTabName, ETabState::OpenedTab)
 				->SetForegroundTab(CallStackTabName)
 			)
@@ -236,23 +232,6 @@ TSharedRef<SDockTab> FBlueprintDebuggerImpl::CreateBluprintDebuggerTab(const FSp
 		)
 	);
 
-	CommandList->MapAction(
-		FBlueprintDebuggerCommands::Get().ShowWatchViewer,
-		FExecuteAction::CreateStatic(
-			ToggleTabVisibility,
-			DebuggingToolsManagerWeak,
-			WatchViewerTabName
-		),
-		FCanExecuteAction::CreateStatic(
-			[]() { return true; }
-		),
-		FIsActionChecked::CreateStatic(
-			IsTabVisible,
-			DebuggingToolsManagerWeak,
-			WatchViewerTabName
-		)
-	);
-
 	CommandList->MapAction(
 		FBlueprintDebuggerCommands::Get().ShowExecutionTrace,
 		FExecuteAction::CreateStatic(
@@ -292,7 +271,6 @@ TSharedRef<SDockTab> FBlueprintDebuggerImpl::CreateBluprintDebuggerTab(const FSp
 						MenuBuilder.PushCommandList(InCommandList.ToSharedRef());
 						{
 							MenuBuilder.AddMenuEntry(FBlueprintDebuggerCommands::Get().ShowCallStackViewer);
-							MenuBuilder.AddMenuEntry(FBlueprintDebuggerCommands::Get().ShowWatchViewer);
 							MenuBuilder.AddMenuEntry(FBlueprintDebuggerCommands::Get().ShowExecutionTrace);
 						}
 						MenuBuilder.PopCommandList();
diff --git a/Engine/Source/Editor/Kismet/Private/BlueprintEditor.cpp b/Engine/Source/Editor/Kismet/Private/BlueprintEditor.cpp
--- a/Engine/Source/Editor/Kismet/Private/BlueprintEditor.cpp
+++ b/Engine/Source/Editor/Kismet/Private/BlueprintEditor.cpp
@@ -1969,7 +1969,6 @@ void FBlueprintEditor::LoadLibrariesFromAssetRegistry()
 						if (BlueprintLibPtr)
 						{
 							StandardLibraries.AddUnique(BlueprintLibPtr);
-							WatchViewer::UpdateWatchListFromBlueprint(BlueprintLibPtr);
 						}
 
 						const double ElapsedTime = FPlatformTime::Seconds() - StartTime;
diff --git a/Engine/Source/Editor/Kismet/Private/BlueprintEditorModule.cpp b/Engine/Source/Editor/Kismet/Private/BlueprintEditorModule.cpp
--- a/Engine/Source/Editor/Kismet/Private/BlueprintEditorModule.cpp
+++ b/Engine/Source/Editor/Kismet/Private/BlueprintEditorModule.cpp
@@ -278,11 +278,6 @@ TSharedRef<IBlueprintEditor> FBlueprintEditorModule::CreateBlueprintEditor(const
 		NewBlueprintEditor->RegisterSCSEditorCustomization(It->Key, It->Value.Execute(NewBlueprintEditor));
 	}
 
-	for (UBlueprint* Blueprint : BlueprintsToEdit)
-	{
-		WatchViewer::UpdateWatchListFromBlueprint(Blueprint);
-	}
-
 	EBlueprintType const BPType = ( (BlueprintsToEdit.Num() > 0) && (BlueprintsToEdit[0] != NULL) ) 
 		? (EBlueprintType) BlueprintsToEdit[0]->BlueprintType
 		: BPTYPE_Normal;
diff --git a/Engine/Source/Editor/Kismet/Private/CallStackViewer.cpp b/Engine/Source/Editor/Kismet/Private/CallStackViewer.cpp
--- a/Engine/Source/Editor/Kismet/Private/CallStackViewer.cpp
+++ b/Engine/Source/Editor/Kismet/Private/CallStackViewer.cpp
@@ -16,14 +16,12 @@
 #include "Widgets/Docking/SDockTab.h"
 #include "Widgets/Images/SImage.h"
 #include "Widgets/Layout/SBox.h"
-#include "Widgets/Layout/SSpacer.h"
 #include "Widgets/Views/STreeView.h"
-
-#include "WorkspaceMenuStructure.h"
-#include "WorkspaceMenuStructureModule.h"
+#include "ToolMenus.h"
 
 // So that we can poll the running state:
 #include "Editor/UnrealEdEngine.h"
+#include "Kismet2/DebuggerCommands.h"
 extern UNREALED_API UUnrealEdEngine* GUnrealEd;
 
 #define LOCTEXT_NAMESPACE "CallStackViewer"
@@ -587,14 +585,31 @@ void CallStackViewer::RegisterTabSpawner(FTabManager& TabManager)
 {
 	const auto SpawnCallStackViewTab = []( const FSpawnTabArgs& Args )
 	{
+		static const FName ToolbarName = TEXT("Kismet.DebuggingViewToolBar");
+		FToolMenuContext MenuContext(FPlayWorldCommands::GlobalPlayWorldActions);
+		TSharedRef<SWidget> ToolbarWidget = UToolMenus::Get()->GenerateWidget(ToolbarName, MenuContext);
 		return SNew(SDockTab)
 			.TabRole( ETabRole::PanelTab )
 			.Label( LOCTEXT("TabTitle", "Call Stack") )
 			[
-				SNew(SBorder)
-				.BorderImage( FEditorStyle::GetBrush("Docking.Tab.ContentAreaBrush") )
+				SNew(SVerticalBox)
+				+ SVerticalBox::Slot()
+				.AutoHeight()
+				[
+					SNew(SBorder)
+					.BorderImage( FEditorStyle::GetBrush( TEXT("NoBorder") ) )
+					[
+						ToolbarWidget
+					]
+				]
+				+ SVerticalBox::Slot()
+				.AutoHeight()
 				[
-					SNew(SCallStackViewer, &Private_CallStackSource)
+					SNew(SBorder)
+					.BorderImage( FEditorStyle::GetBrush("Docking.Tab.ContentAreaBrush") )
+					[
+						SNew(SCallStackViewer, &Private_CallStackSource)
+					]
 				]
 			];
 	};
diff --git a/Engine/Source/Editor/Kismet/Private/Debugging/SKismetDebuggingView.cpp b/Engine/Source/Editor/Kismet/Private/Debugging/SKismetDebuggingView.cpp
--- a/Engine/Source/Editor/Kismet/Private/Debugging/SKismetDebuggingView.cpp
+++ b/Engine/Source/Editor/Kismet/Private/Debugging/SKismetDebuggingView.cpp
@@ -2,6 +2,9 @@
 
 
 #include "Debugging/SKismetDebuggingView.h"
+
+#include "ClassViewerFilter.h"
+#include "ClassViewerModule.h"
 #include "EdGraph/EdGraphPin.h"
 #include "Engine/Blueprint.h"
 #include "Textures/SlateIcon.h"
@@ -24,6 +27,9 @@
 #include "Debugging/KismetDebugCommands.h"
 #include "Widgets/Input/SHyperlink.h"
 #include "ToolMenus.h"
+#include "PropertyEditor/Private/SDetailsView.h"
+#include "Styling/SlateIconFinder.h"
+#include "Styling/StyleColors.h"
 
 #define LOCTEXT_NAMESPACE "DebugViewUI"
 
@@ -37,11 +43,15 @@ namespace KismetDebugViewConstants
 	const FName ColumnId_Value( "Value" );
 	const FText ColumnText_Name( NSLOCTEXT("DebugViewUI", "Name", "Name") );
 	const FText ColumnText_Value( NSLOCTEXT("DebugViewUI", "Value", "Value") );
+	const FText ColumnText_DebugKey( FText::GetEmpty() );
+	const FText ColumnText_Info( NSLOCTEXT("DebugViewUI", "Info", "Info") );
 }
 
 //////////////////////////////////////////////////////////////////////////
 // FDebugLineItem
 
+uint16 FDebugLineItem::ActiveTypeBitset = TNumericLimits<uint16>::Max(); // set all to active by default
+
 FText FDebugLineItem::GetDisplayName() const
 {
 	return FText::GetEmpty();
@@ -70,10 +80,19 @@ UBlueprint* FDebugLineItem::GetBlueprintForObject(UObject* ParentObject)
 	{
 		return nullptr;
 	}
+	
 	if(UBlueprint* ParentBlueprint = Cast<UBlueprint>(ParentObject))
 	{
 		return ParentBlueprint;
 	}
+
+	if(UClass *ParentClass = ParentObject->GetClass())
+	{
+		if(UBlueprint* ParentBlueprint = Cast<UBlueprint>(ParentClass->ClassGeneratedBy))
+		{
+			return ParentBlueprint;
+		}
+	}
 	
 	// recursively walk up ownership heirrarchy until we find the blueprint
 	return GetBlueprintForObject(ParentObject->GetOuter());
@@ -81,7 +100,7 @@ UBlueprint* FDebugLineItem::GetBlueprintForObject(UObject* ParentObject)
 
 UBlueprintGeneratedClass* FDebugLineItem::GetClassForObject(UObject* ParentObject)
 {
-	if (ParentObject != NULL)
+	if (ParentObject != nullptr)
 	{
 		if (UBlueprint* Blueprint = Cast<UBlueprint>(ParentObject))
 		{
@@ -97,31 +116,74 @@ UBlueprintGeneratedClass* FDebugLineItem::GetClassForObject(UObject* ParentObjec
 		}
 	}
 
-	return NULL;
+	return nullptr;
 }
 
-// O(# children)
-void FDebugLineItem::EnsureChildIsAdded(TArray<FDebugTreeItemPtr>& ChildrenMirrors, TArray<FDebugTreeItemPtr>& OutChildren, const FDebugLineItem& Item)
+bool FDebugLineItem::IsDebugLineTypeActive(EDebugLineType Type)
 {
-	for (int32 i = 0; i < ChildrenMirrors.Num(); ++i)
+	const uint16 Mask = 1 << Type;
+	return ActiveTypeBitset & Mask;
+}
+
+void FDebugLineItem::OnDebugLineTypeActiveChanged(ECheckBoxState CheckState, EDebugLineType Type)
+{
+	const uint16 Mask = 1 << Type;
+	switch (CheckState)
 	{
-		TSharedPtr< FDebugLineItem > MirrorItem = ChildrenMirrors[i];
+		case ECheckBoxState::Checked:
+		{
+			ActiveTypeBitset |= Mask;
+			break;
+		}
+		default:
+		{
+			ActiveTypeBitset &= ~Mask;
+			break;
+		}
+	}
+}
+
+//////////////////////////////////////////////////////////////////////////
+// ILineItemWithChildren
 
-		if (MirrorItem->Type == Item.Type)
+class ILineItemWithChildren
+{
+public:
+	virtual ~ILineItemWithChildren() = default;
+	
+protected:
+	// List of children 
+	TArray<FDebugTreeItemPtr> ChildrenMirrors;
+	
+	/**
+	 * Adds either Item or an identical node that was previously
+	 * created (present in ChildrenMirrors) as a child to OutChildren.
+	 *
+	 * O( # Children )
+	 */
+	void EnsureChildIsAdded(TArray<FDebugTreeItemPtr>& OutChildren, const FDebugLineItem& Item)
+	{
+		for (int32 i = 0; i < ChildrenMirrors.Num(); ++i)
 		{
-			if (Item.Compare(MirrorItem.Get()))
+			TSharedPtr< FDebugLineItem > MirrorItem = ChildrenMirrors[i];
+
+			if (MirrorItem->Type == Item.Type)
 			{
-				OutChildren.Add(MirrorItem);
-				return;
+				if (Item.Compare(MirrorItem.Get()))
+				{
+					OutChildren.Add(MirrorItem);
+					return;
+				}
 			}
 		}
-	}
 
-	TSharedPtr< FDebugLineItem > Result = MakeShareable(Item.Duplicate());
-	ChildrenMirrors.Add(Result);
+		TSharedPtr< FDebugLineItem > Result = MakeShareable(Item.Duplicate());
+		ChildrenMirrors.Add(Result);
+
+		OutChildren.Add(Result);
+	}
+};
 
-	OutChildren.Add(Result);
-}
 
 //////////////////////////////////////////////////////////////////////////
 // FMessageLineItem
@@ -171,6 +233,7 @@ public:
 		check(UUID != INDEX_NONE);
 		ParentObjectRef = ParentObject;
 	}
+
 protected:
 	virtual bool Compare(const FDebugLineItem* BaseOther) const override
 	{
@@ -203,7 +266,7 @@ FText FLatentActionLineItem::GetDescription() const
 		}
 	}
 
-	return LOCTEXT("NullObject", "Object has been destroyed");
+	return LOCTEXT("nullptrObject", "Object has been destroyed");
 }
 
 TSharedRef<SWidget> FLatentActionLineItem::GenerateNameWidget()
@@ -236,7 +299,7 @@ UEdGraphNode* FLatentActionLineItem::FindAssociatedNode() const
 		return Class->GetDebugData().FindNodeFromUUID(UUID);
 	}
 
-	return NULL;
+	return nullptr;
 }
 
 FText FLatentActionLineItem::GetDisplayName() const
@@ -263,14 +326,69 @@ void FLatentActionLineItem::OnNavigateToLatentNode( )
 	}
 }
 
+struct FWatchChildLineItem : public FDebugLineItem, ILineItemWithChildren
+{
+protected:
+	FDebugInfo Data;
+public:
+	FWatchChildLineItem(FDebugInfo Child) :
+		FDebugLineItem(DLT_WatchChild),
+		Data(Child)
+	{}
+
+	virtual bool Compare(const FDebugLineItem* BaseOther) const override
+	{
+		FWatchChildLineItem* Other = (FWatchChildLineItem*)BaseOther;
+		return
+			Data.DisplayName.CompareTo(Other->Data.DisplayName) == 0 &&
+			Data.Value.CompareTo(Other->Data.Value) == 0;
+	}
+
+	virtual FDebugLineItem* Duplicate() const override
+	{
+		return new FWatchChildLineItem(Data);
+	}
+	
+	virtual FText GetDescription() const override
+	{
+		FString ValStr = Data.Value.ToString();
+		if(int32 idx; ValStr.FindChar('\n', idx))
+		{
+			if(ValStr.Len() < 60)
+			{
+				return FText::FromString(ValStr.Replace(TEXT("\n") ,TEXT("")));
+			}
+			return FText::Format(LOCTEXT("BracketedDatatype","[{0}]"), Data.Type);
+		}
+		return Data.Value;
+	}
+	
+	virtual FText GetDisplayName() const override
+	{
+		return Data.DisplayName;
+	}
+
+	virtual void GatherChildren(TArray<FDebugTreeItemPtr>& OutChildren) override
+	{
+		for (FDebugInfo& ChildData : Data.Children)
+		{
+			EnsureChildIsAdded(OutChildren, FWatchChildLineItem(ChildData));
+		}
+	}
+};
+
 //////////////////////////////////////////////////////////////////////////
 // FWatchLineItem 
 
-struct FWatchLineItem : public FDebugLineItem
+
+struct FWatchLineItem : public FDebugLineItem, ILineItemWithChildren
 {
 protected:
 	TWeakObjectPtr< UObject > ParentObjectRef;
 	FEdGraphPinReference ObjectRef;
+	
+	// List of children 
+	TArray<FDebugTreeItemPtr> ChildrenMirrors;
 public:
 	FWatchLineItem(UEdGraphPin* PinToWatch, UObject* ParentObject)
 		: FDebugLineItem(DLT_Watch)
@@ -306,6 +424,33 @@ public:
 				ClearThisWatch);
 		}
 	}
+	
+	virtual void GatherChildren(TArray<FDebugTreeItemPtr>& OutChildren) override
+	{
+		if (UEdGraphPin* PinToWatch = ObjectRef.Get())
+		{
+			// Try to determine the blueprint that generated the watch
+			UBlueprint* ParentBlueprint = GetBlueprintForObject(ParentObjectRef.Get());
+
+			// Find a valid property mapping and display the current value
+			UObject* ParentObject = ParentObjectRef.Get();
+			if ((ParentBlueprint != ParentObject) && (ParentBlueprint != nullptr))
+			{
+				FDebugInfo DebugInfo;
+				const FKismetDebugUtilities::EWatchTextResult WatchStatus = FKismetDebugUtilities::GetDebugInfo(DebugInfo, ParentBlueprint, ParentObject, PinToWatch);
+
+				if(WatchStatus == FKismetDebugUtilities::EWTR_Valid)
+				{
+					for (FDebugInfo& ChildData : DebugInfo.Children)
+					{
+						EnsureChildIsAdded(OutChildren, FWatchChildLineItem(ChildData));
+					}
+				}
+			}
+		}
+	}
+	
+
 protected:
 	virtual FText GetDescription() const override;
 	virtual FText GetDisplayName() const override;
@@ -345,15 +490,23 @@ FText FWatchLineItem::GetDescription() const
 
 		// Find a valid property mapping and display the current value
 		UObject* ParentObject = ParentObjectRef.Get();
-		if ((ParentBlueprint != ParentObject) && (ParentBlueprint != NULL))
+		if ((ParentBlueprint != ParentObject) && (ParentBlueprint != nullptr))
 		{
-			FString WatchString;
-			const FKismetDebugUtilities::EWatchTextResult WatchStatus = FKismetDebugUtilities::GetWatchText(WatchString, ParentBlueprint, ParentObject, PinToWatch);
+			
+			FDebugInfo DebugInfo;
+			const FKismetDebugUtilities::EWatchTextResult WatchStatus = FKismetDebugUtilities::GetDebugInfo(DebugInfo, ParentBlueprint, ParentObject, PinToWatch);
 
 			switch (WatchStatus)
 			{
 			case FKismetDebugUtilities::EWTR_Valid:
-				return FText::FromString(WatchString);
+			{
+				int32 Index = INDEX_NONE;
+				if(DebugInfo.Value.ToString().FindChar('\n', Index))
+				{
+					return DebugInfo.Type;
+				}
+				return DebugInfo.Value;
+			}
 
 			case FKismetDebugUtilities::EWTR_NotInScope:
 				return LOCTEXT("NotInScope", "Not in scope");
@@ -402,15 +555,15 @@ void FWatchLineItem::OnNavigateToWatchLocation( )
 }
 
 //////////////////////////////////////////////////////////////////////////
-// FBlueprintBreakpointLineItem
+// FBreakpointLineItem
 
-struct FBlueprintBreakpointLineItem : public FDebugLineItem
+struct FBreakpointLineItem : public FDebugLineItem
 {
 protected:
 	TWeakObjectPtr<UObject> ParentObjectRef;
 	TSoftObjectPtr<UEdGraphNode> BreakpointNode;
 public:
-	FBlueprintBreakpointLineItem(TSoftObjectPtr<UEdGraphNode> BreakpointToWatch, UObject* ParentObject)
+	FBreakpointLineItem(TSoftObjectPtr<UEdGraphNode> BreakpointToWatch, UObject* ParentObject)
 		: FDebugLineItem(DLT_Breakpoint)
 	{
 		BreakpointNode = BreakpointToWatch;
@@ -419,14 +572,14 @@ public:
 
 	virtual bool Compare(const FDebugLineItem* BaseOther) const override
 	{
-		FBlueprintBreakpointLineItem* Other = (FBlueprintBreakpointLineItem*)BaseOther;
+		FBreakpointLineItem* Other = (FBreakpointLineItem*)BaseOther;
 		return (ParentObjectRef.Get() == Other->ParentObjectRef.Get()) &&
 			(BreakpointNode == Other->BreakpointNode);
 	}
 
 	virtual FDebugLineItem* Duplicate() const override
 	{
-		return new FBlueprintBreakpointLineItem(BreakpointNode, ParentObjectRef.Get());
+		return new FBreakpointLineItem(BreakpointNode, ParentObjectRef.Get());
 	}	
 
 	virtual void MakeMenu(class FMenuBuilder& MenuBuilder) override
@@ -438,7 +591,7 @@ public:
 		// Create an empty action to always allow execution for these commands (they are allowed in debug mode)
 		FCanExecuteAction AlwaysAllowExecute;
 
-		if (Breakpoint != NULL)
+		if (Breakpoint != nullptr)
 		{
 			const bool bNewEnabledState = !Breakpoint->IsEnabledByUser();
 
@@ -469,7 +622,7 @@ public:
 			}
 		}
 
-		if ((Breakpoint != NULL) && (ParentBlueprint != NULL))
+		if ((Breakpoint != nullptr) && (ParentBlueprint != nullptr))
 		{
 			FUIAction ClearThisBreakpoint(
 				FExecuteAction::CreateStatic( &FDebuggingActionCallbacks::ClearBreakpoint, BreakpointNode, ParentBlueprint ),
@@ -504,14 +657,14 @@ protected:
 			.AutoWidth()
 			[
 				SNew(SButton)
-				. OnClicked(this, &FBlueprintBreakpointLineItem::OnUserToggledEnabled)
+				. OnClicked(this, &FBreakpointLineItem::OnUserToggledEnabled)
 				. ToolTipText(LOCTEXT("ToggleBreakpointButton_ToolTip", "Toggle this breakpoint"))
 				. ButtonStyle( FEditorStyle::Get(), "NoBorder" )
 				. ContentPadding(0.0f)
 				[
 					SNew(SImage)
-					. Image(this, &FBlueprintBreakpointLineItem::GetStatusImage)
-					. ToolTipText(this, &FBlueprintBreakpointLineItem::GetStatusTooltip)
+					. Image(this, &FBreakpointLineItem::GetStatusImage)
+					. ToolTipText(this, &FBreakpointLineItem::GetStatusTooltip)
 				]
 			]
 
@@ -521,9 +674,9 @@ protected:
 			[
 				SNew(SHyperlink)
 				. Style(FEditorStyle::Get(), "HoverOnlyHyperlink")
-				. Text(this, &FBlueprintBreakpointLineItem::GetLocationDescription)
+				. Text(this, &FBreakpointLineItem::GetLocationDescription)
 				. ToolTipText( LOCTEXT("NavBreakpointLoc", "Navigate to the breakpoint location") )
-				. OnNavigate(this, &FBlueprintBreakpointLineItem::OnNavigateToBreakpointLocation)
+				. OnNavigate(this, &FBreakpointLineItem::OnNavigateToBreakpointLocation)
 			];
 	}
 	
@@ -537,7 +690,7 @@ protected:
 	FText GetStatusTooltip() const;
 };
 
-FText FBlueprintBreakpointLineItem::GetLocationDescription() const
+FText FBreakpointLineItem::GetLocationDescription() const
 {
 	if (FBlueprintBreakpoint* MyBreakpoint = GetBreakpoint())
 	{
@@ -546,7 +699,7 @@ FText FBlueprintBreakpointLineItem::GetLocationDescription() const
 	return FText::GetEmpty();
 }
 
-FReply FBlueprintBreakpointLineItem::OnUserToggledEnabled()
+FReply FBreakpointLineItem::OnUserToggledEnabled()
 {
 	if (FBlueprintBreakpoint* MyBreakpoint = GetBreakpoint())
 	{
@@ -555,7 +708,7 @@ FReply FBlueprintBreakpointLineItem::OnUserToggledEnabled()
 	return FReply::Handled();
 }
 
-void FBlueprintBreakpointLineItem::OnNavigateToBreakpointLocation()
+void FBreakpointLineItem::OnNavigateToBreakpointLocation()
 {
 	if (FBlueprintBreakpoint* MyBreakpoint = GetBreakpoint())
 	{
@@ -563,7 +716,7 @@ void FBlueprintBreakpointLineItem::OnNavigateToBreakpointLocation()
 	}
 }
 
-const FSlateBrush* FBlueprintBreakpointLineItem::GetStatusImage() const
+const FSlateBrush* FBreakpointLineItem::GetStatusImage() const
 {
 	if (FBlueprintBreakpoint* MyBreakpoint = GetBreakpoint())
 	{
@@ -580,7 +733,7 @@ const FSlateBrush* FBlueprintBreakpointLineItem::GetStatusImage() const
 	return FEditorStyle::GetDefaultBrush();
 }
 
-FText FBlueprintBreakpointLineItem::GetStatusTooltip() const
+FText FBreakpointLineItem::GetStatusTooltip() const
 {
 	if (FBlueprintBreakpoint* MyBreakpoint = GetBreakpoint())
 	{
@@ -599,10 +752,72 @@ FText FBlueprintBreakpointLineItem::GetStatusTooltip() const
 	}
 }
 
+//////////////////////////////////////////////////////////////////////////
+// FTraceStackParentItem
+
+class FBreakpointParentItem : public FDebugLineItem, ILineItemWithChildren
+{
+protected:
+
+	// List of children 
+	TArray<FDebugTreeItemPtr> ChildrenMirrors;
+public:
+	// The parent object
+	TWeakObjectPtr<UBlueprint> Blueprint;
+	
+	FBreakpointParentItem(TWeakObjectPtr<UBlueprint> Blueprint)
+		: FDebugLineItem(DLT_TraceStackParent),
+		  Blueprint(Blueprint)
+	{
+	}
+
+	virtual void GatherChildren(TArray<FDebugTreeItemPtr>& OutChildren) override
+	{
+		if(!Blueprint.IsValid())
+		{
+			return;
+		}
+		
+		// Create children for each breakpoint
+		FKismetDebugUtilities::ForeachBreakpoint(
+			Blueprint.Get(),
+			[this, &OutChildren]
+			(FBlueprintBreakpoint& Breakpoint)
+			{
+				EnsureChildIsAdded(OutChildren, FBreakpointLineItem(Breakpoint.GetLocation(), Blueprint.Get()));
+			}
+		);
+
+		// Make sure there is something there, to let the user know if there is nothing
+		if (OutChildren.Num() == 0)
+		{
+			EnsureChildIsAdded(OutChildren, FMessageLineItem( LOCTEXT("NoBreakpoints", "No breakpoints").ToString() ));
+		}
+	}
+
+protected:
+	virtual FText GetDisplayName() const override
+	{
+		return LOCTEXT("Breakpoints", "Breakpoints");
+	}
+
+	virtual bool Compare(const FDebugLineItem* BaseOther) const override
+	{
+		check(false);
+		return false;
+	}
+
+	virtual FDebugLineItem* Duplicate() const override
+	{
+		check(false);
+		return nullptr;
+	}
+};
+
 //////////////////////////////////////////////////////////////////////////
 // FParentLineItem
 
-class FParentLineItem : public FDebugLineItem
+class FParentLineItem : public FDebugLineItem, ILineItemWithChildren
 {
 protected:
 	// The parent object
@@ -627,80 +842,98 @@ public:
 		if (UObject* ParentObject = ObjectRef.Get())
 		{
 			UBlueprint* ParentBP = FDebugLineItem::GetBlueprintForObject(ParentObject);
-			if ((ParentBP != NULL) && (ParentBP == ParentObject))
+			if ((ParentBP != nullptr) && (ParentBP == ParentObject))
 			{
 				// Create children for each watch
-				for (int32 WatchIndex = 0; WatchIndex < ParentBP->WatchedPins.Num(); ++WatchIndex)
+				if(IsDebugLineTypeActive(DLT_Watch))
 				{
-					UEdGraphPin* WatchedPin = ParentBP->WatchedPins[WatchIndex].Get();
-
-					EnsureChildIsAdded(ChildrenMirrors, OutChildren, FWatchLineItem(WatchedPin, ParentObject));
+					for (FEdGraphPinReference WatchedPin : ParentBP->WatchedPins)
+                    {
+                    	EnsureChildIsAdded(OutChildren, FWatchLineItem(WatchedPin.Get(), ParentObject));
+                    }
 				}
 
-				// Create children for each breakpoint
-				FKismetDebugUtilities::ForeachBreakpoint(
-					ParentBP,
-					[this, &OutChildren, ParentObject](FBlueprintBreakpoint& Breakpoint)
-					{
-						EnsureChildIsAdded(ChildrenMirrors, OutChildren, FBlueprintBreakpointLineItem(Breakpoint.GetLocation(), ParentObject));
-					}
-				);
-
 				// Make sure there is something there, to let the user know if there is nothing
 				if (OutChildren.Num() == 0)
 				{
-					EnsureChildIsAdded(ChildrenMirrors, OutChildren, FMessageLineItem( LOCTEXT("NoWatchesOrBreakpoints", "No watches or breakpoints").ToString() ));
+					EnsureChildIsAdded(OutChildren, FMessageLineItem( LOCTEXT("NoWatches", "No watches").ToString() ));
 				}
 			}
 			else
 			{
-				if (ParentBP != NULL)
+				if (ParentBP != nullptr)
 				{
 					// Create children for each watch
-					for (int32 WatchIndex = 0; WatchIndex < ParentBP->WatchedPins.Num(); ++WatchIndex)
+					if(IsDebugLineTypeActive(DLT_Watch))
 					{
-						UEdGraphPin* WatchedPin = ParentBP->WatchedPins[WatchIndex].Get();
-
-						EnsureChildIsAdded(ChildrenMirrors, OutChildren, FWatchLineItem(WatchedPin, ParentObject));
-					}
-
-					// Create children for each breakpoint
-					FKismetDebugUtilities::ForeachBreakpoint(
-						ParentBP,
-						[this, &OutChildren, ParentObject](FBlueprintBreakpoint& Breakpoint)
+						for (FEdGraphPinReference WatchedPin : ParentBP->WatchedPins)
 						{
-							EnsureChildIsAdded(ChildrenMirrors, OutChildren, FBlueprintBreakpointLineItem(Breakpoint.GetLocation(), ParentObject));
+							EnsureChildIsAdded(OutChildren, FWatchLineItem(WatchedPin.Get(), ParentObject));
 						}
-					);
+					}
 				}
 
 				// It could also have active latent behaviors
-				if (UWorld* World = GEngine->GetWorldFromContextObject(ParentObject, EGetWorldErrorMode::ReturnNull))
+				if(IsDebugLineTypeActive(DLT_LatentAction))
 				{
-					FLatentActionManager& LatentActionManager = World->GetLatentActionManager();
+					if (UWorld* World = GEngine->GetWorldFromContextObject(ParentObject, EGetWorldErrorMode::ReturnNull))
+					{
+						FLatentActionManager& LatentActionManager = World->GetLatentActionManager();
 
-					// Get the current list of action UUIDs
-					TSet<int32> UUIDSet;
-					LatentActionManager.GetActiveUUIDs(ParentObject, /*inout*/ UUIDSet);
+						// Get the current list of action UUIDs
+						TSet<int32> UUIDSet;
+						LatentActionManager.GetActiveUUIDs(ParentObject, /*inout*/ UUIDSet);
 
-					// Add the new ones
-					for (TSet<int32>::TConstIterator RemainingIt(UUIDSet); RemainingIt; ++RemainingIt)
-					{
-						const int32 UUID = *RemainingIt;
-						EnsureChildIsAdded(ChildrenMirrors, OutChildren, FLatentActionLineItem(UUID, ParentObject));
+						// Add the new ones
+						for (TSet<int32>::TConstIterator RemainingIt(UUIDSet); RemainingIt; ++RemainingIt)
+						{
+							const int32 UUID = *RemainingIt;
+							EnsureChildIsAdded(OutChildren, FLatentActionLineItem(UUID, ParentObject));
+						}
 					}
 				}
 
 				// Make sure there is something there, to let the user know if there is nothing
 				if (OutChildren.Num() == 0)
 				{
-					EnsureChildIsAdded(ChildrenMirrors, OutChildren, FMessageLineItem( LOCTEXT("NoDebugInfo", "No debugging info").ToString() ));
+					EnsureChildIsAdded(OutChildren, FMessageLineItem( LOCTEXT("NoDebugInfo", "No debugging info").ToString() ));
 				}
 			}
 			//@TODO: try to get at TArray<struct FDebugDisplayProperty> DebugProperties in UGameViewportClient, if available
 		}
 	}
 
+	const FSlateBrush* GetStatusImage() const
+	{
+		if(SKismetDebuggingView::CurrentActiveObject == ObjectRef)
+		{
+			return FEditorStyle::GetBrush(TEXT("Kismet.Trace.CurrentIndex"));
+		}
+		if(ObjectRef.IsValid())
+		{
+			return FSlateIconFinder::FindIconBrushForClass(ObjectRef->GetClass());
+		}
+		return FEditorStyle::GetBrush(TEXT("None"));
+	}
+
+	FSlateColor GetStatusColor() const
+	{
+		if(SKismetDebuggingView::CurrentActiveObject == ObjectRef)
+		{
+			return FSlateColor(EStyleColor::AccentYellow);
+		}
+		return FSlateColor::UseForeground();
+	}
+	
+	FText GetStatusTooltip() const
+	{
+		if(SKismetDebuggingView::CurrentActiveObject == ObjectRef)
+		{
+			return LOCTEXT("BreakpointHIt", "Breakpoint Hit");
+		}
+		return FText::GetEmpty();
+	}
+
 protected:
 	virtual bool Compare(const FDebugLineItem* BaseOther) const override
 	{
@@ -718,21 +951,41 @@ protected:
 		UObject* Object = ObjectRef.Get();
 		AActor* Actor = Cast<AActor>(Object);
 
-		if (Actor != NULL)
+		if (Actor != nullptr)
 		{
 			 return FText::FromString(Actor->GetActorLabel());
 		}
 		else
 		{
-			return (Object != NULL) ? FText::FromString(Object->GetName()) : LOCTEXT("Null", "(null)");
+			return (Object != nullptr) ? FText::FromString(Object->GetName()) : LOCTEXT("nullptr", "(nullptr)");
 		}
 	}
 
+	virtual TSharedRef<SWidget> GenerateNameWidget() override
+	{
+		return SNew(SHorizontalBox)
+			+ SHorizontalBox::Slot()
+			.AutoWidth()
+			[
+				SNew(SImage)
+				.Image(this, &FParentLineItem::GetStatusImage)
+				.ColorAndOpacity_Raw(this, &FParentLineItem::GetStatusColor)
+				.ToolTipText(this, &FParentLineItem::GetStatusTooltip)
+			]
+			+ SHorizontalBox::Slot()
+			.AutoWidth()
+			. VAlign(VAlign_Center)
+			[
+				SNew(STextBlock)
+				.Text(this, &FParentLineItem::GetDisplayName)
+			];
+	}
+
 	virtual void MakeMenu(class FMenuBuilder& MenuBuilder) override
 	{
 		if (UBlueprint* BP = Cast<UBlueprint>(ObjectRef.Get()))
 		{
-			if (FKismetDebugUtilities::BlueprintHasBreakpoints(BP))
+			if (BP->WatchedPins.Num() > 0)
 			{
 				FUIAction ClearAllWatches(
 					FExecuteAction::CreateStatic( &FDebuggingActionCallbacks::ClearWatches, BP )
@@ -743,7 +996,10 @@ protected:
 					LOCTEXT("ClearWatches_ToolTip", "Clear all watches in this blueprint"),
 					FSlateIcon(),
 					ClearAllWatches);
-			
+			}
+
+			if (FKismetDebugUtilities::BlueprintHasBreakpoints(BP))
+			{
 				FUIAction ClearAllBreakpoints(
 					FExecuteAction::CreateStatic( &FDebuggingActionCallbacks::ClearBreakpoints, BP )
 					);
@@ -781,7 +1037,7 @@ protected:
 	virtual FDebugLineItem* Duplicate() const override
 	{
 		check(false);
-		return NULL;
+		return nullptr;
 	}
 
 	UEdGraphNode* GetNode() const
@@ -792,10 +1048,10 @@ protected:
 			const FKismetTraceSample& Sample = TraceStack(StackIndex);
 			UObject* ObjectContext = Sample.Context.Get();
 
-			FString ContextName = (ObjectContext != NULL) ? ObjectContext->GetName() : LOCTEXT("ObjectDoesNotExist", "(object no longer exists)").ToString();
+			FString ContextName = (ObjectContext != nullptr) ? ObjectContext->GetName() : LOCTEXT("ObjectDoesNotExist", "(object no longer exists)").ToString();
 			FString NodeName = TEXT(" ");
 
-			if (ObjectContext != NULL)
+			if (ObjectContext != nullptr)
 			{
 				// Try to find the node that got executed
 				UEdGraphNode* Node = FKismetDebugUtilities::FindSourceNodeForCodeLocation(ObjectContext, Sample.Function.Get(), Sample.Offset);
@@ -803,13 +1059,13 @@ protected:
 			}
 		}
 
-		return NULL;
+		return nullptr;
 	}
 
 	virtual FText GetDisplayName() const override
 	{
 		UEdGraphNode* Node = GetNode();
-		if (Node != NULL)
+		if (Node != nullptr)
 		{
 			return Node->GetNodeTitle(ENodeTitleType::ListView);
 		}
@@ -884,9 +1140,9 @@ protected:
 	{
 		const TSimpleRingBuffer<FKismetTraceSample>& TraceStack = FKismetDebugUtilities::GetTraceStack();
 
-		UObject* ObjectContext = (StackIndex < TraceStack.Num()) ? TraceStack(StackIndex).Context.Get() : NULL;
+		UObject* ObjectContext = (StackIndex < TraceStack.Num()) ? TraceStack(StackIndex).Context.Get() : nullptr;
 		
-		return (ObjectContext != NULL) ? FText::FromString(ObjectContext->GetName()) : LOCTEXT("ObjectDoesNotExist", "(object no longer exists)");
+		return (ObjectContext != nullptr) ? FText::FromString(ObjectContext->GetName()) : LOCTEXT("ObjectDoesNotExist", "(object no longer exists)");
 	}
 
 	void OnNavigateToNode()
@@ -901,7 +1157,7 @@ protected:
 	{
 		const TSimpleRingBuffer<FKismetTraceSample>& TraceStack = FKismetDebugUtilities::GetTraceStack();
 
-		UObject* ObjectContext = (StackIndex < TraceStack.Num()) ? TraceStack(StackIndex).Context.Get() : NULL;
+		UObject* ObjectContext = (StackIndex < TraceStack.Num()) ? TraceStack(StackIndex).Context.Get() : nullptr;
 
 		// Add the object to the selection set
 		if (AActor* Actor = Cast<AActor>(ObjectContext))
@@ -910,7 +1166,7 @@ protected:
 		}
 		else
 		{
-			UE_LOG(LogBlueprintDebuggingView, Warning, TEXT("Cannot select the non-actor object '%s'"), (ObjectContext != NULL) ? *ObjectContext->GetName() : TEXT("(null)") );
+			UE_LOG(LogBlueprintDebuggingView, Warning, TEXT("Cannot select the non-actor object '%s'"), (ObjectContext != nullptr) ? *ObjectContext->GetName() : TEXT("(nullptr)") );
 		}
 	}
 
@@ -947,6 +1203,7 @@ public:
 			OutChildren.Add(ChildrenMirrors[i]);
 		}
 	}
+
 protected:
 	virtual FText GetDisplayName() const override
 	{
@@ -962,7 +1219,7 @@ protected:
 	virtual FDebugLineItem* Duplicate() const override
 	{
 		check(false);
-		return NULL;
+		return nullptr;
 	}
 };
 
@@ -979,9 +1236,10 @@ public:
 
 	virtual TSharedRef<SWidget> GenerateWidgetForColumn( const FName& ColumnName ) override
 	{
+		TSharedPtr<SWidget> ColumnContent = nullptr;
 		if (ColumnName == KismetDebugViewConstants::ColumnId_Name)
 		{
-			return SNew(SHorizontalBox)
+			ColumnContent = SNew(SHorizontalBox)
 				+SHorizontalBox::Slot()
 				.AutoWidth()
 				[
@@ -995,13 +1253,18 @@ public:
 		}
 		else if (ColumnName == KismetDebugViewConstants::ColumnId_Value)
 		{
-			return ItemToEdit->GenerateValueWidget();
+			ColumnContent = ItemToEdit->GenerateValueWidget();
 		}
 		else
 		{
-			return SNew(STextBlock)
+			ColumnContent = SNew(STextBlock)
 				. Text( LOCTEXT("Error", "Error") );
 		}
+		
+		return SNew(SBox).Padding(FMargin(0.0f, 5.0f))
+		[
+			ColumnContent.ToSharedRef()
+		];
 	}
 
 	void Construct(const FArguments& InArgs, TSharedRef<STableViewBase> OwnerTableView, FDebugTreeItemPtr InItemToEdit)
@@ -1015,6 +1278,8 @@ public:
 //////////////////////////////////////////////////////////////////////////
 // SKismetDebuggingView
 
+TWeakObjectPtr<const UObject> SKismetDebuggingView::CurrentActiveObject = nullptr;
+
 TSharedRef<ITableRow> SKismetDebuggingView::OnGenerateRowForWatchTree(FDebugTreeItemPtr InItem, const TSharedRef<STableViewBase>& OwnerTable)
 {
 	return SNew( SDebugLineItem, OwnerTable, InItem );
@@ -1025,9 +1290,29 @@ void SKismetDebuggingView::OnGetChildrenForWatchTree(FDebugTreeItemPtr InParent,
 	InParent->GatherChildren(OutChildren);
 }
 
+TSharedRef<SHorizontalBox> SKismetDebuggingView::GetDebugLineTypeToggle(FDebugLineItem::EDebugLineType Type, const FText& Text)
+{
+	return SNew(SHorizontalBox)
+	+ SHorizontalBox::Slot()
+	.AutoWidth()
+	[
+		SNew(SCheckBox)
+			.IsChecked(true)
+			.OnCheckStateChanged_Static(&FDebugLineItem::OnDebugLineTypeActiveChanged, Type)
+	]
+	+ SHorizontalBox::Slot()
+	.AutoWidth()
+	.Padding(4.0f, 0.0f, 10.0f, 0.0f)
+	.VAlign(VAlign_Center)
+	[
+		SNew( STextBlock )
+			.Text(Text)
+	];
+}
+
 TSharedPtr<SWidget> SKismetDebuggingView::OnMakeContextMenu()
 {
-	FMenuBuilder MenuBuilder(true, NULL);
+	FMenuBuilder MenuBuilder(true, nullptr);
 
 	MenuBuilder.BeginSection("DebugActions", LOCTEXT("DebugActionsMenuHeading", "Debug Actions"));
 	{
@@ -1044,48 +1329,84 @@ TSharedPtr<SWidget> SKismetDebuggingView::OnMakeContextMenu()
 	return MenuBuilder.MakeWidget();
 }
 
-FText SKismetDebuggingView::GetTopText() const
-{
-	const bool bIsDebugging = GEditor->PlayWorld != nullptr;
-	UBlueprint* BlueprintObj = BlueprintToWatchPtr.Get();
-	UObject* DebuggedObject = (BlueprintObj != nullptr) ? BlueprintObj->GetObjectBeingDebugged() : nullptr;
 
-	TSet<UObject*> NewRootSet;
-	if (bIsDebugging && (BlueprintObj != nullptr) && (DebuggedObject != nullptr))
-	{
-		return FText::Format(LOCTEXT("ShowDebugForObject", "Showing debug info for {0}"), FText::FromName(DebuggedObject->GetFName()));
-	}
-	else if (!bIsDebugging && (BlueprintObj != nullptr))
-	{
-		return LOCTEXT("ShowDebugForBlueprint", "Showing debug info for this blueprint");
-	}
-	else
-	{
-		return LOCTEXT("ShowDebugForActors", "Showing debug info for selected actors");
-	}
+FText SKismetDebuggingView::GetTabLabel() const
+{
+	return BlueprintToWatchPtr.IsValid() ?
+		FText::FromString(BlueprintToWatchPtr->GetName()) :
+		NSLOCTEXT("BlueprintExecutionFlow", "TabTitle", "Data Flow");
 }
 
-EVisibility SKismetDebuggingView::IsDebuggerVisible() const
+FText SKismetDebuggingView::GetTopText() const
 {
-	return (GEditor->PlayWorld != NULL) ? EVisibility::Visible : EVisibility::Collapsed;
+	return LOCTEXT("ShowDebugForActors", "Showing debug info for instances of the blueprint:");
 }
 
 bool SKismetDebuggingView::CanDisableAllBreakpoints() const
 {
-	const UBlueprint* BlueprintObj = BlueprintToWatchPtr.Get();
-	return BlueprintObj && FKismetDebugUtilities::BlueprintHasBreakpoints(BlueprintObj);
+	if(BlueprintToWatchPtr.IsValid())
+	{
+		return FKismetDebugUtilities::BlueprintHasBreakpoints(BlueprintToWatchPtr.Get());
+	}
+	return false;
 }
 
 FReply SKismetDebuggingView::OnDisableAllBreakpointsClicked()
 {
-	if(UBlueprint* BlueprintObj = BlueprintToWatchPtr.Get())
+	if(BlueprintToWatchPtr.IsValid())
 	{
-		FDebuggingActionCallbacks::SetEnabledOnAllBreakpoints(BlueprintObj, false);		
+		FDebuggingActionCallbacks::SetEnabledOnAllBreakpoints(BlueprintToWatchPtr.Get(), false);
 	}
 	
 	return FReply::Handled();
 }
 
+class FBlueprintFilter : public IClassViewerFilter
+{
+public:
+	FBlueprintFilter() = default;
+
+	virtual bool IsClassAllowed(const FClassViewerInitializationOptions& InInitOptions, const UClass* InClass, TSharedRef< FClassViewerFilterFuncs > InFilterFuncs ) override
+	{
+		return InClass && !InClass->HasAnyClassFlags(CLASS_Deprecated) &&
+				InClass->HasAllClassFlags(CLASS_CompiledFromBlueprint);
+	}
+
+	virtual bool IsUnloadedClassAllowed(const FClassViewerInitializationOptions& InInitOptions, const TSharedRef< const IUnloadedBlueprintData > InUnloadedClassData, TSharedRef< FClassViewerFilterFuncs > InFilterFuncs) override
+	{
+		return !InUnloadedClassData->HasAnyClassFlags(CLASS_Deprecated) &&
+				InUnloadedClassData->HasAllClassFlags(CLASS_CompiledFromBlueprint);
+	}
+};
+
+void SKismetDebuggingView::OnBlueprintClassPicked(UClass* PickedClass)
+{
+	BlueprintToWatchPtr = Cast<UBlueprint>(PickedClass->ClassGeneratedBy);
+	BreakpointParentItem->Blueprint = BlueprintToWatchPtr;
+	DebugClassComboButton->SetIsOpen(false);
+}
+
+TSharedRef<SWidget> SKismetDebuggingView::ConstructBlueprintClassPicker()
+{
+	FClassViewerInitializationOptions Options;
+	Options.Mode = EClassViewerMode::ClassPicker;
+	Options.bShowBackgroundBorder = false;
+	Options.ClassFilters.Add(MakeShared<FBlueprintFilter>());
+	Options.bIsBlueprintBaseOnly = true;
+	Options.bShowUnloadedBlueprints = false;
+	
+	FClassViewerModule &ClassViewerModule = FModuleManager::LoadModuleChecked<FClassViewerModule>("ClassViewer");
+
+	FOnClassPicked OnClassPicked;
+	OnClassPicked.BindRaw(this, &SKismetDebuggingView::OnBlueprintClassPicked);
+	return
+		SNew(SBox)
+		.HeightOverride(500.f)
+		[
+			ClassViewerModule.CreateClassViewer(Options, OnClassPicked)
+		];
+}
+
 void SKismetDebuggingView::Construct(const FArguments& InArgs)
 {	
 	BlueprintToWatchPtr = InArgs._BlueprintToWatch;
@@ -1104,7 +1425,36 @@ void SKismetDebuggingView::Construct(const FArguments& InArgs)
 
 	FToolMenuContext MenuContext(FPlayWorldCommands::GlobalPlayWorldActions);
 	TSharedRef<SWidget> ToolbarWidget = UToolMenus::Get()->GenerateWidget(ToolbarName, MenuContext);
+	
+	DebugClassComboButton =
+		SNew(SComboButton)
+		.OnGetMenuContent_Raw(this, &SKismetDebuggingView::ConstructBlueprintClassPicker)
+		.ButtonContent()
+		[
+			SNew(STextBlock)
+			.Text_Lambda([BlueprintToWatchPtr = BlueprintToWatchPtr]()
+			{
+				return BlueprintToWatchPtr.IsValid()?
+					FText::FromString(BlueprintToWatchPtr->GetName()) :
+					LOCTEXT("SelectBlueprint", "Select Blueprint");
+			})
+		];
+
+	FBlueprintContextTracker::OnEnterScriptContext.AddLambda(
+		[](const FBlueprintContextTracker& ContextTracker, const UObject* ContextObject, const UFunction* ContextFunction)
+		{
+			CurrentActiveObject = ContextObject;
+		}
+	);
+	
+	FBlueprintContextTracker::OnExitScriptContext.AddLambda(
+		[](const FBlueprintContextTracker& ContextTracker)
+		{
+			CurrentActiveObject = nullptr;
+		}
+	);
 
+	
 	this->ChildSlot
 	[
 		SNew(SVerticalBox)
@@ -1112,7 +1462,6 @@ void SKismetDebuggingView::Construct(const FArguments& InArgs)
 		.AutoHeight()
 		[
 			SNew(SBorder)
-			.Visibility( this, &SKismetDebuggingView::IsDebuggerVisible )
 			.BorderImage( FEditorStyle::GetBrush( TEXT("NoBorder") ) )
 			[
 				ToolbarWidget
@@ -1121,42 +1470,93 @@ void SKismetDebuggingView::Construct(const FArguments& InArgs)
 		+SVerticalBox::Slot()
 		.AutoHeight()
 		[
-			SNew( SHorizontalBox )
-			+ SHorizontalBox::Slot()
-			.HAlign( HAlign_Left )
+			
+			SNew( SVerticalBox )
+			+SVerticalBox::Slot()
+			.AutoHeight()
 			[
 				SNew( STextBlock )
 					.Text( this, &SKismetDebuggingView::GetTopText )
 			]
-			+ SHorizontalBox::Slot()
-			.HAlign( HAlign_Right )
+			+ SVerticalBox::Slot()
+			.AutoHeight()
 			[
-				SNew( SButton )
-	                .IsEnabled( this, &SKismetDebuggingView::CanDisableAllBreakpoints )
-	                .Text( LOCTEXT( "DisableAllBreakPoints", "Disable All Breakpoints" ) )
-	                .OnClicked( this, &SKismetDebuggingView::OnDisableAllBreakpointsClicked )
+				SNew( SHorizontalBox )
+				+ SHorizontalBox::Slot()
+				.HAlign( HAlign_Left )
+				[
+					SNew(SBox)
+					.WidthOverride(400.f)
+					[
+						DebugClassComboButton.ToSharedRef()
+					]
+				]
+				+ SHorizontalBox::Slot()
+				.HAlign( HAlign_Right )
+				[
+					SNew( SButton )
+		                .IsEnabled( this, &SKismetDebuggingView::CanDisableAllBreakpoints )
+		                .Text( LOCTEXT( "DisableAllBreakPoints", "Disable All Breakpoints" ) )
+		                .OnClicked( this, &SKismetDebuggingView::OnDisableAllBreakpointsClicked )
+				]
 			]
 		]
 		+SVerticalBox::Slot()
 		[
-			SAssignNew( DebugTreeView, STreeView< FDebugTreeItemPtr > )
-			.TreeItemsSource( &RootTreeItems )
-			.SelectionMode( ESelectionMode::Single )
-			.OnGetChildren( this, &SKismetDebuggingView::OnGetChildrenForWatchTree )
-			.OnGenerateRow( this, &SKismetDebuggingView::OnGenerateRowForWatchTree ) 
-			.OnContextMenuOpening( this, &SKismetDebuggingView::OnMakeContextMenu )
-			. HeaderRow
-			(
-				SNew(SHeaderRow)
-				+ SHeaderRow::Column(KismetDebugViewConstants::ColumnId_Name)
-				.DefaultLabel(KismetDebugViewConstants::ColumnText_Name)
-				+ SHeaderRow::Column(KismetDebugViewConstants::ColumnId_Value)
-				.DefaultLabel(KismetDebugViewConstants::ColumnText_Value)
-			)
+			SNew(SSplitter)
+			.Orientation(Orient_Vertical)
+			+SSplitter::Slot()
+			[
+				SAssignNew( DebugTreeView, STreeView< FDebugTreeItemPtr > )
+				.TreeItemsSource( &RootTreeItems )
+				.SelectionMode( ESelectionMode::Single )
+				.OnGetChildren( this, &SKismetDebuggingView::OnGetChildrenForWatchTree )
+				.OnGenerateRow( this, &SKismetDebuggingView::OnGenerateRowForWatchTree ) 
+				.OnContextMenuOpening( this, &SKismetDebuggingView::OnMakeContextMenu )
+				. HeaderRow
+				(
+					SNew(SHeaderRow)
+					+ SHeaderRow::Column(KismetDebugViewConstants::ColumnId_Name)
+					.DefaultLabel(KismetDebugViewConstants::ColumnText_Name)
+					+ SHeaderRow::Column(KismetDebugViewConstants::ColumnId_Value)
+					.DefaultLabel(KismetDebugViewConstants::ColumnText_Value)
+				)
+			]
+			+SSplitter::Slot()
+			[
+				SAssignNew( OtherTreeView, STreeView< FDebugTreeItemPtr > )
+				.TreeItemsSource( &OtherTreeItems )
+				.SelectionMode( ESelectionMode::Single )
+				.OnGetChildren( this, &SKismetDebuggingView::OnGetChildrenForWatchTree )
+				.OnGenerateRow( this, &SKismetDebuggingView::OnGenerateRowForWatchTree ) 
+				.OnContextMenuOpening( this, &SKismetDebuggingView::OnMakeContextMenu )
+				. HeaderRow
+				(
+					SNew(SHeaderRow)
+					+ SHeaderRow::Column(KismetDebugViewConstants::ColumnId_Name)
+					.DefaultLabel(KismetDebugViewConstants::ColumnText_DebugKey)
+					+ SHeaderRow::Column(KismetDebugViewConstants::ColumnId_Value)
+					.DefaultLabel(KismetDebugViewConstants::ColumnText_Info)
+				)
+			]
+		]
+		+ SVerticalBox::Slot()
+		.AutoHeight()
+		[
+			SNew(SHorizontalBox)
+			+ SHorizontalBox::Slot().AutoWidth()
+				[ GetDebugLineTypeToggle(FDebugLineItem::DLT_Watch, LOCTEXT("Watchpoints", "Watchpoints")) ]
+			+ SHorizontalBox::Slot().AutoWidth()
+				[ GetDebugLineTypeToggle(FDebugLineItem::DLT_LatentAction, LOCTEXT("LatentActions", "Latent Actions")) ]
+			+ SHorizontalBox::Slot().AutoWidth()
+				[ GetDebugLineTypeToggle(FDebugLineItem::DLT_BreakpointParent, LOCTEXT("Breakpoints", "Breakpoints")) ]
+			+ SHorizontalBox::Slot().AutoWidth()
+				[ GetDebugLineTypeToggle(FDebugLineItem::DLT_TraceStackParent, LOCTEXT("ExecutionTrace", "Execution Trace")) ]
 		]
 	];
 
 	TraceStackItem = MakeShareable(new FTraceStackParentItem);
+	BreakpointParentItem = MakeShareable(new FBreakpointParentItem(BlueprintToWatchPtr));
 }
 
 void SKismetDebuggingView::Tick( const FGeometry& AllottedGeometry, const double InCurrentTime, const float InDeltaTime )
@@ -1173,45 +1573,62 @@ void SKismetDebuggingView::Tick( const FGeometry& AllottedGeometry, const double
 
 	// Gather what we'd like to be the new root set
 	const bool bIsDebugging = GEditor->PlayWorld != nullptr;
-	UBlueprint* BlueprintObj = BlueprintToWatchPtr.Get();
-	UObject* DebuggedObject = (BlueprintObj != nullptr) ? BlueprintObj->GetObjectBeingDebugged() : nullptr;
 
 	TSet<UObject*> NewRootSet;
-	if (bIsDebugging && (BlueprintObj != nullptr) && (DebuggedObject != nullptr))
-	{
-		// If we have a currently debugged object & we were launched from a specific Kismet window, display the debugged object
-		NewRootSet.Add(DebuggedObject);
-	}
-	else if (!bIsDebugging && (BlueprintObj != nullptr))
-	{
-		// If not debugging and summoned from a specific Kismet window, just display the currently open blueprint
-		NewRootSet.Add(BlueprintObj);
-	}
-	else
+	
+	if(bIsDebugging && BlueprintToWatchPtr.IsValid())
 	{
-		// Get the set of objects being debugged
-		USelection* SelectedActors = GEditor->GetSelectedActors();
-		for (FObjectsBeingDebuggedIterator Iter; Iter; ++Iter)
-		{
-			NewRootSet.Add(*Iter);
-		}
+		UClass* GeneratedClass = Cast<UClass>(BlueprintToWatchPtr->GeneratedClass);
+		for(FThreadSafeObjectIterator Iter(GeneratedClass); Iter; ++Iter)
+        {
+            UObject* Instance = *Iter;
+			if(!Instance)
+			{
+				continue;
+			}
+
+			// only include non temporary, non archetype objects
+            if(Instance->HasAnyFlags(RF_ArchetypeObject | RF_Transient))
+            {
+                continue;
+            }
+			
+			// only include actors in current world
+            if(AActor* Actor = Cast<AActor>(Instance))
+            {
+                if(!GEditor->PlayWorld->ContainsActor(Actor))
+                {
+                    continue;
+                }
+            }
+            
+            NewRootSet.Add(Instance);
+        }
 	}
 
 	// This will pull anything out of Old that is also New (sticking around), so afterwards Old is a list of things to remove
 	RootTreeItems.Empty();
 	for (TSet<UObject*>::TIterator NewIter(NewRootSet); NewIter; ++NewIter)
 	{
-		UObject* ObjectToAdd = *NewIter;
+		TWeakObjectPtr ObjectToAdd = *NewIter;
 
-		if (OldRootSet.Contains(ObjectToAdd))
+		// destroyed objects can still appear if they haven't ben GCed yet.
+		// weak object pointers will detect it and return nullptr
+		if(!ObjectToAdd.Get())
 		{
-			OldRootSet.Remove(ObjectToAdd);
-			RootTreeItems.Add(ObjectToTreeItemMap.FindChecked(ObjectToAdd));
+			continue;
+		}
+
+		if (OldRootSet.Contains(ObjectToAdd.Get()))
+		{
+			OldRootSet.Remove(ObjectToAdd.Get());
+			RootTreeItems.Add(ObjectToTreeItemMap.FindChecked(ObjectToAdd.Get()));
 		}
 		else
 		{
-			FDebugTreeItemPtr NewPtr = FDebugTreeItemPtr(new FParentLineItem(ObjectToAdd));
-			ObjectToTreeItemMap.Add(ObjectToAdd, NewPtr);
+			
+			FDebugTreeItemPtr NewPtr = FDebugTreeItemPtr(new FParentLineItem(ObjectToAdd.Get()));
+			ObjectToTreeItemMap.Add(ObjectToAdd.Get(), NewPtr);
 			RootTreeItems.Add(NewPtr);
 
 			// Autoexpand newly selected items
@@ -1226,21 +1643,30 @@ void SKismetDebuggingView::Tick( const FGeometry& AllottedGeometry, const double
 		ObjectToTreeItemMap.Remove(ObjectToRemove);
 	}
 
-	// Add a message if there are no actors selected
+	// Add a message if there are no active instances of DebugClass
 	if (RootTreeItems.Num() == 0)
 	{
 		RootTreeItems.Add(MakeShareable(new FMessageLineItem(
-			bIsDebugging ? LOCTEXT("NoActorsSelected", "No actors selected").ToString() : LOCTEXT("NoPIEorSIE", "Not running PIE or SIE").ToString())));
+			bIsDebugging ? LOCTEXT("NoInstances", "No instances of this blueprint in existence").ToString() : LOCTEXT("NoPIEorSIE", "run PIE or SIE to see instance debug info").ToString())));
 	}
 
-	// Show the trace stack when debugging
-	if (bIsDebugging)
+	// Refresh the list
+	DebugTreeView->RequestTreeRefresh();
+
+	OtherTreeItems.Empty();
+
+	// Show Breakpoints
+	if(FDebugLineItem::IsDebugLineTypeActive(FDebugLineItem::DLT_BreakpointParent))
 	{
-		RootTreeItems.Add(TraceStackItem);
+		OtherTreeItems.Add(BreakpointParentItem);
 	}
 
-	// Refresh the list
-	DebugTreeView->RequestTreeRefresh();
+	// Show the trace stack when debugging
+	if (bIsDebugging && FDebugLineItem::IsDebugLineTypeActive(FDebugLineItem::DLT_TraceStackParent))
+	{
+		OtherTreeItems.Add(TraceStackItem);
+	}
+	OtherTreeView->RequestTreeRefresh();
 }
 
 
diff --git a/Engine/Source/Editor/Kismet/Private/Debugging/SKismetDebuggingView.h b/Engine/Source/Editor/Kismet/Private/Debugging/SKismetDebuggingView.h
--- a/Engine/Source/Editor/Kismet/Private/Debugging/SKismetDebuggingView.h
+++ b/Engine/Source/Editor/Kismet/Private/Debugging/SKismetDebuggingView.h
@@ -3,7 +3,6 @@
 #pragma once
 
 #include "CoreMinimal.h"
-#include "Layout/Visibility.h"
 #include "Widgets/SWidget.h"
 #include "Widgets/DeclarativeSyntaxSupport.h"
 #include "Widgets/SCompoundWidget.h"
@@ -15,6 +14,7 @@ class FDebugLineItem;
 class FTraceStackParentItem;
 class UBlueprint;
 class UBlueprintGeneratedClass;
+class FBreakpointParentItem;
 
 //////////////////////////////////////////////////////////////////////////
 // FDebugLineItem
@@ -26,6 +26,7 @@ typedef TSharedPtr<class FDebugLineItem> FDebugTreeItemPtr;
 class FDebugLineItem : public TSharedFromThis<FDebugLineItem>
 {
 public:
+	friend class ILineItemWithChildren;
 	enum EDebugLineType
 	{
 		DLT_Message,
@@ -33,8 +34,10 @@ public:
 		DLT_TraceStackChild,
 		DLT_Parent,
 		DLT_Watch,
+		DLT_WatchChild,
 		DLT_LatentAction,
-		DLT_Breakpoint
+		DLT_Breakpoint,
+		DLT_BreakpointParent
 	};
 
 	virtual ~FDebugLineItem() {}
@@ -54,6 +57,11 @@ public:
 	// @return The object that will act as a parent to more items in the tree, or NULL if this is a leaf node
 	virtual UObject* GetParentObject() { return NULL; }
 
+	virtual EDebugLineType GetType() const
+	{
+		return Type;
+	}
+
 	// Helper function to try to get the blueprint for a given object;
 	//   Returns the blueprint that was used to create the instance if there was one
 	//   Returns the object itself if it is already a blueprint
@@ -61,10 +69,11 @@ public:
 	static UBlueprint* GetBlueprintForObject(UObject* ParentObject);
 
 	static UBlueprintGeneratedClass* GetClassForObject(UObject* ParentObject);
-protected:
-	// Adds either Item or an identical node that was previously created (present in ChildrenMirrors) as a child to OutChildren
-	void EnsureChildIsAdded(TArray<FDebugTreeItemPtr>& ChildrenMirrors, TArray<FDebugTreeItemPtr>& OutChildren, const FDebugLineItem& Item);
 
+	static bool IsDebugLineTypeActive(EDebugLineType Type);
+	static void OnDebugLineTypeActiveChanged(ECheckBoxState CheckState, EDebugLineType Type);
+protected:
+	
 	// Cannot create an instance of this class, it's just for use as a base class
 	FDebugLineItem(EDebugLineType InType)
 		: Type(InType)
@@ -85,6 +94,8 @@ protected:
 private:
 	// Type of action (poor mans RTTI for the tree, really only used to accelerate Compare checks)
 	EDebugLineType Type;
+
+	static uint16 ActiveTypeBitset;
 };
 
 //////////////////////////////////////////////////////////////////////////
@@ -105,24 +116,38 @@ public:
 	// SWidget interface
 	virtual void Tick( const FGeometry& AllottedGeometry, const double InCurrentTime, const float InDeltaTime ) override;
 	// End of SWidget interface
+
+	/* set to an object that's paused at a breakpoint and null otherwise */
+	static TWeakObjectPtr<const UObject> CurrentActiveObject;
+	
+	FText GetTabLabel() const;
 protected:
 	FText GetTopText() const;
 	bool CanDisableAllBreakpoints() const;
 	FReply OnDisableAllBreakpointsClicked();
 
+	void OnBlueprintClassPicked(UClass* PickedClass);
+	TSharedRef<SWidget> ConstructBlueprintClassPicker();
+
 	TSharedRef<ITableRow> OnGenerateRowForWatchTree(FDebugTreeItemPtr InItem, const TSharedRef<STableViewBase>& OwnerTable);
 	void OnGetChildrenForWatchTree(FDebugTreeItemPtr InParent, TArray<FDebugTreeItemPtr>& OutChildren);
-	EVisibility IsDebuggerVisible() const;
 
+	static TSharedRef<SHorizontalBox> GetDebugLineTypeToggle(FDebugLineItem::EDebugLineType Type, const FText& Text);
+	
 	TSharedPtr<SWidget> OnMakeContextMenu();
 protected:
 	TSharedPtr< STreeView<FDebugTreeItemPtr> > DebugTreeView;
 	TMap<UObject*, FDebugTreeItemPtr> ObjectToTreeItemMap;
 	TArray<FDebugTreeItemPtr> RootTreeItems;
 
-	// The trace stack parent item
-	TSharedPtr< class FTraceStackParentItem > TraceStackItem;
+	// includes items such as breakpoints and Exectution trace
+	TSharedPtr< STreeView<FDebugTreeItemPtr> > OtherTreeView;
+	TArray<FDebugTreeItemPtr> OtherTreeItems;
+
+	// UI tree entries for stack trace and breakpoints
+	TSharedPtr< FTraceStackParentItem > TraceStackItem;
+	TSharedPtr< FBreakpointParentItem > BreakpointParentItem;
 
-	/** Pointer to the blueprint to observe when not in PIE/SIE; can be NULL for a free-floating watch window */
+	TSharedPtr<SComboButton> DebugClassComboButton;
 	TWeakObjectPtr<UBlueprint> BlueprintToWatchPtr;
 };
diff --git a/Engine/Source/Editor/Kismet/Private/WatchPointViewer.cpp b/Engine/Source/Editor/Kismet/Private/WatchPointViewer.cpp
--- a/Engine/Source/Editor/Kismet/Private/WatchPointViewer.cpp
+++ b/Engine/Source/Editor/Kismet/Private/WatchPointViewer.cpp
@@ -18,6 +18,7 @@
 #include "AssetRegistryModule.h"
 
 #include "Editor.h"
+#include "ToolMenus.h"
 
 #include "Widgets/Docking/SDockTab.h"
 #include "Widgets/Images/SImage.h"
@@ -29,6 +30,7 @@
 
 #include "WorkspaceMenuStructure.h"
 #include "WorkspaceMenuStructureModule.h"
+#include "Kismet2/DebuggerCommands.h"
 
 #include "KismetNodes/KismetNodeInfoContext.h"
 #include "Stats/Stats.h"
@@ -255,7 +257,7 @@ namespace
 
 		if (IsPaused())
 		{
-			WatchViewer::UpdateInstancedWatchDisplay();
+			DEPRICATED_WatchViewer::UpdateInstancedWatchDisplay();
 		}
 
 		// Notify subscribers:
@@ -581,29 +583,29 @@ void SWatchViewer::HandleGetChildren(TSharedRef<FWatchRow> InWatchRow, TArray<TS
 
 void SWatchViewer::HandleWatchedPinsChanged(UBlueprint* BlueprintObj)
 {
-	WatchViewer::UpdateWatchListFromBlueprint(BlueprintObj);
+	DEPRICATED_WatchViewer::UpdateWatchListFromBlueprint(BlueprintObj);
 }
 
 void SWatchViewer::HandleResumePIE(bool)
 {
 	// swap to displaying the unpaused watches
-	WatchViewer::ContinueExecution();
+	DEPRICATED_WatchViewer::ContinueExecution();
 }
 
 void SWatchViewer::HandleEndPIE(bool)
 {
 	// show the unpaused watches in case we stopped PIE while at a breakpoint
-	WatchViewer::ContinueExecution();
+	DEPRICATED_WatchViewer::ContinueExecution();
 }
 
 void SWatchViewer::HandleAssetRemoved(const FAssetData& InAssetData)
 {
-	WatchViewer::RemoveWatchesForAsset(InAssetData);
+	DEPRICATED_WatchViewer::RemoveWatchesForAsset(InAssetData);
 }
 
 void SWatchViewer::HandleAssetRenamed(const FAssetData& InAssetData, const FString& InOldName)
 {
-	WatchViewer::OnRenameAsset(InAssetData, InOldName);
+	DEPRICATED_WatchViewer::OnRenameAsset(InAssetData, InOldName);
 }
 
 void SWatchViewer::UpdateWatches(TArray<TSharedRef<FWatchRow>>* Watches)
@@ -761,7 +763,7 @@ TSharedRef<SWidget> SWatchTreeWidgetItem::GenerateWidgetForColumn(const FName& C
 	}
 }
 
-void WatchViewer::UpdateInstancedWatchDisplay()
+void DEPRICATED_WatchViewer::UpdateInstancedWatchDisplay()
 {
 #if DO_BLUEPRINT_GUARD
 	{
@@ -885,19 +887,19 @@ void WatchViewer::UpdateInstancedWatchDisplay()
 #endif
 }
 
-void WatchViewer::ContinueExecution()
+void DEPRICATED_WatchViewer::ContinueExecution()
 {
 	// Notify subscribers:
 	WatchListSubscribers.Broadcast(&Private_WatchSource);
 }
 
-FName WatchViewer::GetTabName()
+FName DEPRICATED_WatchViewer::GetTabName()
 {
 	const FName TabName = TEXT("WatchViewer");
 	return TabName;
 }
 
-void WatchViewer::RemoveWatchesForBlueprint(TWeakObjectPtr<UBlueprint> BlueprintObj)
+void DEPRICATED_WatchViewer::RemoveWatchesForBlueprint(TWeakObjectPtr<UBlueprint> BlueprintObj)
 {
 	if (!ensure(BlueprintObj.IsValid()))
 	{
@@ -918,14 +920,14 @@ void WatchViewer::RemoveWatchesForBlueprint(TWeakObjectPtr<UBlueprint> Blueprint
 
 	if (IsPaused())
 	{
-		WatchViewer::UpdateInstancedWatchDisplay();
+		DEPRICATED_WatchViewer::UpdateInstancedWatchDisplay();
 	}
 
 	// Notify subscribers
 	WatchListSubscribers.Broadcast(&Private_WatchSource);
 }
 
-void WatchViewer::RemoveWatchesForAsset(const struct FAssetData& AssetData)
+void DEPRICATED_WatchViewer::RemoveWatchesForAsset(const struct FAssetData& AssetData)
 {
 	for (TSharedRef<FWatchRow> WatchRow : Private_WatchSource)
 	{
@@ -937,7 +939,7 @@ void WatchViewer::RemoveWatchesForAsset(const struct FAssetData& AssetData)
 	}
 }
 
-void WatchViewer::OnRenameAsset(const struct FAssetData& AssetData, const FString& OldAssetName)
+void DEPRICATED_WatchViewer::OnRenameAsset(const struct FAssetData& AssetData, const FString& OldAssetName)
 {
 	FString OldPackageName;
 	FString OldBPName;
@@ -962,7 +964,7 @@ void WatchViewer::OnRenameAsset(const struct FAssetData& AssetData, const FStrin
 
 			if (IsPaused())
 			{
-				WatchViewer::UpdateInstancedWatchDisplay();
+				DEPRICATED_WatchViewer::UpdateInstancedWatchDisplay();
 			}
 
 			// Notify subscribers if necessary
@@ -971,17 +973,17 @@ void WatchViewer::OnRenameAsset(const struct FAssetData& AssetData, const FStrin
 	}
 }
 
-void WatchViewer::UpdateWatchListFromBlueprint(TWeakObjectPtr<UBlueprint> BlueprintObj)
+void DEPRICATED_WatchViewer::UpdateWatchListFromBlueprint(TWeakObjectPtr<UBlueprint> BlueprintObj)
 {
 	UpdateWatchListFromBlueprintImpl(BlueprintObj, true);
 }
 
-void WatchViewer::ClearWatchListFromBlueprint(TWeakObjectPtr<UBlueprint> BlueprintObj)
+void DEPRICATED_WatchViewer::ClearWatchListFromBlueprint(TWeakObjectPtr<UBlueprint> BlueprintObj)
 {
 	UpdateWatchListFromBlueprintImpl(BlueprintObj, false);
 }
 
-void WatchViewer::RegisterTabSpawner(FTabManager& TabManager)
+void DEPRICATED_WatchViewer::RegisterTabSpawner(FTabManager& TabManager)
 {
 	const auto SpawnWatchViewTab = []( const FSpawnTabArgs& Args )
 	{
@@ -990,20 +992,38 @@ void WatchViewer::RegisterTabSpawner(FTabManager& TabManager)
 		{
 			Source = &Private_InstanceWatchSource;
 		}
+		
+		static const FName ToolbarName = TEXT("Kismet.DebuggingViewToolBar");
+		FToolMenuContext MenuContext(FPlayWorldCommands::GlobalPlayWorldActions);
+		TSharedRef<SWidget> ToolbarWidget = UToolMenus::Get()->GenerateWidget(ToolbarName, MenuContext);
 
 		return SNew(SDockTab)
 			.TabRole( ETabRole::PanelTab )
 			.Label( LOCTEXT("TabTitle", "Watches") )
 			[
-				SNew(SBorder)
-				.BorderImage( FEditorStyle::GetBrush("Docking.Tab.ContentAreaBrush") )
+				SNew(SVerticalBox)
+				+ SVerticalBox::Slot()
+				.AutoHeight()
 				[
-					SNew(SWatchViewer, Source)
+					SNew(SBorder)
+					.BorderImage( FEditorStyle::GetBrush( TEXT("NoBorder") ) )
+					[
+						ToolbarWidget
+					]
+				]
+				+ SVerticalBox::Slot()
+				.AutoHeight()
+				[
+					SNew(SBorder)
+					.BorderImage( FEditorStyle::GetBrush("Docking.Tab.ContentAreaBrush") )
+					[
+						SNew(SWatchViewer, Source)
+					]
 				]
 			];
 	};
 	
-	TabManager.RegisterTabSpawner( WatchViewer::GetTabName(), FOnSpawnTab::CreateStatic(SpawnWatchViewTab) )
+	TabManager.RegisterTabSpawner( DEPRICATED_WatchViewer::GetTabName(), FOnSpawnTab::CreateStatic(SpawnWatchViewTab) )
 		.SetDisplayName( LOCTEXT("SpawnerTitle", "Watch Window") )
 		.SetTooltipText( LOCTEXT("SpawnerTooltipText", "Open the watch window tab") );
 }
diff --git a/Engine/Source/Editor/Kismet/Public/WatchPointViewer.h b/Engine/Source/Editor/Kismet/Public/WatchPointViewer.h
--- a/Engine/Source/Editor/Kismet/Public/WatchPointViewer.h
+++ b/Engine/Source/Editor/Kismet/Public/WatchPointViewer.h
@@ -2,13 +2,11 @@
 
 #pragma once
 
-#include "Widgets/SCompoundWidget.h"
-#include "WorkflowOrientedApp/WorkflowTabFactory.h"
-
 struct FAssetData;
 class UBlueprint;
 
-namespace WatchViewer
+UE_DEPRECATED(5.0, "WatchViewer has been deprecated, use SKismetDebuggingView instead");
+namespace DEPRICATED_WatchViewer
 {
 	// updates the instanced watch values, these are only valid while execution is paused
 	void KISMET_API UpdateInstancedWatchDisplay();
diff --git a/Engine/Source/Editor/UnrealEd/Private/Kismet2/KismetDebugUtilities.cpp b/Engine/Source/Editor/UnrealEd/Private/Kismet2/KismetDebugUtilities.cpp
--- a/Engine/Source/Editor/UnrealEd/Private/Kismet2/KismetDebugUtilities.cpp
+++ b/Engine/Source/Editor/UnrealEd/Private/Kismet2/KismetDebugUtilities.cpp
@@ -733,7 +733,6 @@ void FKismetDebugUtilities::AttemptToBreakExecution(UBlueprint* BlueprintObj, co
 		Data.LastExceptionMessage = Info.GetDescription();
 		FKismetEditorUtilities::BringKismetToFocusAttentionOnObject(NodeStoppedAt);
 		CallStackViewer::UpdateDisplayedCallstack(ScriptStack);
-		WatchViewer::UpdateInstancedWatchDisplay();
 		FSlateApplication::Get().EnterDebuggingMode();
 	}
 }
diff --git a/Engine/Source/Runtime/Engine/Private/Blueprint.cpp b/Engine/Source/Runtime/Engine/Private/Blueprint.cpp
--- a/Engine/Source/Runtime/Engine/Private/Blueprint.cpp
+++ b/Engine/Source/Runtime/Engine/Private/Blueprint.cpp
@@ -1388,7 +1388,6 @@ void UBlueprint::BeginDestroy()
 	Super::BeginDestroy();
 
 	FBlueprintEditorUtils::RemoveAllLocalBookmarks(this);
-	WatchViewer::ClearWatchListFromBlueprint(this);
 }
 
 #endif // WITH_EDITOR
