attempted fix to deprecation build errors

diff --git a/Engine/Source/Editor/Kismet/Public/WatchPointViewer.h b/Engine/Source/Editor/Kismet/Public/WatchPointViewer.h
--- a/Engine/Source/Editor/Kismet/Public/WatchPointViewer.h
+++ b/Engine/Source/Editor/Kismet/Public/WatchPointViewer.h
@@ -5,8 +5,8 @@
 struct FAssetData;
 class UBlueprint;
 
-UE_DEPRECATED(5.0, "WatchViewer has been deprecated, use SKismetDebuggingView instead");
-namespace DEPRICATED_WatchViewer
+
+namespace UE_DEPRECATED(5.0, "WatchViewer has been deprecated, use SKismetDebuggingView instead") DEPRECATED_WatchViewer
 {
 	// updates the instanced watch values, these are only valid while execution is paused
 	void KISMET_API UpdateInstancedWatchDisplay();
