- Fix compile error

#rnx

diff --git a/Engine/Source/Editor/Kismet/Private/WatchPointViewer.cpp b/Engine/Source/Editor/Kismet/Private/WatchPointViewer.cpp
--- a/Engine/Source/Editor/Kismet/Private/WatchPointViewer.cpp
+++ b/Engine/Source/Editor/Kismet/Private/WatchPointViewer.cpp
@@ -257,7 +257,7 @@ namespace
 
 		if (IsPaused())
 		{
-			DEPRICATED_WatchViewer::UpdateInstancedWatchDisplay();
+			DEPRECATED_WatchViewer::UpdateInstancedWatchDisplay();
 		}
 
 		// Notify subscribers:
@@ -583,29 +583,29 @@ void SWatchViewer::HandleGetChildren(TSharedRef<FWatchRow> InWatchRow, TArray<TS
 
 void SWatchViewer::HandleWatchedPinsChanged(UBlueprint* BlueprintObj)
 {
-	DEPRICATED_WatchViewer::UpdateWatchListFromBlueprint(BlueprintObj);
+	DEPRECATED_WatchViewer::UpdateWatchListFromBlueprint(BlueprintObj);
 }
 
 void SWatchViewer::HandleResumePIE(bool)
 {
 	// swap to displaying the unpaused watches
-	DEPRICATED_WatchViewer::ContinueExecution();
+	DEPRECATED_WatchViewer::ContinueExecution();
 }
 
 void SWatchViewer::HandleEndPIE(bool)
 {
 	// show the unpaused watches in case we stopped PIE while at a breakpoint
-	DEPRICATED_WatchViewer::ContinueExecution();
+	DEPRECATED_WatchViewer::ContinueExecution();
 }
 
 void SWatchViewer::HandleAssetRemoved(const FAssetData& InAssetData)
 {
-	DEPRICATED_WatchViewer::RemoveWatchesForAsset(InAssetData);
+	DEPRECATED_WatchViewer::RemoveWatchesForAsset(InAssetData);
 }
 
 void SWatchViewer::HandleAssetRenamed(const FAssetData& InAssetData, const FString& InOldName)
 {
-	DEPRICATED_WatchViewer::OnRenameAsset(InAssetData, InOldName);
+	DEPRECATED_WatchViewer::OnRenameAsset(InAssetData, InOldName);
 }
 
 void SWatchViewer::UpdateWatches(TArray<TSharedRef<FWatchRow>>* Watches)
@@ -763,7 +763,7 @@ TSharedRef<SWidget> SWatchTreeWidgetItem::GenerateWidgetForColumn(const FName& C
 	}
 }
 
-void DEPRICATED_WatchViewer::UpdateInstancedWatchDisplay()
+void DEPRECATED_WatchViewer::UpdateInstancedWatchDisplay()
 {
 #if DO_BLUEPRINT_GUARD
 	{
@@ -887,19 +887,19 @@ void DEPRICATED_WatchViewer::UpdateInstancedWatchDisplay()
 #endif
 }
 
-void DEPRICATED_WatchViewer::ContinueExecution()
+void DEPRECATED_WatchViewer::ContinueExecution()
 {
 	// Notify subscribers:
 	WatchListSubscribers.Broadcast(&Private_WatchSource);
 }
 
-FName DEPRICATED_WatchViewer::GetTabName()
+FName DEPRECATED_WatchViewer::GetTabName()
 {
 	const FName TabName = TEXT("WatchViewer");
 	return TabName;
 }
 
-void DEPRICATED_WatchViewer::RemoveWatchesForBlueprint(TWeakObjectPtr<UBlueprint> BlueprintObj)
+void DEPRECATED_WatchViewer::RemoveWatchesForBlueprint(TWeakObjectPtr<UBlueprint> BlueprintObj)
 {
 	if (!ensure(BlueprintObj.IsValid()))
 	{
@@ -920,14 +920,14 @@ void DEPRICATED_WatchViewer::RemoveWatchesForBlueprint(TWeakObjectPtr<UBlueprint
 
 	if (IsPaused())
 	{
-		DEPRICATED_WatchViewer::UpdateInstancedWatchDisplay();
+		DEPRECATED_WatchViewer::UpdateInstancedWatchDisplay();
 	}
 
 	// Notify subscribers
 	WatchListSubscribers.Broadcast(&Private_WatchSource);
 }
 
-void DEPRICATED_WatchViewer::RemoveWatchesForAsset(const struct FAssetData& AssetData)
+void DEPRECATED_WatchViewer::RemoveWatchesForAsset(const struct FAssetData& AssetData)
 {
 	for (TSharedRef<FWatchRow> WatchRow : Private_WatchSource)
 	{
@@ -939,7 +939,7 @@ void DEPRICATED_WatchViewer::RemoveWatchesForAsset(const struct FAssetData& Asse
 	}
 }
 
-void DEPRICATED_WatchViewer::OnRenameAsset(const struct FAssetData& AssetData, const FString& OldAssetName)
+void DEPRECATED_WatchViewer::OnRenameAsset(const struct FAssetData& AssetData, const FString& OldAssetName)
 {
 	FString OldPackageName;
 	FString OldBPName;
@@ -964,7 +964,7 @@ void DEPRICATED_WatchViewer::OnRenameAsset(const struct FAssetData& AssetData, c
 
 			if (IsPaused())
 			{
-				DEPRICATED_WatchViewer::UpdateInstancedWatchDisplay();
+				DEPRECATED_WatchViewer::UpdateInstancedWatchDisplay();
 			}
 
 			// Notify subscribers if necessary
@@ -973,17 +973,17 @@ void DEPRICATED_WatchViewer::OnRenameAsset(const struct FAssetData& AssetData, c
 	}
 }
 
-void DEPRICATED_WatchViewer::UpdateWatchListFromBlueprint(TWeakObjectPtr<UBlueprint> BlueprintObj)
+void DEPRECATED_WatchViewer::UpdateWatchListFromBlueprint(TWeakObjectPtr<UBlueprint> BlueprintObj)
 {
 	UpdateWatchListFromBlueprintImpl(BlueprintObj, true);
 }
 
-void DEPRICATED_WatchViewer::ClearWatchListFromBlueprint(TWeakObjectPtr<UBlueprint> BlueprintObj)
+void DEPRECATED_WatchViewer::ClearWatchListFromBlueprint(TWeakObjectPtr<UBlueprint> BlueprintObj)
 {
 	UpdateWatchListFromBlueprintImpl(BlueprintObj, false);
 }
 
-void DEPRICATED_WatchViewer::RegisterTabSpawner(FTabManager& TabManager)
+void DEPRECATED_WatchViewer::RegisterTabSpawner(FTabManager& TabManager)
 {
 	const auto SpawnWatchViewTab = []( const FSpawnTabArgs& Args )
 	{
@@ -1023,7 +1023,7 @@ void DEPRICATED_WatchViewer::RegisterTabSpawner(FTabManager& TabManager)
 			];
 	};
 	
-	TabManager.RegisterTabSpawner( DEPRICATED_WatchViewer::GetTabName(), FOnSpawnTab::CreateStatic(SpawnWatchViewTab) )
+	TabManager.RegisterTabSpawner( DEPRECATED_WatchViewer::GetTabName(), FOnSpawnTab::CreateStatic(SpawnWatchViewTab) )
 		.SetDisplayName( LOCTEXT("SpawnerTitle", "Watch Window") )
 		.SetTooltipText( LOCTEXT("SpawnerTooltipText", "Open the watch window tab") );
 }
