fix namespace deprecation warnings

diff --git a/Engine/Source/Editor/Kismet/Private/WatchPointViewer.cpp b/Engine/Source/Editor/Kismet/Private/WatchPointViewer.cpp
--- a/Engine/Source/Editor/Kismet/Private/WatchPointViewer.cpp
+++ b/Engine/Source/Editor/Kismet/Private/WatchPointViewer.cpp
@@ -257,7 +257,9 @@ namespace
 
 		if (IsPaused())
 		{
-			DEPRECATED_WatchViewer::UpdateInstancedWatchDisplay();
+#ifndef WATCH_VIEWER_DEPRECATED
+			WatchViewer::UpdateInstancedWatchDisplay();
+#endif
 		}
 
 		// Notify subscribers:
@@ -583,29 +585,39 @@ void SWatchViewer::HandleGetChildren(TSharedRef<FWatchRow> InWatchRow, TArray<TS
 
 void SWatchViewer::HandleWatchedPinsChanged(UBlueprint* BlueprintObj)
 {
-	DEPRECATED_WatchViewer::UpdateWatchListFromBlueprint(BlueprintObj);
+#ifndef WATCH_VIEWER_DEPRECATED
+	WatchViewer::UpdateWatchListFromBlueprint(BlueprintObj);
+#endif
 }
 
 void SWatchViewer::HandleResumePIE(bool)
 {
+#ifndef WATCH_VIEWER_DEPRECATED
 	// swap to displaying the unpaused watches
-	DEPRECATED_WatchViewer::ContinueExecution();
+	WatchViewer::ContinueExecution();
+#endif
 }
 
 void SWatchViewer::HandleEndPIE(bool)
 {
+#ifndef WATCH_VIEWER_DEPRECATED
 	// show the unpaused watches in case we stopped PIE while at a breakpoint
-	DEPRECATED_WatchViewer::ContinueExecution();
+	WatchViewer::ContinueExecution();
+#endif
 }
 
 void SWatchViewer::HandleAssetRemoved(const FAssetData& InAssetData)
 {
-	DEPRECATED_WatchViewer::RemoveWatchesForAsset(InAssetData);
+#ifndef WATCH_VIEWER_DEPRECATED
+	WatchViewer::RemoveWatchesForAsset(InAssetData);
+#endif
 }
 
 void SWatchViewer::HandleAssetRenamed(const FAssetData& InAssetData, const FString& InOldName)
 {
-	DEPRECATED_WatchViewer::OnRenameAsset(InAssetData, InOldName);
+#ifndef WATCH_VIEWER_DEPRECATED
+	WatchViewer::OnRenameAsset(InAssetData, InOldName);
+#endif
 }
 
 void SWatchViewer::UpdateWatches(TArray<TSharedRef<FWatchRow>>* Watches)
@@ -763,7 +775,7 @@ TSharedRef<SWidget> SWatchTreeWidgetItem::GenerateWidgetForColumn(const FName& C
 	}
 }
 
-void DEPRECATED_WatchViewer::UpdateInstancedWatchDisplay()
+void WatchViewer::UpdateInstancedWatchDisplay()
 {
 #if DO_BLUEPRINT_GUARD
 	{
@@ -887,19 +899,19 @@ void DEPRECATED_WatchViewer::UpdateInstancedWatchDisplay()
 #endif
 }
 
-void DEPRECATED_WatchViewer::ContinueExecution()
+void WatchViewer::ContinueExecution()
 {
 	// Notify subscribers:
 	WatchListSubscribers.Broadcast(&Private_WatchSource);
 }
 
-FName DEPRECATED_WatchViewer::GetTabName()
+FName WatchViewer::GetTabName()
 {
 	const FName TabName = TEXT("WatchViewer");
 	return TabName;
 }
 
-void DEPRECATED_WatchViewer::RemoveWatchesForBlueprint(TWeakObjectPtr<UBlueprint> BlueprintObj)
+void WatchViewer::RemoveWatchesForBlueprint(TWeakObjectPtr<UBlueprint> BlueprintObj)
 {
 	if (!ensure(BlueprintObj.IsValid()))
 	{
@@ -920,14 +932,14 @@ void DEPRECATED_WatchViewer::RemoveWatchesForBlueprint(TWeakObjectPtr<UBlueprint
 
 	if (IsPaused())
 	{
-		DEPRECATED_WatchViewer::UpdateInstancedWatchDisplay();
+		WatchViewer::UpdateInstancedWatchDisplay();
 	}
 
 	// Notify subscribers
 	WatchListSubscribers.Broadcast(&Private_WatchSource);
 }
 
-void DEPRECATED_WatchViewer::RemoveWatchesForAsset(const struct FAssetData& AssetData)
+void WatchViewer::RemoveWatchesForAsset(const struct FAssetData& AssetData)
 {
 	for (TSharedRef<FWatchRow> WatchRow : Private_WatchSource)
 	{
@@ -939,7 +951,7 @@ void DEPRECATED_WatchViewer::RemoveWatchesForAsset(const struct FAssetData& Asse
 	}
 }
 
-void DEPRECATED_WatchViewer::OnRenameAsset(const struct FAssetData& AssetData, const FString& OldAssetName)
+void WatchViewer::OnRenameAsset(const struct FAssetData& AssetData, const FString& OldAssetName)
 {
 	FString OldPackageName;
 	FString OldBPName;
@@ -964,7 +976,7 @@ void DEPRECATED_WatchViewer::OnRenameAsset(const struct FAssetData& AssetData, c
 
 			if (IsPaused())
 			{
-				DEPRECATED_WatchViewer::UpdateInstancedWatchDisplay();
+				WatchViewer::UpdateInstancedWatchDisplay();
 			}
 
 			// Notify subscribers if necessary
@@ -973,17 +985,17 @@ void DEPRECATED_WatchViewer::OnRenameAsset(const struct FAssetData& AssetData, c
 	}
 }
 
-void DEPRECATED_WatchViewer::UpdateWatchListFromBlueprint(TWeakObjectPtr<UBlueprint> BlueprintObj)
+void WatchViewer::UpdateWatchListFromBlueprint(TWeakObjectPtr<UBlueprint> BlueprintObj)
 {
 	UpdateWatchListFromBlueprintImpl(BlueprintObj, true);
 }
 
-void DEPRECATED_WatchViewer::ClearWatchListFromBlueprint(TWeakObjectPtr<UBlueprint> BlueprintObj)
+void WatchViewer::ClearWatchListFromBlueprint(TWeakObjectPtr<UBlueprint> BlueprintObj)
 {
 	UpdateWatchListFromBlueprintImpl(BlueprintObj, false);
 }
 
-void DEPRECATED_WatchViewer::RegisterTabSpawner(FTabManager& TabManager)
+void WatchViewer::RegisterTabSpawner(FTabManager& TabManager)
 {
 	const auto SpawnWatchViewTab = []( const FSpawnTabArgs& Args )
 	{
@@ -1023,9 +1035,9 @@ void DEPRECATED_WatchViewer::RegisterTabSpawner(FTabManager& TabManager)
 			];
 	};
 	
-	TabManager.RegisterTabSpawner( DEPRECATED_WatchViewer::GetTabName(), FOnSpawnTab::CreateStatic(SpawnWatchViewTab) )
+	TabManager.RegisterTabSpawner( WatchViewer::GetTabName(), FOnSpawnTab::CreateStatic(SpawnWatchViewTab) )
 		.SetDisplayName( LOCTEXT("SpawnerTitle", "Watch Window") )
 		.SetTooltipText( LOCTEXT("SpawnerTooltipText", "Open the watch window tab") );
 }
 
-#undef LOCTEXT_NAMESPACE
+#undef LOCTEXT_NAMESPACE
\ No newline at end of file
diff --git a/Engine/Source/Editor/Kismet/Public/WatchPointViewer.h b/Engine/Source/Editor/Kismet/Public/WatchPointViewer.h
--- a/Engine/Source/Editor/Kismet/Public/WatchPointViewer.h
+++ b/Engine/Source/Editor/Kismet/Public/WatchPointViewer.h
@@ -5,8 +5,13 @@
 struct FAssetData;
 class UBlueprint;
 
+#define WATCH_VIEWER_DEPRECATED
 
-namespace UE_DEPRECATED(5.0, "WatchViewer has been deprecated, use SKismetDebuggingView instead") DEPRECATED_WatchViewer
+#ifdef WATCH_VIEWER_DEPRECATED
+namespace UE_DEPRECATED(5.0, "WatchViewer has been deprecated, use SKismetDebuggingView instead") WatchViewer
+#else
+namespace WatchViewer
+#endif
 {
 	// updates the instanced watch values, these are only valid while execution is paused
 	void KISMET_API UpdateInstancedWatchDisplay();
