KismetDebuggingView Polish

- made it so 'self' is watched automatically which includes all BlueprintVisible properties

- icons on every line for watch datatypes
- polished the look of the tree making it very similar to the details panel

diff --git a/Engine/Source/Editor/Kismet/Private/Debugging/SKismetDebuggingView.cpp b/Engine/Source/Editor/Kismet/Private/Debugging/SKismetDebuggingView.cpp
--- a/Engine/Source/Editor/Kismet/Private/Debugging/SKismetDebuggingView.cpp
+++ b/Engine/Source/Editor/Kismet/Private/Debugging/SKismetDebuggingView.cpp
@@ -3,6 +3,7 @@
 
 #include "Debugging/SKismetDebuggingView.h"
 
+#include "BlueprintEditor.h"
 #include "ClassViewerFilter.h"
 #include "ClassViewerModule.h"
 #include "EdGraph/EdGraphPin.h"
@@ -19,6 +20,7 @@
 #include "EngineGlobals.h"
 #include "Engine/BlueprintGeneratedClass.h"
 #include "Editor.h"
+#include "GraphEditorSettings.h"
 #include "Kismet2/KismetEditorUtilities.h"
 #include "K2Node.h"
 #include "Kismet2/Breakpoint.h"
@@ -30,6 +32,8 @@
 #include "PropertyEditor/Private/SDetailsView.h"
 #include "Styling/SlateIconFinder.h"
 #include "Styling/StyleColors.h"
+#include "Kismet2/BlueprintEditorUtils.h"
+#include "PropertyInfoViewStyle.h"
 
 #define LOCTEXT_NAMESPACE "DebugViewUI"
 
@@ -65,15 +69,44 @@ FText FDebugLineItem::GetDescription() const
 TSharedRef<SWidget> FDebugLineItem::GenerateNameWidget()
 {
 	return SNew(STextBlock)
+		.ToolTipText(this, &FDebugLineItem::GetDisplayName)
 		.Text(this, &FDebugLineItem::GetDisplayName);
 }
 
 TSharedRef<SWidget> FDebugLineItem::GenerateValueWidget()
 {
 	return SNew(STextBlock)
+		.ToolTipText(this, &FDebugLineItem::GetDescription)
 		.Text(this, &FDebugLineItem::GetDescription);
 }
 
+bool FDebugLineItem::HasChildren()
+{
+	return false;
+}
+
+TSharedRef<SWidget> FDebugLineItem::GetNameIcon()
+{
+	static const FSlateBrush* CachedBrush = FEditorStyle::GetBrush(TEXT("NoBrush"));	
+	return SNew(SImage).Image(CachedBrush);
+}
+
+TSharedRef<SWidget> FDebugLineItem::GetValueIcon()
+{
+	static const FSlateBrush* CachedBrush = FEditorStyle::GetBrush(TEXT("NoBrush"));	
+	return SNew(SImage).Image(CachedBrush);
+}
+
+FMargin FDebugLineItem::GetInnerMargin()
+{
+	return {0.f};
+}
+
+FMargin FDebugLineItem::GetOuterMargin()
+{
+	return {0.5f, 0.5f};
+}
+
 UBlueprint* FDebugLineItem::GetBlueprintForObject(UObject* ParentObject)
 {
 	if(ParentObject == nullptr)
@@ -146,11 +179,20 @@ void FDebugLineItem::OnDebugLineTypeActiveChanged(ECheckBoxState CheckState, EDe
 //////////////////////////////////////////////////////////////////////////
 // ILineItemWithChildren
 
-class ILineItemWithChildren
+class FLineItemWithChildren : public FDebugLineItem
 {
 public:
-	virtual ~ILineItemWithChildren() = default;
+	FLineItemWithChildren(EDebugLineType InType) :
+		FDebugLineItem(InType)
+	{}
 	
+	virtual ~FLineItemWithChildren() override = default;
+
+	virtual bool HasChildren() override
+	{
+		return !ChildrenMirrors.IsEmpty();
+	}
+
 protected:
 	// List of children 
 	TArray<FDebugTreeItemPtr> ChildrenMirrors;
@@ -171,6 +213,7 @@ protected:
 			{
 				if (Item.Compare(MirrorItem.Get()))
 				{
+					MirrorItem->UpdateData(Item);
 					OutChildren.Add(MirrorItem);
 					return;
 				}
@@ -248,6 +291,7 @@ protected:
 	}
 protected:
 	virtual TSharedRef<SWidget> GenerateNameWidget() override;
+	virtual TSharedRef<SWidget> GetNameIcon() override;
 	virtual FText GetDescription() const override;
 	virtual FText GetDisplayName() const override;
 	void OnNavigateToLatentNode();
@@ -271,25 +315,17 @@ FText FLatentActionLineItem::GetDescription() const
 
 TSharedRef<SWidget> FLatentActionLineItem::GenerateNameWidget()
 {
-	return
-		SNew(SHorizontalBox)
-		+SHorizontalBox::Slot()
-		.AutoWidth()
-		[
-			SNew(SImage)
-			. Image(FEditorStyle::GetBrush(TEXT("Kismet.LatentActionIcon")))
-		]
-		+SHorizontalBox::Slot()
-		.AutoWidth()
-		.HAlign(HAlign_Left)
-		.AutoWidth()
-		[
-				SNew(SHyperlink)
-			.Style(FEditorStyle::Get(), "HoverOnlyHyperlink")
-			.OnNavigate(this, &FLatentActionLineItem::OnNavigateToLatentNode)
-			.Text(this, &FLatentActionLineItem::GetDisplayName)
-			.ToolTipText( LOCTEXT("NavLatentActionLoc_Tooltip", "Navigate to the latent action location") )
-		];
+	return SNew(SHyperlink)
+		.Style(FEditorStyle::Get(), "HoverOnlyHyperlink")
+		.OnNavigate(this, &FLatentActionLineItem::OnNavigateToLatentNode)
+		.Text(this, &FLatentActionLineItem::GetDisplayName)
+		.ToolTipText( LOCTEXT("NavLatentActionLoc_Tooltip", "Navigate to the latent action location") );
+}
+
+TSharedRef<SWidget> FLatentActionLineItem::GetNameIcon()
+{
+	return SNew(SImage)
+		.Image(FEditorStyle::GetBrush(TEXT("Kismet.LatentActionIcon")));
 }
 
 UEdGraphNode* FLatentActionLineItem::FindAssociatedNode() const
@@ -326,22 +362,31 @@ void FLatentActionLineItem::OnNavigateToLatentNode( )
 	}
 }
 
-struct FWatchChildLineItem : public FDebugLineItem, ILineItemWithChildren
+struct FWatchChildLineItem : public FLineItemWithChildren
 {
 protected:
-	FDebugInfo Data;
+	FPropertyInstanceInfo Data;
 public:
-	FWatchChildLineItem(FDebugInfo Child) :
-		FDebugLineItem(DLT_WatchChild),
+	FWatchChildLineItem(const FPropertyInstanceInfo& Child) :
+		FLineItemWithChildren(DLT_WatchChild),
 		Data(Child)
 	{}
 
 	virtual bool Compare(const FDebugLineItem* BaseOther) const override
 	{
 		FWatchChildLineItem* Other = (FWatchChildLineItem*)BaseOther;
-		return
-			Data.DisplayName.CompareTo(Other->Data.DisplayName) == 0 &&
-			Data.Value.CompareTo(Other->Data.Value) == 0;
+		
+		return Data.Property == Other->Data.Property &&
+			Data.DisplayName.CompareTo(Other->Data.DisplayName) == 0;
+	}
+
+	virtual void UpdateData(const FDebugLineItem& NewerData) override
+	{
+		// Compare returns true even if the value or children of this node
+		// is different. use this function to update the data without completely
+		// replacing the node
+		FWatchChildLineItem& Other = (FWatchChildLineItem&)NewerData;
+		Data = Other.Data;
 	}
 
 	virtual FDebugLineItem* Duplicate() const override
@@ -351,16 +396,8 @@ public:
 	
 	virtual FText GetDescription() const override
 	{
-		FString ValStr = Data.Value.ToString();
-		if(int32 idx; ValStr.FindChar('\n', idx))
-		{
-			if(ValStr.Len() < 60)
-			{
-				return FText::FromString(ValStr.Replace(TEXT("\n") ,TEXT("")));
-			}
-			return FText::Format(LOCTEXT("BracketedDatatype","[{0}]"), Data.Type);
-		}
-		return Data.Value;
+		const FString ValStr = Data.Value.ToString();
+		return FText::FromString(ValStr.Replace(TEXT("\n") ,TEXT(" ")));
 	}
 	
 	virtual FText GetDisplayName() const override
@@ -368,30 +405,100 @@ public:
 		return Data.DisplayName;
 	}
 
+	// uses the icon and color associated with the property type
+	virtual TSharedRef<SWidget> GetNameIcon() override
+	{
+		FSlateColor Color;
+		FSlateColor SecondaryColor;
+		FSlateBrush const* SecondaryIcon;
+		const FSlateBrush* Icon = FBlueprintEditor::GetVarIconAndColorFromProperty(
+			Data.Property.Get(),
+			Color,
+			SecondaryIcon,
+			SecondaryColor
+		);
+		return SNew(SImage)
+			.Image(Icon)
+			.ColorAndOpacity(Color);
+	}
+
 	virtual void GatherChildren(TArray<FDebugTreeItemPtr>& OutChildren) override
 	{
-		for (FDebugInfo& ChildData : Data.Children)
+		for (const TSharedPtr<FPropertyInstanceInfo> &ChildData : Data.Children)
 		{
-			EnsureChildIsAdded(OutChildren, FWatchChildLineItem(ChildData));
+			EnsureChildIsAdded(OutChildren, FWatchChildLineItem(*ChildData));
 		}
 	}
 };
 
+//////////////////////////////////////////////////////////////////////////\
+// FSelfWatchLineItem
+struct FSelfWatchLineItem : public FLineItemWithChildren
+{
+protected:
+	// watches a UObject instead of a pin
+	TWeakObjectPtr<UObject> ObjectToWatch;
+public:
+	FSelfWatchLineItem(UObject* Object)	: 
+		FLineItemWithChildren(DLT_Watch),
+		ObjectToWatch(Object)
+	{}
+
+	virtual bool Compare(const FDebugLineItem* BaseOther) const override
+	{
+		FSelfWatchLineItem* Other = (FSelfWatchLineItem*)BaseOther;
+		return (ObjectToWatch.Get() == Other->ObjectToWatch.Get());
+	}
+
+	virtual FDebugLineItem* Duplicate() const override
+	{
+		return new FSelfWatchLineItem(ObjectToWatch.Get());
+	}
+	
+	virtual void GatherChildren(TArray<FDebugTreeItemPtr>& OutChildren) override
+	{
+		if (UObject* Object = ObjectToWatch.Get())
+		{
+			for (TFieldIterator<FProperty> It(Object->GetClass()); It; ++It)
+			{
+				TSharedPtr<FPropertyInstanceInfo> DebugInfo;
+				FProperty* Property = *It;
+				if(Property->HasAllPropertyFlags(CPF_BlueprintVisible))
+				{
+					void* Value = Property->ContainerPtrToValuePtr<void*>(Object);
+					FKismetDebugUtilities::GetDebugInfoInternal(DebugInfo, Property, Value);
+				
+					EnsureChildIsAdded(OutChildren, FWatchChildLineItem(*DebugInfo));
+				}
+			}
+		}
+	}
+
+protected:
+	virtual FText GetDisplayName() const override
+	{
+		return LOCTEXT("SelfName", "Self");
+	}
+
+	virtual TSharedRef<SWidget> GetNameIcon() override
+	{
+		return SNew(SImage)
+			.Image(FEditorStyle::GetBrush(TEXT("Kismet.WatchIcon")));
+	}
+};
+
 //////////////////////////////////////////////////////////////////////////
 // FWatchLineItem 
 
 
-struct FWatchLineItem : public FDebugLineItem, ILineItemWithChildren
+struct FWatchLineItem : public FLineItemWithChildren
 {
 protected:
 	TWeakObjectPtr< UObject > ParentObjectRef;
 	FEdGraphPinReference ObjectRef;
-	
-	// List of children 
-	TArray<FDebugTreeItemPtr> ChildrenMirrors;
 public:
 	FWatchLineItem(UEdGraphPin* PinToWatch, UObject* ParentObject)
-		: FDebugLineItem(DLT_Watch)
+		: FLineItemWithChildren(DLT_Watch)
 	{
 		ObjectRef = PinToWatch;
 		ParentObjectRef = ParentObject;
@@ -400,8 +507,8 @@ public:
 	virtual bool Compare(const FDebugLineItem* BaseOther) const override
 	{
 		FWatchLineItem* Other = (FWatchLineItem*)BaseOther;
-		return (ParentObjectRef.Get() == Other->ParentObjectRef.Get()) &&
-			(ObjectRef.Get() == Other->ObjectRef.Get());
+		return (ParentObjectRef == Other->ParentObjectRef) &&
+			(ObjectRef == Other->ObjectRef);
 	}
 
 	virtual FDebugLineItem* Duplicate() const override
@@ -436,14 +543,15 @@ public:
 			UObject* ParentObject = ParentObjectRef.Get();
 			if ((ParentBlueprint != ParentObject) && (ParentBlueprint != nullptr))
 			{
-				FDebugInfo DebugInfo;
+				TSharedPtr<FPropertyInstanceInfo> DebugInfo;
 				const FKismetDebugUtilities::EWatchTextResult WatchStatus = FKismetDebugUtilities::GetDebugInfo(DebugInfo, ParentBlueprint, ParentObject, PinToWatch);
 
 				if(WatchStatus == FKismetDebugUtilities::EWTR_Valid)
 				{
-					for (FDebugInfo& ChildData : DebugInfo.Children)
+					check(DebugInfo);
+					for (const TSharedPtr<FPropertyInstanceInfo> &ChildData : DebugInfo->Children)
 					{
-						EnsureChildIsAdded(OutChildren, FWatchChildLineItem(ChildData));
+						EnsureChildIsAdded(OutChildren, FWatchChildLineItem(*ChildData));
 					}
 				}
 			}
@@ -455,6 +563,7 @@ protected:
 	virtual FText GetDescription() const override;
 	virtual FText GetDisplayName() const override;
 	virtual TSharedRef<SWidget> GenerateNameWidget() override;
+	virtual TSharedRef<SWidget> GetNameIcon() override;
 
 	void OnNavigateToWatchLocation();
 };
@@ -492,20 +601,16 @@ FText FWatchLineItem::GetDescription() const
 		UObject* ParentObject = ParentObjectRef.Get();
 		if ((ParentBlueprint != ParentObject) && (ParentBlueprint != nullptr))
 		{
-			
-			FDebugInfo DebugInfo;
+			TSharedPtr<FPropertyInstanceInfo> DebugInfo;
 			const FKismetDebugUtilities::EWatchTextResult WatchStatus = FKismetDebugUtilities::GetDebugInfo(DebugInfo, ParentBlueprint, ParentObject, PinToWatch);
 
 			switch (WatchStatus)
 			{
 			case FKismetDebugUtilities::EWTR_Valid:
 			{
-				int32 Index = INDEX_NONE;
-				if(DebugInfo.Value.ToString().FindChar('\n', Index))
-				{
-					return DebugInfo.Type;
-				}
-				return DebugInfo.Value;
+				check(DebugInfo);
+				const FString ValStr = DebugInfo->Value.ToString();
+				return FText::FromString(ValStr.Replace(TEXT("\n") ,TEXT(" ")));
 			}
 
 			case FKismetDebugUtilities::EWTR_NotInScope:
@@ -526,23 +631,44 @@ FText FWatchLineItem::GetDescription() const
 
 TSharedRef<SWidget> FWatchLineItem::GenerateNameWidget()
 {
-	return SNew(SHorizontalBox)
+	return SNew(SHyperlink)
+		.Style(FEditorStyle::Get(), "HoverOnlyHyperlink")
+		.OnNavigate(this, &FWatchLineItem::OnNavigateToWatchLocation)
+		.Text(this, &FWatchLineItem::GetDisplayName)
+		.ToolTipText( LOCTEXT("NavWatchLoc", "Navigate to the watch location") );
+}
+
+// overlays the watch icon on top of a faded icon associated with the pin type
+TSharedRef<SWidget> FWatchLineItem::GetNameIcon()
+{
+	const FSlateBrush* PinIcon;
+	FLinearColor PinIconColor;
+	if (UEdGraphPin* ObjectToFocus = ObjectRef.Get())
+	{
+		PinIcon = FBlueprintEditorUtils::GetIconFromPin(ObjectToFocus->PinType);
 		
-		+SHorizontalBox::Slot()
-		.AutoWidth()
+		const UEdGraphSchema* Schema = ObjectToFocus->GetSchema();
+		PinIconColor = Schema->GetPinTypeColor(ObjectToFocus->PinType);
+		PinIconColor.A = 0.3f;
+	}
+	else
+	{
+		PinIcon = FEditorStyle::GetBrush(TEXT("NoBrush"));
+	}
+	
+	return SNew(SOverlay)
+		+ SOverlay::Slot()
+		.Padding(FMargin(10.f, 0.f, 0.f, 0.f))
 		[
 			SNew(SImage)
-			.Image(FEditorStyle::GetBrush(TEXT("Kismet.WatchIcon")))
+				.Image(PinIcon)
+				.ColorAndOpacity(PinIconColor)
 		]
-		+SHorizontalBox::Slot()
+		+ SOverlay::Slot()
 		.HAlign(HAlign_Left)
-		.AutoWidth()
 		[
-				SNew(SHyperlink)
-			.Style(FEditorStyle::Get(), "HoverOnlyHyperlink")
-			.OnNavigate(this, &FWatchLineItem::OnNavigateToWatchLocation)
-			.Text(this, &FWatchLineItem::GetDisplayName)
-			.ToolTipText( LOCTEXT("NavWatchLoc", "Navigate to the watch location") )
+			SNew(SImage)
+				.Image(FEditorStyle::GetBrush(TEXT("Kismet.WatchIcon")))
 		];
 }
 
@@ -587,7 +713,7 @@ public:
 		FBlueprintBreakpoint* Breakpoint = GetBreakpoint();
 		const UBlueprint* ParentBlueprint = GetBlueprintForObject(ParentObjectRef.Get());
 
-		// By default, we don't allow actions to execute when in debug mode. 
+		// By default, we don't allow actions to execute when in debug mode.
 		// Create an empty action to always allow execution for these commands (they are allowed in debug mode)
 		FCanExecuteAction AlwaysAllowExecute;
 
@@ -651,32 +777,24 @@ protected:
 	
 	virtual TSharedRef<SWidget> GenerateNameWidget() override
 	{
-		return SNew(SHorizontalBox)
-
-			+ SHorizontalBox::Slot()
-			.AutoWidth()
-			[
-				SNew(SButton)
-				. OnClicked(this, &FBreakpointLineItem::OnUserToggledEnabled)
-				. ToolTipText(LOCTEXT("ToggleBreakpointButton_ToolTip", "Toggle this breakpoint"))
-				. ButtonStyle( FEditorStyle::Get(), "NoBorder" )
-				. ContentPadding(0.0f)
-				[
-					SNew(SImage)
-					. Image(this, &FBreakpointLineItem::GetStatusImage)
-					. ToolTipText(this, &FBreakpointLineItem::GetStatusTooltip)
-				]
-			]
+		return SNew(SHyperlink)
+			.Style(FEditorStyle::Get(), "HoverOnlyHyperlink")
+			.Text(this, &FBreakpointLineItem::GetLocationDescription)
+			.ToolTipText( LOCTEXT("NavBreakpointLoc", "Navigate to the breakpoint location") )
+			.OnNavigate(this, &FBreakpointLineItem::OnNavigateToBreakpointLocation);
+	}
 
-		+ SHorizontalBox::Slot()
-			.AutoWidth()
-			. VAlign(VAlign_Center)
+	virtual TSharedRef<SWidget> GetNameIcon() override
+	{
+		return SNew(SButton)
+			.OnClicked(this, &FBreakpointLineItem::OnUserToggledEnabled)
+			.ToolTipText(LOCTEXT("ToggleBreakpointButton_ToolTip", "Toggle this breakpoint"))
+			.ButtonStyle( FEditorStyle::Get(), "NoBorder" )
+			.ContentPadding(0.0f)
 			[
-				SNew(SHyperlink)
-				. Style(FEditorStyle::Get(), "HoverOnlyHyperlink")
-				. Text(this, &FBreakpointLineItem::GetLocationDescription)
-				. ToolTipText( LOCTEXT("NavBreakpointLoc", "Navigate to the breakpoint location") )
-				. OnNavigate(this, &FBreakpointLineItem::OnNavigateToBreakpointLocation)
+				SNew(SImage)
+					.Image(this, &FBreakpointLineItem::GetStatusImage)
+					.ToolTipText(this, &FBreakpointLineItem::GetStatusTooltip)
 			];
 	}
 	
@@ -755,18 +873,14 @@ FText FBreakpointLineItem::GetStatusTooltip() const
 //////////////////////////////////////////////////////////////////////////
 // FTraceStackParentItem
 
-class FBreakpointParentItem : public FDebugLineItem, ILineItemWithChildren
+class FBreakpointParentItem : public FLineItemWithChildren
 {
-protected:
-
-	// List of children 
-	TArray<FDebugTreeItemPtr> ChildrenMirrors;
 public:
 	// The parent object
 	TWeakObjectPtr<UBlueprint> Blueprint;
 	
 	FBreakpointParentItem(TWeakObjectPtr<UBlueprint> Blueprint)
-		: FDebugLineItem(DLT_TraceStackParent),
+		: FLineItemWithChildren(DLT_TraceStackParent),
 		  Blueprint(Blueprint)
 	{
 	}
@@ -817,17 +931,14 @@ protected:
 //////////////////////////////////////////////////////////////////////////
 // FParentLineItem
 
-class FParentLineItem : public FDebugLineItem, ILineItemWithChildren
+class FParentLineItem : public FLineItemWithChildren
 {
 protected:
 	// The parent object
 	TWeakObjectPtr<UObject> ObjectRef;
-
-	// List of children 
-	TArray<FDebugTreeItemPtr> ChildrenMirrors;
 public:
 	FParentLineItem(UObject* Object)
-		: FDebugLineItem(DLT_Parent)
+		: FLineItemWithChildren(DLT_Parent)
 	{
 		ObjectRef = Object;
 	}
@@ -841,35 +952,18 @@ public:
 	{
 		if (UObject* ParentObject = ObjectRef.Get())
 		{
+			// every instance should have an automatic watch for 'self'
+			EnsureChildIsAdded(OutChildren, FSelfWatchLineItem(ParentObject));
+			
 			UBlueprint* ParentBP = FDebugLineItem::GetBlueprintForObject(ParentObject);
-			if ((ParentBP != nullptr) && (ParentBP == ParentObject))
+			if (ParentBP != nullptr)
 			{
 				// Create children for each watch
 				if(IsDebugLineTypeActive(DLT_Watch))
 				{
 					for (FEdGraphPinReference WatchedPin : ParentBP->WatchedPins)
-                    {
-                    	EnsureChildIsAdded(OutChildren, FWatchLineItem(WatchedPin.Get(), ParentObject));
-                    }
-				}
-
-				// Make sure there is something there, to let the user know if there is nothing
-				if (OutChildren.Num() == 0)
-				{
-					EnsureChildIsAdded(OutChildren, FMessageLineItem( LOCTEXT("NoWatches", "No watches").ToString() ));
-				}
-			}
-			else
-			{
-				if (ParentBP != nullptr)
-				{
-					// Create children for each watch
-					if(IsDebugLineTypeActive(DLT_Watch))
 					{
-						for (FEdGraphPinReference WatchedPin : ParentBP->WatchedPins)
-						{
-							EnsureChildIsAdded(OutChildren, FWatchLineItem(WatchedPin.Get(), ParentObject));
-						}
+						EnsureChildIsAdded(OutChildren, FWatchLineItem(WatchedPin.Get(), ParentObject));
 					}
 				}
 
@@ -922,7 +1016,8 @@ public:
 		{
 			return FSlateColor(EStyleColor::AccentYellow);
 		}
-		return FSlateColor::UseForeground();
+		const UGraphEditorSettings* Settings = GetDefault<UGraphEditorSettings>();
+		return Settings->ObjectPinTypeColor;
 	}
 	
 	FText GetStatusTooltip() const
@@ -961,24 +1056,12 @@ protected:
 		}
 	}
 
-	virtual TSharedRef<SWidget> GenerateNameWidget() override
+	virtual TSharedRef<SWidget> GetNameIcon() override
 	{
-		return SNew(SHorizontalBox)
-			+ SHorizontalBox::Slot()
-			.AutoWidth()
-			[
-				SNew(SImage)
-				.Image(this, &FParentLineItem::GetStatusImage)
-				.ColorAndOpacity_Raw(this, &FParentLineItem::GetStatusColor)
-				.ToolTipText(this, &FParentLineItem::GetStatusTooltip)
-			]
-			+ SHorizontalBox::Slot()
-			.AutoWidth()
-			. VAlign(VAlign_Center)
-			[
-				SNew(STextBlock)
-				.Text(this, &FParentLineItem::GetDisplayName)
-			];
+		return SNew(SImage)
+			.Image(this, &FParentLineItem::GetStatusImage)
+			.ColorAndOpacity_Raw(this, &FParentLineItem::GetStatusColor)
+			.ToolTipText(this, &FParentLineItem::GetStatusTooltip);
 	}
 
 	virtual void MakeMenu(class FMenuBuilder& MenuBuilder) override
@@ -1075,50 +1158,44 @@ protected:
 		}
 	}
 
-	// Index icon and node name
 	virtual TSharedRef<SWidget> GenerateNameWidget() override
 	{
-		return SNew(SHorizontalBox)
-
-			+ SHorizontalBox::Slot()
-			.AutoWidth()
-			[
-				SNew(SImage)
-				. Image( FEditorStyle::GetBrush((StackIndex > 0) ? TEXT("Kismet.Trace.PreviousIndex") : TEXT("Kismet.Trace.CurrentIndex")) )
-			]
+		return SNew(SHyperlink)
+			.Text(this, &FTraceStackChildItem::GetDisplayName)
+			.Style(FEditorStyle::Get(), "HoverOnlyHyperlink")
+			.ToolTipText(LOCTEXT("NavigateToDebugTraceLocationHyperlink_ToolTip", "Navigate to the trace location"))
+			.OnNavigate(this, &FTraceStackChildItem::OnNavigateToNode);
+	}
 
-		+ SHorizontalBox::Slot()
-			.AutoWidth()
-			. VAlign(VAlign_Center)
-			[
-				SNew(SHyperlink)
-				. Text(this, &FTraceStackChildItem::GetDisplayName)
-				. Style(FEditorStyle::Get(), "HoverOnlyHyperlink")
-				. ToolTipText(LOCTEXT("NavigateToDebugTraceLocationHyperlink_ToolTip", "Navigate to the trace location"))
-				. OnNavigate(this, &FTraceStackChildItem::OnNavigateToNode)
-			];
+	virtual TSharedRef<SWidget> GetNameIcon() override
+	{
+		return SNew(SImage)
+			.Image( FEditorStyle::GetBrush(
+				(StackIndex > 0) ?
+				TEXT("Kismet.Trace.PreviousIndex") :
+				TEXT("Kismet.Trace.CurrentIndex"))
+			);
 	}
 
 	// Visit time and actor name
 	virtual TSharedRef<SWidget> GenerateValueWidget() override
 	{
 		return SNew(SHorizontalBox)
-			
 			+ SHorizontalBox::Slot()
 			.AutoWidth()
 			[
 				SNew(SHyperlink)
-				. Text(this, &FTraceStackChildItem::GetContextObjectName)
-				. Style(FEditorStyle::Get(), "HoverOnlyHyperlink")
-				. ToolTipText( LOCTEXT("SelectActor_Tooltip", "Select this actor") )
-				. OnNavigate(this, &FTraceStackChildItem::OnSelectContextObject)
+					.Text(this, &FTraceStackChildItem::GetContextObjectName)
+					.Style(FEditorStyle::Get(), "HoverOnlyHyperlink")
+					.ToolTipText( LOCTEXT("SelectActor_Tooltip", "Select this actor") )
+					.OnNavigate(this, &FTraceStackChildItem::OnSelectContextObject)
 			]
 			
 			+ SHorizontalBox::Slot()
 			.AutoWidth()
 			[
 				SNew(STextBlock)
-				.Text(this, &FTraceStackChildItem::GetVisitTime)
+					.Text(this, &FTraceStackChildItem::GetVisitTime)
 			];
 	}
 
@@ -1175,14 +1252,11 @@ protected:
 //////////////////////////////////////////////////////////////////////////
 // FTraceStackParentItem
 
-class FTraceStackParentItem : public FDebugLineItem
+class FTraceStackParentItem : public FLineItemWithChildren
 {
-protected:
-	// List of children 
-	TArray< TSharedPtr<FTraceStackChildItem> > ChildrenMirrors;
 public:
 	FTraceStackParentItem()
-		: FDebugLineItem(DLT_TraceStackParent)
+		: FLineItemWithChildren(DLT_TraceStackParent)
 	{
 	}
 
@@ -1239,32 +1313,78 @@ public:
 		TSharedPtr<SWidget> ColumnContent = nullptr;
 		if (ColumnName == KismetDebugViewConstants::ColumnId_Name)
 		{
-			ColumnContent = SNew(SHorizontalBox)
+			SAssignNew(ColumnContent, SHorizontalBox)
 				+SHorizontalBox::Slot()
+					.HAlign(HAlign_Left)
+					.VAlign(VAlign_Fill)
+					.AutoWidth()
+					[
+						SNew( PropertyInfoViewStyle::SIndent, SharedThis(this) )
+					]
+				+ SHorizontalBox::Slot()
+					.HAlign(HAlign_Left)
+					.VAlign(VAlign_Center)
+					.AutoWidth()
+					[
+						SNew(PropertyInfoViewStyle::SExpanderArrow, SharedThis(this))
+							.HasChildren_Lambda([ItemToEdit = ItemToEdit]()
+							{
+								const bool HasChildren = ItemToEdit->HasChildren();
+								return HasChildren;
+							})
+					]
+	            + SHorizontalBox::Slot()
+	                .AutoWidth()
+	                .HAlign(HAlign_Left)
+					.VAlign(VAlign_Center)
+	                [
+	                    ItemToEdit->GetNameIcon()
+	                ]
+	                + SHorizontalBox::Slot()
+	                .AutoWidth()
+	                .HAlign(HAlign_Left)
+					.VAlign(VAlign_Center)
+	                .Padding(5.f,0.f,0.f,0.f)
+	                [
+	                    ItemToEdit->GenerateNameWidget()
+	                ];
+		}
+		else if (ColumnName == KismetDebugViewConstants::ColumnId_Value)
+		{
+			SAssignNew(ColumnContent, SHorizontalBox)
+				+ SHorizontalBox::Slot()
 				.AutoWidth()
 				[
-					SNew( SExpanderArrow, SharedThis(this) )
+					ItemToEdit->GetValueIcon()
 				]
-			+SHorizontalBox::Slot()
-				. FillWidth(1.0f)
+				+ SHorizontalBox::Slot()
+				.AutoWidth()
+				.HAlign(HAlign_Left)
+				.Padding(.5f, 1.f)
 				[
-					ItemToEdit->GenerateNameWidget()
+					ItemToEdit->GenerateValueWidget()
 				];
 		}
-		else if (ColumnName == KismetDebugViewConstants::ColumnId_Value)
-		{
-			ColumnContent = ItemToEdit->GenerateValueWidget();
-		}
 		else
 		{
-			ColumnContent = SNew(STextBlock)
-				. Text( LOCTEXT("Error", "Error") );
+			SAssignNew(ColumnContent, STextBlock)
+				.Text( LOCTEXT("Error", "Error") );
 		}
 		
-		return SNew(SBox).Padding(FMargin(0.0f, 5.0f))
-		[
-			ColumnContent.ToSharedRef()
-		];
+		return SNew(SBox)
+			.Padding(ItemToEdit->GetOuterMargin())
+			[
+				SNew(SBorder)
+					.Padding(ItemToEdit->GetInnerMargin())
+					.BorderImage(FAppStyle::Get().GetBrush("DetailsView.CategoryMiddle"))
+					.BorderBackgroundColor_Static(
+						PropertyInfoViewStyle::GetRowBackgroundColor,
+						static_cast<ITableRow*>(this)
+					)
+					[
+						ColumnContent.ToSharedRef()
+					]
+			];
 	}
 
 	void Construct(const FArguments& InArgs, TSharedRef<STableViewBase> OwnerTableView, FDebugTreeItemPtr InItemToEdit)
@@ -1293,21 +1413,21 @@ void SKismetDebuggingView::OnGetChildrenForWatchTree(FDebugTreeItemPtr InParent,
 TSharedRef<SHorizontalBox> SKismetDebuggingView::GetDebugLineTypeToggle(FDebugLineItem::EDebugLineType Type, const FText& Text)
 {
 	return SNew(SHorizontalBox)
-	+ SHorizontalBox::Slot()
-	.AutoWidth()
-	[
-		SNew(SCheckBox)
-			.IsChecked(true)
-			.OnCheckStateChanged_Static(&FDebugLineItem::OnDebugLineTypeActiveChanged, Type)
-	]
-	+ SHorizontalBox::Slot()
-	.AutoWidth()
-	.Padding(4.0f, 0.0f, 10.0f, 0.0f)
-	.VAlign(VAlign_Center)
-	[
-		SNew( STextBlock )
-			.Text(Text)
-	];
+		+ SHorizontalBox::Slot()
+		.AutoWidth()
+		[
+			SNew(SCheckBox)
+				.IsChecked(true)
+				.OnCheckStateChanged_Static(&FDebugLineItem::OnDebugLineTypeActiveChanged, Type)
+		]
+		+ SHorizontalBox::Slot()
+		.AutoWidth()
+		.Padding(4.0f, 0.0f, 10.0f, 0.0f)
+		.VAlign(VAlign_Center)
+		[
+			SNew( STextBlock )
+				.Text(Text)
+		];
 }
 
 TSharedPtr<SWidget> SKismetDebuggingView::OnMakeContextMenu()
@@ -1399,8 +1519,7 @@ TSharedRef<SWidget> SKismetDebuggingView::ConstructBlueprintClassPicker()
 
 	FOnClassPicked OnClassPicked;
 	OnClassPicked.BindRaw(this, &SKismetDebuggingView::OnBlueprintClassPicked);
-	return
-		SNew(SBox)
+	return SNew(SBox)
 		.HeightOverride(500.f)
 		[
 			ClassViewerModule.CreateClassViewer(Options, OnClassPicked)
@@ -1428,17 +1547,17 @@ void SKismetDebuggingView::Construct(const FArguments& InArgs)
 	
 	DebugClassComboButton =
 		SNew(SComboButton)
-		.OnGetMenuContent_Raw(this, &SKismetDebuggingView::ConstructBlueprintClassPicker)
-		.ButtonContent()
-		[
-			SNew(STextBlock)
-			.Text_Lambda([BlueprintToWatchPtr = BlueprintToWatchPtr]()
-			{
-				return BlueprintToWatchPtr.IsValid()?
-					FText::FromString(BlueprintToWatchPtr->GetName()) :
-					LOCTEXT("SelectBlueprint", "Select Blueprint");
-			})
-		];
+			.OnGetMenuContent_Raw(this, &SKismetDebuggingView::ConstructBlueprintClassPicker)
+			.ButtonContent()
+			[
+				SNew(STextBlock)
+					.Text_Lambda([&BlueprintToWatchPtr = BlueprintToWatchPtr]()
+					{
+						return BlueprintToWatchPtr.IsValid()?
+							FText::FromString(BlueprintToWatchPtr->GetName()) :
+							LOCTEXT("SelectBlueprint", "Select Blueprint");
+					})
+			];
 
 	FBlueprintContextTracker::OnEnterScriptContext.AddLambda(
 		[](const FBlueprintContextTracker& ContextTracker, const UObject* ContextObject, const UFunction* ContextFunction)
@@ -1458,101 +1577,115 @@ void SKismetDebuggingView::Construct(const FArguments& InArgs)
 	this->ChildSlot
 	[
 		SNew(SVerticalBox)
-		+ SVerticalBox::Slot()
-		.AutoHeight()
-		[
-			SNew(SBorder)
-			.BorderImage( FEditorStyle::GetBrush( TEXT("NoBorder") ) )
-			[
-				ToolbarWidget
-			]
-		]
-		+SVerticalBox::Slot()
-		.AutoHeight()
-		[
-			
-			SNew( SVerticalBox )
-			+SVerticalBox::Slot()
+			+ SVerticalBox::Slot()
 			.AutoHeight()
 			[
-				SNew( STextBlock )
-					.Text( this, &SKismetDebuggingView::GetTopText )
+				SNew(SBorder)
+					.BorderImage( FEditorStyle::GetBrush( TEXT("NoBorder") ) )
+					[
+						ToolbarWidget
+					]
 			]
-			+ SVerticalBox::Slot()
+			+SVerticalBox::Slot()
 			.AutoHeight()
 			[
-				SNew( SHorizontalBox )
-				+ SHorizontalBox::Slot()
-				.HAlign( HAlign_Left )
-				[
-					SNew(SBox)
-					.WidthOverride(400.f)
+				
+				SNew( SVerticalBox )
+					+SVerticalBox::Slot()
+					.AutoHeight()
 					[
-						DebugClassComboButton.ToSharedRef()
+						SNew( STextBlock )
+							.Text( this, &SKismetDebuggingView::GetTopText )
+					]
+					+ SVerticalBox::Slot()
+					.AutoHeight()
+					[
+						SNew( SHorizontalBox )
+							+ SHorizontalBox::Slot()
+							.HAlign( HAlign_Left )
+							[
+								SNew(SBox)
+									.WidthOverride(400.f)
+									[
+										DebugClassComboButton.ToSharedRef()
+									]
+							]
+							+ SHorizontalBox::Slot()
+							.HAlign( HAlign_Right )
+							[
+								SNew( SButton )
+					                .IsEnabled( this, &SKismetDebuggingView::CanDisableAllBreakpoints )
+					                .Text( LOCTEXT( "DisableAllBreakPoints", "Disable All Breakpoints" ) )
+					                .OnClicked( this, &SKismetDebuggingView::OnDisableAllBreakpointsClicked )
+							]
 					]
-				]
-				+ SHorizontalBox::Slot()
-				.HAlign( HAlign_Right )
-				[
-					SNew( SButton )
-		                .IsEnabled( this, &SKismetDebuggingView::CanDisableAllBreakpoints )
-		                .Text( LOCTEXT( "DisableAllBreakPoints", "Disable All Breakpoints" ) )
-		                .OnClicked( this, &SKismetDebuggingView::OnDisableAllBreakpointsClicked )
-				]
 			]
-		]
-		+SVerticalBox::Slot()
-		[
-			SNew(SSplitter)
-			.Orientation(Orient_Vertical)
-			+SSplitter::Slot()
+			+SVerticalBox::Slot()
 			[
-				SAssignNew( DebugTreeView, STreeView< FDebugTreeItemPtr > )
-				.TreeItemsSource( &RootTreeItems )
-				.SelectionMode( ESelectionMode::Single )
-				.OnGetChildren( this, &SKismetDebuggingView::OnGetChildrenForWatchTree )
-				.OnGenerateRow( this, &SKismetDebuggingView::OnGenerateRowForWatchTree ) 
-				.OnContextMenuOpening( this, &SKismetDebuggingView::OnMakeContextMenu )
-				. HeaderRow
-				(
-					SNew(SHeaderRow)
-					+ SHeaderRow::Column(KismetDebugViewConstants::ColumnId_Name)
-					.DefaultLabel(KismetDebugViewConstants::ColumnText_Name)
-					+ SHeaderRow::Column(KismetDebugViewConstants::ColumnId_Value)
-					.DefaultLabel(KismetDebugViewConstants::ColumnText_Value)
-				)
+				SNew(SSplitter)
+					.Orientation(Orient_Vertical)
+					+SSplitter::Slot()
+					[
+						SAssignNew( DebugTreeView, STreeView< FDebugTreeItemPtr > )
+							.TreeItemsSource( &RootTreeItems )
+							.SelectionMode( ESelectionMode::Single )
+							.OnGetChildren( this, &SKismetDebuggingView::OnGetChildrenForWatchTree )
+							.OnGenerateRow( this, &SKismetDebuggingView::OnGenerateRowForWatchTree ) 
+							.OnContextMenuOpening( this, &SKismetDebuggingView::OnMakeContextMenu )
+							.TreeViewStyle(&FAppStyle::Get().GetWidgetStyle<FTableViewStyle>("PropertyTable.InViewport.ListView"))
+							.HeaderRow
+							(
+								SNew(SHeaderRow)
+									+ SHeaderRow::Column(KismetDebugViewConstants::ColumnId_Name)
+									.DefaultLabel(KismetDebugViewConstants::ColumnText_Name)
+									+ SHeaderRow::Column(KismetDebugViewConstants::ColumnId_Value)
+									.DefaultLabel(KismetDebugViewConstants::ColumnText_Value)
+							)
+					]
+					+SSplitter::Slot()
+					[
+						SAssignNew( OtherTreeView, STreeView< FDebugTreeItemPtr > )
+							.TreeItemsSource( &OtherTreeItems )
+							.SelectionMode( ESelectionMode::Single )
+							.OnGetChildren( this, &SKismetDebuggingView::OnGetChildrenForWatchTree )
+							.OnGenerateRow( this, &SKismetDebuggingView::OnGenerateRowForWatchTree ) 
+							.OnContextMenuOpening( this, &SKismetDebuggingView::OnMakeContextMenu )
+							.TreeViewStyle(&FAppStyle::Get().GetWidgetStyle<FTableViewStyle>("PropertyTable.InViewport.ListView"))
+							.HeaderRow
+							(
+								SNew(SHeaderRow)
+									+ SHeaderRow::Column(KismetDebugViewConstants::ColumnId_Name)
+									.DefaultLabel(KismetDebugViewConstants::ColumnText_DebugKey)
+									+ SHeaderRow::Column(KismetDebugViewConstants::ColumnId_Value)
+									.DefaultLabel(KismetDebugViewConstants::ColumnText_Info)
+							)
+					]
 			]
-			+SSplitter::Slot()
+			+ SVerticalBox::Slot()
+			.AutoHeight()
 			[
-				SAssignNew( OtherTreeView, STreeView< FDebugTreeItemPtr > )
-				.TreeItemsSource( &OtherTreeItems )
-				.SelectionMode( ESelectionMode::Single )
-				.OnGetChildren( this, &SKismetDebuggingView::OnGetChildrenForWatchTree )
-				.OnGenerateRow( this, &SKismetDebuggingView::OnGenerateRowForWatchTree ) 
-				.OnContextMenuOpening( this, &SKismetDebuggingView::OnMakeContextMenu )
-				. HeaderRow
-				(
-					SNew(SHeaderRow)
-					+ SHeaderRow::Column(KismetDebugViewConstants::ColumnId_Name)
-					.DefaultLabel(KismetDebugViewConstants::ColumnText_DebugKey)
-					+ SHeaderRow::Column(KismetDebugViewConstants::ColumnId_Value)
-					.DefaultLabel(KismetDebugViewConstants::ColumnText_Info)
-				)
+				SNew(SHorizontalBox)
+					+ SHorizontalBox::Slot()
+					.AutoWidth()
+					[
+						GetDebugLineTypeToggle(FDebugLineItem::DLT_Watch, LOCTEXT("Watchpoints", "Watchpoints"))
+					]
+					+ SHorizontalBox::Slot()
+					.AutoWidth()
+					[
+						GetDebugLineTypeToggle(FDebugLineItem::DLT_LatentAction, LOCTEXT("LatentActions", "Latent Actions"))
+					]
+					+ SHorizontalBox::Slot()
+					.AutoWidth()
+					[
+						GetDebugLineTypeToggle(FDebugLineItem::DLT_BreakpointParent, LOCTEXT("Breakpoints", "Breakpoints"))
+					]
+					+ SHorizontalBox::Slot()
+					.AutoWidth()
+					[
+						GetDebugLineTypeToggle(FDebugLineItem::DLT_TraceStackParent, LOCTEXT("ExecutionTrace", "Execution Trace"))
+					]
 			]
-		]
-		+ SVerticalBox::Slot()
-		.AutoHeight()
-		[
-			SNew(SHorizontalBox)
-			+ SHorizontalBox::Slot().AutoWidth()
-				[ GetDebugLineTypeToggle(FDebugLineItem::DLT_Watch, LOCTEXT("Watchpoints", "Watchpoints")) ]
-			+ SHorizontalBox::Slot().AutoWidth()
-				[ GetDebugLineTypeToggle(FDebugLineItem::DLT_LatentAction, LOCTEXT("LatentActions", "Latent Actions")) ]
-			+ SHorizontalBox::Slot().AutoWidth()
-				[ GetDebugLineTypeToggle(FDebugLineItem::DLT_BreakpointParent, LOCTEXT("Breakpoints", "Breakpoints")) ]
-			+ SHorizontalBox::Slot().AutoWidth()
-				[ GetDebugLineTypeToggle(FDebugLineItem::DLT_TraceStackParent, LOCTEXT("ExecutionTrace", "Execution Trace")) ]
-		]
 	];
 
 	TraceStackItem = MakeShareable(new FTraceStackParentItem);
diff --git a/Engine/Source/Editor/Kismet/Private/Debugging/SKismetDebuggingView.h b/Engine/Source/Editor/Kismet/Private/Debugging/SKismetDebuggingView.h
--- a/Engine/Source/Editor/Kismet/Private/Debugging/SKismetDebuggingView.h
+++ b/Engine/Source/Editor/Kismet/Private/Debugging/SKismetDebuggingView.h
@@ -26,7 +26,7 @@ typedef TSharedPtr<class FDebugLineItem> FDebugTreeItemPtr;
 class FDebugLineItem : public TSharedFromThis<FDebugLineItem>
 {
 public:
-	friend class ILineItemWithChildren;
+	friend class FLineItemWithChildren; // used by FLineItemWithChildren::EnsureChildIsAdded
 	enum EDebugLineType
 	{
 		DLT_Message,
@@ -54,6 +54,9 @@ public:
 	// Gather all of the children
 	virtual void GatherChildren(TArray<FDebugTreeItemPtr>& OutChildren) {}
 
+	// returns whether this tree node has children (used by drop down arrows)
+	virtual bool HasChildren();
+
 	// @return The object that will act as a parent to more items in the tree, or NULL if this is a leaf node
 	virtual UObject* GetParentObject() { return NULL; }
 
@@ -62,6 +65,18 @@ public:
 		return Type;
 	}
 
+	// returns a widget that will go to the left of the Name Widget.
+	virtual TSharedRef<SWidget> GetNameIcon();
+
+	// returns a widget that will go to the left of the Value Widget.
+	virtual TSharedRef<SWidget> GetValueIcon();
+
+	// determines the margin between the Name/Value widgets and the boarder
+	virtual FMargin GetInnerMargin();
+	
+	// determines the space between multiple rows/colums
+	virtual FMargin GetOuterMargin();
+
 	// Helper function to try to get the blueprint for a given object;
 	//   Returns the blueprint that was used to create the instance if there was one
 	//   Returns the object itself if it is already a blueprint
@@ -86,6 +101,10 @@ protected:
 	// Compare this item to another of the same type
 	virtual bool Compare(const FDebugLineItem* Other) const=0;
 
+	// Used to update the state of a line item rather than replace it.
+	// called after Compare returns true
+	virtual void UpdateData(const FDebugLineItem& NewerData) {}
+
 	// @return The text to display in the name column, unless GenerateNameWidget is overridden
 	virtual FText GetDisplayName() const;
 
diff --git a/Engine/Source/Editor/Kismet/Private/WatchPointViewer.cpp b/Engine/Source/Editor/Kismet/Private/WatchPointViewer.cpp
--- a/Engine/Source/Editor/Kismet/Private/WatchPointViewer.cpp
+++ b/Engine/Source/Editor/Kismet/Private/WatchPointViewer.cpp
@@ -78,7 +78,7 @@ namespace
 			FText InBlueprintName,
 			FText InGraphName,
 			FText InNodeName,
-			FDebugInfo Info
+			FPropertyInstanceInfo Info
 		)
 			: BP(InBP)
 			, Node(InNode)
@@ -96,9 +96,9 @@ namespace
 			UPackage* Package = Cast<UPackage>(BP.IsValid() ? BP->GetOuter() : nullptr);
 			BlueprintPackageName = Package ? Package->GetFName() : FName();
 
-			for (FDebugInfo& ChildInfo : Info.Children)
+			for (TSharedPtr<FPropertyInstanceInfo>& ChildInfo : Info.Children)
 			{
-				Children.Add(MakeShared<FWatchRow>(InBP, InNode, InPin, InObjectBeingDebugged, BlueprintName, GraphName, NodeName, MoveTemp(ChildInfo)));
+				Children.Add(MakeShared<FWatchRow>(InBP, InNode, InPin, InObjectBeingDebugged, BlueprintName, GraphName, NodeName, MoveTemp(*ChildInfo)));
 			}
 		}
 
@@ -192,10 +192,10 @@ namespace
 
 					const UEdGraphSchema* Schema = Pin->GetOwningNode()->GetSchema();
 
-					FDebugInfo DebugInfo;
-					DebugInfo.DisplayName = Schema->GetPinDisplayName(Pin);
-					DebugInfo.Type = UEdGraphSchema_K2::TypeToText(Pin->PinType);
-					DebugInfo.Value = LOCTEXT("ExecutionNotPaused", "(execution not paused)");
+					TSharedPtr<FPropertyInstanceInfo> DebugInfo;
+					DebugInfo->DisplayName = Schema->GetPinDisplayName(Pin);
+					DebugInfo->Type = UEdGraphSchema_K2::TypeToText(Pin->PinType);
+					DebugInfo->Value = LOCTEXT("ExecutionNotPaused", "(execution not paused)");
 
 					Private_WatchSource.Add(
 						MakeShared<FWatchRow>(
@@ -206,7 +206,7 @@ namespace
 							BlueprintName,
 							MoveTemp(GraphName),
 							MoveTemp(NodeName),
-							MoveTemp(DebugInfo)
+							MoveTemp(*DebugInfo)
 						)
 					);
 				}
@@ -849,27 +849,27 @@ void WatchViewer::UpdateInstancedWatchDisplay()
 				FText GraphName = FText::FromString(Pin->GetOwningNode()->GetGraph()->GetName());
 				FText NodeName = Pin->GetOwningNode()->GetNodeTitle(ENodeTitleType::ListView);
 
-				FDebugInfo DebugInfo;
+				TSharedPtr<FPropertyInstanceInfo> DebugInfo;
 				const FKismetDebugUtilities::EWatchTextResult WatchStatus = FKismetDebugUtilities::GetDebugInfo(DebugInfo, BlueprintObj, BlueprintInstance, Pin);
 
 				if (WatchStatus != FKismetDebugUtilities::EWTR_Valid)
 				{
 					const UEdGraphSchema* Schema = Pin->GetOwningNode()->GetSchema();
-					DebugInfo.DisplayName = Schema->GetPinDisplayName(Pin);
-					DebugInfo.Type = UEdGraphSchema_K2::TypeToText(Pin->PinType);
+					DebugInfo->DisplayName = Schema->GetPinDisplayName(Pin);
+					DebugInfo->Type = UEdGraphSchema_K2::TypeToText(Pin->PinType);
 
 					switch (WatchStatus)
 					{
 					case FKismetDebugUtilities::EWTR_NotInScope:
-						DebugInfo.Value = LOCTEXT("NotInScope", "(not in scope)");
+						DebugInfo->Value = LOCTEXT("NotInScope", "(not in scope)");
 						break;
 
 					case FKismetDebugUtilities::EWTR_NoProperty:
-						DebugInfo.Value = LOCTEXT("NoDebugData", "(no debug data)");
+						DebugInfo->Value = LOCTEXT("NoDebugData", "(no debug data)");
 						break;
 
 					case FKismetDebugUtilities::EWTR_NoDebugObject:
-						DebugInfo.Value = LOCTEXT("NoDebugObject", "(no debug object)");
+						DebugInfo->Value = LOCTEXT("NoDebugObject", "(no debug object)");
 						break;
 
 					default:
@@ -887,7 +887,7 @@ void WatchViewer::UpdateInstancedWatchDisplay()
 						BlueprintName,
 						GraphName,
 						NodeName,
-						DebugInfo
+						*DebugInfo
 						)
 				);
 			}
diff --git a/Engine/Source/Editor/UnrealEd/Private/Kismet2/KismetDebugUtilities.cpp b/Engine/Source/Editor/UnrealEd/Private/Kismet2/KismetDebugUtilities.cpp
--- a/Engine/Source/Editor/UnrealEd/Private/Kismet2/KismetDebugUtilities.cpp
+++ b/Engine/Source/Editor/UnrealEd/Private/Kismet2/KismetDebugUtilities.cpp
@@ -1315,26 +1315,26 @@ FKismetDebugUtilities::EWatchTextResult FKismetDebugUtilities::GetWatchText(FStr
 	return Result;
 }
 
-FKismetDebugUtilities::EWatchTextResult FKismetDebugUtilities::GetDebugInfo(FDebugInfo& OutDebugInfo, UBlueprint* Blueprint, UObject* ActiveObject, const UEdGraphPin* WatchPin)
+FKismetDebugUtilities::EWatchTextResult FKismetDebugUtilities::GetDebugInfo(TSharedPtr<FPropertyInstanceInfo> &OutDebugInfo, UBlueprint* Blueprint, UObject* ActiveObject, const UEdGraphPin* WatchPin)
 {
-	FProperty* PropertyToDebug = nullptr;
 	void* DataPtr = nullptr;
 	void* DeltaPtr = nullptr;
 	UObject* ParentObj = nullptr;
 	TArray<UObject*> SeenObjects;
 	bool bIsDirectPtr = false;
-	FKismetDebugUtilities::EWatchTextResult Result = FindDebuggingData(Blueprint, ActiveObject, WatchPin, PropertyToDebug, DataPtr, DeltaPtr, ParentObj, SeenObjects, &bIsDirectPtr);
+	FProperty* Property = nullptr;
+	FKismetDebugUtilities::EWatchTextResult Result = FindDebuggingData(Blueprint, ActiveObject, WatchPin, Property, DataPtr, DeltaPtr, ParentObj, SeenObjects, &bIsDirectPtr);
 
 	if (Result == FKismetDebugUtilities::EWatchTextResult::EWTR_Valid)
 	{
 		// If this came from an out parameter it isn't in a property container, so must be accessed directly
 		if (bIsDirectPtr)
 		{
-			GetDebugInfoInternal(OutDebugInfo, PropertyToDebug, DataPtr);
+			GetDebugInfoInternal(OutDebugInfo, Property, DataPtr);
 		}
 		else
 		{
-			GetDebugInfo_InContainer(0, OutDebugInfo, PropertyToDebug, DataPtr);
+			GetDebugInfo_InContainer(0, OutDebugInfo, Property, DataPtr);
 		}
 	}
 
@@ -1549,34 +1549,40 @@ FKismetDebugUtilities::EWatchTextResult FKismetDebugUtilities::FindDebuggingData
 	}
 }
 
-void FKismetDebugUtilities::GetDebugInfo_InContainer(int32 Index, FDebugInfo& DebugInfo, FProperty* Property, const void* Data)
+void FKismetDebugUtilities::GetDebugInfo_InContainer(int32 Index, TSharedPtr<FPropertyInstanceInfo> &DebugInfo, FProperty* Property, const void* Data)
 {
 	GetDebugInfoInternal(DebugInfo, Property, Property->ContainerPtrToValuePtr<void>(Data, Index));
 }
 
-void FKismetDebugUtilities::GetDebugInfoInternal(FDebugInfo& DebugInfo, FProperty* Property, const void* PropertyValue)
+void FKismetDebugUtilities::GetDebugInfoInternal(TSharedPtr<FPropertyInstanceInfo> &DebugInfo, FProperty* Property, const void* PropertyValue)
+{
+	TMap<FPropertyInstanceInfo::FPropertyInstance, TSharedPtr<FPropertyInstanceInfo>> VisitedNodes;
+	DebugInfo = FPropertyInstanceInfo::FindOrMake({Property, PropertyValue}, VisitedNodes);
+}
+
+FPropertyInstanceInfo::FPropertyInstanceInfo(FPropertyInstance PropertyInstance) :
+	DisplayName(PropertyInstance.Property->GetDisplayNameText()),
+	Type(UEdGraphSchema_K2::TypeToText(PropertyInstance.Property)),
+	Property(PropertyInstance.Property)
 {
-	if (Property == nullptr)
+	check(PropertyInstance.Property);
+	if (PropertyInstance.Value == nullptr)
 	{
 		return;
 	}
 
-	DebugInfo.Type = UEdGraphSchema_K2::TypeToText(Property);
-	DebugInfo.DisplayName = Property->GetDisplayNameText();
-
-	FByteProperty* ByteProperty = CastField<FByteProperty>(Property);
-	if (ByteProperty)
+	if (FByteProperty* ByteProperty = CastField<FByteProperty>(Property.Get()))
 	{
 		UEnum* Enum = ByteProperty->GetIntPropertyEnum();
 		if (Enum)
 		{
-			if (Enum->IsValidEnumValue(*(const uint8*)PropertyValue))
+			if (Enum->IsValidEnumValue(*(const uint8*)PropertyInstance.Value))
 			{
-				DebugInfo.Value = Enum->GetDisplayNameTextByValue(*(const uint8*)PropertyValue);
+				Value = Enum->GetDisplayNameTextByValue(*(const uint8*)PropertyInstance.Value);
 			}
 			else
 			{
-				DebugInfo.Value = FText::FromString(TEXT("(INVALID)"));
+				Value = FText::FromString(TEXT("(INVALID)"));
 			}
 
 			return;
@@ -1585,219 +1591,266 @@ void FKismetDebugUtilities::GetDebugInfoInternal(FDebugInfo& DebugInfo, FPropert
 		// if there is no Enum we need to fall through and treat this as a FNumericProperty
 	}
 
-	FNumericProperty* NumericProperty = CastField<FNumericProperty>(Property);
-	if (NumericProperty)
+	if (FNumericProperty* NumericProperty = CastField<FNumericProperty>(Property.Get()))
 	{
-		DebugInfo.Value = FText::FromString(NumericProperty->GetNumericPropertyValueToString(PropertyValue));
+		Value = FText::FromString(NumericProperty->GetNumericPropertyValueToString(PropertyInstance.Value));
 		return;
 	}
-
-	FBoolProperty* BoolProperty = CastField<FBoolProperty>(Property);
-	if (BoolProperty)
+	else if (FBoolProperty* BoolProperty = CastField<FBoolProperty>(Property.Get()))
 	{
 		const FCoreTexts& CoreTexts = FCoreTexts::Get();
 
-		DebugInfo.Value = BoolProperty->GetPropertyValue(PropertyValue) ? CoreTexts.True : CoreTexts.False;
+		Value = BoolProperty->GetPropertyValue(PropertyInstance.Value) ? CoreTexts.True : CoreTexts.False;
 		return;
 	}
-
-	FNameProperty* NameProperty = CastField<FNameProperty>(Property);
-	if (NameProperty)
+	else if (FNameProperty* NameProperty = CastField<FNameProperty>(Property.Get()))
 	{
-		DebugInfo.Value = FText::FromName(*(FName*)PropertyValue);
+		Value = FText::FromName(*(FName*)PropertyInstance.Value);
 		return;
 	}
-
-	FTextProperty* TextProperty = CastField<FTextProperty>(Property);
-	if (TextProperty)
+	else if (FTextProperty* TextProperty = CastField<FTextProperty>(Property.Get()))
 	{
-		DebugInfo.Value = TextProperty->GetPropertyValue(PropertyValue);
+		Value = TextProperty->GetPropertyValue(PropertyInstance.Value);
 		return;
 	}
-
-	FStrProperty* StringProperty = CastField<FStrProperty>(Property);
-	if (StringProperty)
+	else if (FStrProperty* StringProperty = CastField<FStrProperty>(Property.Get()))
 	{
-		DebugInfo.Value = FText::FromString(StringProperty->GetPropertyValue(PropertyValue));
+		Value = FText::FromString(StringProperty->GetPropertyValue(PropertyInstance.Value));
 		return;
 	}
-
-	FArrayProperty* ArrayProperty = CastField<FArrayProperty>(Property);
-	if (ArrayProperty)
+	else if (FArrayProperty* ArrayProperty = CastField<FArrayProperty>(Property.Get()))
 	{
 		checkSlow(ArrayProperty->Inner);
 
-		FScriptArrayHelper ArrayHelper(ArrayProperty, PropertyValue);
-
-		DebugInfo.Value = FText::Format(LOCTEXT("ArraySize", "Num={0}"), FText::AsNumber(ArrayHelper.Num()));
-
-		for (int32 i = 0; i < ArrayHelper.Num(); i++)
-		{
-			FDebugInfo ArrayDebugInfo;
-
-			uint8* PropData = ArrayHelper.GetRawPtr(i);
-			GetDebugInfoInternal(ArrayDebugInfo, ArrayProperty->Inner, PropData);
-			// overwrite the display name with the array index for the current element
-			ArrayDebugInfo.DisplayName = FText::Format(LOCTEXT("ArrayIndexName", "[{0}]"), FText::AsNumber(i));
-			DebugInfo.Children.Add(ArrayDebugInfo);
-		}
+		FScriptArrayHelper ArrayHelper(ArrayProperty, PropertyInstance.Value);
 
+		Value = FText::Format(LOCTEXT("ArraySize", "Num={0}"), FText::AsNumber(ArrayHelper.Num()));
 		return;
 	}
-
-	FStructProperty* StructProperty = CastField<FStructProperty>(Property);
-	if (StructProperty)
+	else if (FStructProperty* StructProperty = CastField<FStructProperty>(Property.Get()))
 	{
 		FString WatchText;
-		StructProperty->ExportTextItem(WatchText, PropertyValue, PropertyValue, nullptr, PPF_PropertyWindow | PPF_BlueprintDebugView, nullptr);
-		DebugInfo.Value = FText::FromString(WatchText);
-
-		for (TFieldIterator<FProperty> It(StructProperty->Struct); It; ++It)
-		{
-			FDebugInfo StructDebugInfo;
-			GetDebugInfoInternal(StructDebugInfo, *It, It->ContainerPtrToValuePtr<void>(PropertyValue, 0));
-
-			DebugInfo.Children.Add(StructDebugInfo);
-		}
-
+		StructProperty->ExportTextItem(WatchText, PropertyInstance.Value, PropertyInstance.Value, nullptr, PPF_PropertyWindow | PPF_BlueprintDebugView, nullptr);
+		Value = FText::FromString(WatchText);
 		return;
 	}
-
-	FEnumProperty* EnumProperty = CastField<FEnumProperty>(Property);
-	if (EnumProperty)
+	else if (FEnumProperty* EnumProperty = CastField<FEnumProperty>(Property.Get()))
 	{
 		FNumericProperty* LocalUnderlyingProp = EnumProperty->GetUnderlyingProperty();
 		UEnum* Enum = EnumProperty->GetEnum();
 
-		int64 Value = LocalUnderlyingProp->GetSignedIntPropertyValue(PropertyValue);
+		int64 NumValue = LocalUnderlyingProp->GetSignedIntPropertyValue(PropertyInstance.Value);
 
 		// if the value is the max value (the autogenerated *_MAX value), export as "INVALID", unless we're exporting text for copy/paste (for copy/paste,
 		// the property text value must actually match an entry in the enum's names array)
 		if (Enum)
 		{
-			if (Enum->IsValidEnumValue(Value))
+			if (Enum->IsValidEnumValue(NumValue))
 			{
-				DebugInfo.Value = Enum->GetDisplayNameTextByValue(Value);
+				Value = Enum->GetDisplayNameTextByValue(NumValue);
 			}
 			else
 			{
-				DebugInfo.Value = LOCTEXT("Invalid", "(INVALID)");
+				Value = LOCTEXT("Invalid", "(INVALID)");
 			}
 		}
 		else
 		{
-			DebugInfo.Value = FText::AsNumber(Value);
+			Value = FText::AsNumber(NumValue);
 		}
 
 		return;
 	}
-
-	FMapProperty* MapProperty = CastField<FMapProperty>(Property);
-	if (MapProperty)
+	else if (FMapProperty* MapProperty = CastField<FMapProperty>(Property.Get()))
 	{
-		FScriptMapHelper MapHelper(MapProperty, PropertyValue);
-		DebugInfo.Value = FText::Format(LOCTEXT("MapSize", "Num={0}"), FText::AsNumber(MapHelper.Num()));
-		uint8* PropData = MapHelper.GetPairPtr(0);
-
-		int32 Index = 0;
-		for (int32 Count = MapHelper.Num(); Count; PropData += MapProperty->MapLayout.SetLayout.Size, ++Index)
-		{
-			if (MapHelper.IsValidIndex(Index))
-			{
-				FDebugInfo ChildInfo;
-
-				GetDebugInfoInternal(ChildInfo, MapProperty->ValueProp, PropData + MapProperty->MapLayout.ValueOffset);
-
-				// use the info from the ValueProp and then overwrite the name with the KeyProp data
-				FString NameStr = TEXT("[");
-				MapProperty->KeyProp->ExportTextItem(NameStr, PropData, nullptr, nullptr, PPF_PropertyWindow | PPF_BlueprintDebugView | PPF_Delimited, nullptr);
-				NameStr += TEXT("] ");
-
-				ChildInfo.DisplayName = FText::FromString(NameStr);
-
-				DebugInfo.Children.Add(ChildInfo);
-
-				--Count;
-			}
-		}
-
+		FScriptMapHelper MapHelper(MapProperty, PropertyInstance.Value);
+		Value = FText::Format(LOCTEXT("MapSize", "Num={0}"), FText::AsNumber(MapHelper.Num()));
 		return;
 	}
-
-	FSetProperty* SetProperty = CastField<FSetProperty>(Property);
-	if (SetProperty)
+	else if (FSetProperty* SetProperty = CastField<FSetProperty>(Property.Get()))
 	{
-		FScriptSetHelper SetHelper(SetProperty, PropertyValue);
-		DebugInfo.Value = FText::Format(LOCTEXT("SetSize", "Num={0}"), FText::AsNumber(SetHelper.Num()));
-		uint8* PropData = SetHelper.GetElementPtr(0);
-
-		int32 Index = 0;
-		for (int32 Count = SetHelper.Num(); Count; PropData += SetProperty->SetLayout.Size, ++Index)
-		{
-			if (SetHelper.IsValidIndex(Index))
-			{
-				FDebugInfo ChildInfo;
-				GetDebugInfoInternal(ChildInfo, SetProperty->ElementProp, PropData);
-
-				// members of sets don't have their own names
-				ChildInfo.DisplayName = FText::GetEmpty();
-
-				DebugInfo.Children.Add(ChildInfo);
-
-				--Count;
-			}
-		}
-
+		FScriptSetHelper SetHelper(SetProperty, PropertyInstance.Value);
+		Value = FText::Format(LOCTEXT("SetSize", "Num={0}"), FText::AsNumber(SetHelper.Num()));
 		return;
 	}
-
-	FObjectPropertyBase* ObjectPropertyBase = CastField<FObjectPropertyBase>(Property);
-	if (ObjectPropertyBase)
+	else if (FObjectPropertyBase* ObjectPropertyBase = CastField<FObjectPropertyBase>(Property.Get()))
 	{
-		UObject* Obj = ObjectPropertyBase->GetObjectPropertyValue(PropertyValue);
-		if (Obj != nullptr)
+		if (UObject* Obj = ObjectPropertyBase->GetObjectPropertyValue(PropertyInstance.Value))
 		{
-			DebugInfo.Value = FText::FromString(Obj->GetFullName());
+			Value = FText::FromString(Obj->GetFullName());
 		}
 		else
 		{
-			DebugInfo.Value = FText::FromString(TEXT("None"));
+			Value = FText::FromString(TEXT("None"));
 		}
 
 		return;
 	}
-
-	FDelegateProperty* DelegateProperty = CastField<FDelegateProperty>(Property);
-	if (DelegateProperty)
+	else if (FDelegateProperty* DelegateProperty = CastField<FDelegateProperty>(Property.Get()))
 	{
 		if (DelegateProperty->SignatureFunction)
 		{
-			DebugInfo.Value = DelegateProperty->SignatureFunction->GetDisplayNameText();
+			Value = DelegateProperty->SignatureFunction->GetDisplayNameText();
 		}
 		else
 		{
-			DebugInfo.Value = LOCTEXT("NoFunc", "(No bound function)");
+			Value = LOCTEXT("NoFunc", "(No bound function)");
 		}
 
 		return;
 	}
-
-	FMulticastDelegateProperty* MulticastDelegateProperty = CastField<FMulticastDelegateProperty>(Property);
-	if (MulticastDelegateProperty)
+	else if (FMulticastDelegateProperty* MulticastDelegateProperty = CastField<FMulticastDelegateProperty>(Property.Get()))
 	{
 		if (MulticastDelegateProperty->SignatureFunction)
 		{
-			DebugInfo.Value = MulticastDelegateProperty->SignatureFunction->GetDisplayNameText();
+			Value = MulticastDelegateProperty->SignatureFunction->GetDisplayNameText();
 		}
 		else
 		{
-			DebugInfo.Value = LOCTEXT("NoFunc", "(No bound function)");
+			Value = LOCTEXT("NoFunc", "(No bound function)");
 		}
 
 		return;
 	}
 
-	ensure(false);
+	ensureMsgf(false, TEXT("Failed to identify property type. This function may need to be exanded to include it"));
+}
+
+TSharedPtr<FPropertyInstanceInfo> FPropertyInstanceInfo::FindOrMake(FPropertyInstance PropertyInstance,
+	TMap<FPropertyInstance, TSharedPtr<FPropertyInstanceInfo>>& VisitedNodes)
+{
+	if(TSharedPtr<FPropertyInstanceInfo>* Found = VisitedNodes.Find(PropertyInstance))
+	{
+		return *Found;
+	}
+	
+	// wait to populate children until after inserting into VisitedNodes.
+	// this way we will catch circular references
+	TSharedPtr<FPropertyInstanceInfo> DebugInfo = VisitedNodes.Add(
+		PropertyInstance,
+		MakeShared<FPropertyInstanceInfo>(PropertyInstance)
+	);
+
+	DebugInfo->PopulateChildren(PropertyInstance, VisitedNodes);
+	return DebugInfo;
+}
+
+void FPropertyInstanceInfo::PopulateChildren(FPropertyInstance PropertyInstance,
+	TMap<FPropertyInstance, TSharedPtr<FPropertyInstanceInfo>>& VisitedNodes)
+{
+	check(PropertyInstance.Property);
+	if (PropertyInstance.Value == nullptr)
+	{
+		return;
+	}
+
+	if (FArrayProperty* ArrayProperty = CastField<FArrayProperty>(Property.Get()))
+	{
+		checkSlow(ArrayProperty->Inner);
+
+		FScriptArrayHelper ArrayHelper(ArrayProperty, PropertyInstance.Value);
+		for (int32 i = 0; i < ArrayHelper.Num(); i++)
+		{
+			FPropertyInstance ChildProperty = {
+				ArrayProperty->Inner,
+				ArrayHelper.GetRawPtr(i)
+			};
+			const TSharedPtr<FPropertyInstanceInfo> ChildInfo = FindOrMake(ChildProperty, VisitedNodes);
+			
+			// overwrite the display name with the array index for the current element
+			ChildInfo->DisplayName = FText::Format(LOCTEXT("ArrayIndexName", "[{0}]"), FText::AsNumber(i));
+			Children.Add(ChildInfo);
+		}
+	}
+	else if (FStructProperty* StructProperty = CastField<FStructProperty>(Property.Get()))
+	{
+		for (TFieldIterator<FProperty> It(StructProperty->Struct); It; ++It)
+		{
+			FPropertyInstance ChildProperty = {
+				*It,
+				It->ContainerPtrToValuePtr<void>(PropertyInstance.Value, 0)
+			};
+			const TSharedPtr<FPropertyInstanceInfo> ChildInfo = FindOrMake(ChildProperty, VisitedNodes);
+			
+			Children.Add(ChildInfo);
+		}
+	}
+	else if (FMapProperty* MapProperty = CastField<FMapProperty>(Property.Get()))
+	{
+		FScriptMapHelper MapHelper(MapProperty, PropertyInstance.Value);
+		uint8* PropData = MapHelper.GetPairPtr(0);
+
+		int32 Index = 0;
+		for (int32 Count = MapHelper.Num(); Count; PropData += MapProperty->MapLayout.SetLayout.Size, ++Index)
+		{
+			if (MapHelper.IsValidIndex(Index))
+			{
+				FPropertyInstance ChildProperty = {
+					MapProperty->ValueProp,
+					PropData + MapProperty->MapLayout.ValueOffset
+				};
+				const TSharedPtr<FPropertyInstanceInfo> ChildInfo = FindOrMake(ChildProperty, VisitedNodes);
+				
+				FString NameStr = TEXT("[");
+				MapProperty->KeyProp->ExportTextItem(
+					NameStr,
+					PropData,
+					nullptr,
+					nullptr,
+					PPF_PropertyWindow | PPF_BlueprintDebugView | PPF_Delimited,
+					nullptr
+				);
+				NameStr += TEXT("] ");
+				ChildInfo->DisplayName = FText::FromString(NameStr);
+				
+				Children.Add(ChildInfo);
+				--Count;
+			}
+		}
+	}
+	else if (FSetProperty* SetProperty = CastField<FSetProperty>(Property.Get()))
+	{
+		FScriptSetHelper SetHelper(SetProperty, PropertyInstance.Value);
+		uint8* PropData = SetHelper.GetElementPtr(0);
+
+		int32 Index = 0;
+		for (int32 Count = SetHelper.Num(); Count; PropData += SetProperty->SetLayout.Size, ++Index)
+		{
+			if (SetHelper.IsValidIndex(Index))
+			{
+				FPropertyInstance ChildProperty = {
+					SetProperty->ElementProp,
+					PropData
+				};
+				const TSharedPtr<FPropertyInstanceInfo> ChildInfo = FindOrMake(ChildProperty, VisitedNodes);
+				
+				// members of sets don't have their own names
+				ChildInfo->DisplayName = FText::GetEmpty();
+
+				Children.Add(ChildInfo);
+
+				--Count;
+			}
+		}
+	}
+	else if (FObjectPropertyBase* ObjectPropertyBase = CastField<FObjectPropertyBase>(Property.Get()))
+	{
+		if (UObject* Obj = ObjectPropertyBase->GetObjectPropertyValue(PropertyInstance.Value))
+		{
+			for (TFieldIterator<FProperty> It(Obj->GetClass()); It; ++It)
+			{
+				if(It->HasAllPropertyFlags(CPF_BlueprintVisible))
+				{
+					FPropertyInstance ChildProperty = {
+						*It,
+						It->ContainerPtrToValuePtr<void*>(Obj)
+					};
+					const TSharedPtr<FPropertyInstanceInfo> ChildInfo = FindOrMake(ChildProperty, VisitedNodes);
+					Children.Add(ChildInfo);
+				}
+			}
+		}
+	}
 }
 
 FText FKismetDebugUtilities::GetAndClearLastExceptionMessage()
diff --git a/Engine/Source/Editor/UnrealEd/Private/PropertyInfoViewStyle.cpp b/Engine/Source/Editor/UnrealEd/Private/PropertyInfoViewStyle.cpp
new file mode 100644
--- /dev/null
+++ b/Engine/Source/Editor/UnrealEd/Private/PropertyInfoViewStyle.cpp
@@ -0,0 +1,270 @@
+﻿// Copyright Epic Games, Inc. All Rights Reserved.
+
+#include "PropertyInfoViewStyle.h"
+#include "Math/Color.h"
+
+namespace PropertyInfoViewStyle
+{
+	FSlateColor GetIndentBackgroundColor(int32 IndentLevel, bool IsHovered)
+	{
+		if (IsHovered)
+		{
+			return FAppStyle::Get().GetSlateColor("Colors.Header");
+		}
+
+		if (IndentLevel == 0)
+		{
+			return FAppStyle::Get().GetSlateColor("Colors.Panel");
+		}
+
+		int32 ColorIndex = 0;
+		int32 Increment = 1;
+
+		for (int i = 0; i < IndentLevel; ++i)
+		{
+			ColorIndex += Increment;
+
+			if (ColorIndex == 0 || ColorIndex == 3)
+			{
+				Increment = -Increment;
+			}
+		}
+
+		static const uint8 ColorOffsets[] =
+		{
+			2, 6, 12, 20
+		};
+
+		FColor BaseColor = FAppStyle::Get().GetSlateColor("Colors.Panel").GetSpecifiedColor().ToFColor(true);
+
+		FColor ColorWithOffset(
+		BaseColor.R + ColorOffsets[ColorIndex], 
+		BaseColor.G + ColorOffsets[ColorIndex], 
+		BaseColor.B + ColorOffsets[ColorIndex]);
+
+		return FSlateColor(FLinearColor::FromSRGBColor(ColorWithOffset));
+	}
+
+	FSlateColor GetRowBackgroundColor(ITableRow* Row)
+	{
+		check(Row);
+		return GetIndentBackgroundColor(Row->GetIndentLevel(), Row->AsWidget()->IsHovered());
+	}
+
+	void SConstrainedBox::Construct(const FArguments& InArgs)
+	{
+		MinWidth = InArgs._MinWidth;
+		MaxWidth = InArgs._MaxWidth;
+
+		ChildSlot
+		[
+			InArgs._Content.Widget
+		];
+	}
+
+	FVector2D SConstrainedBox::ComputeDesiredSize(float LayoutScaleMultiplier) const
+	{
+		const float MinWidthVal = MinWidth.Get().Get(0.0f);
+		const float MaxWidthVal = MaxWidth.Get().Get(0.0f);
+
+		if (MinWidthVal == 0.0f && MaxWidthVal == 0.0f)
+		{
+			return SCompoundWidget::ComputeDesiredSize(LayoutScaleMultiplier);
+		}
+		else
+		{
+			FVector2D ChildSize = ChildSlot.GetWidget()->GetDesiredSize();
+
+			float XVal = FMath::Max(MinWidthVal, ChildSize.X);
+			if (MaxWidthVal >= MinWidthVal)
+			{
+				XVal = FMath::Min(MaxWidthVal, XVal);
+			}
+
+			return FVector2D(XVal, ChildSize.Y);
+		}
+	}
+
+	int32 SIndent::OnPaint(const FPaintArgs& Args, const FGeometry& AllottedGeometry,
+							const FSlateRect& MyCullingRect, FSlateWindowElementList& OutDrawElements,
+							int32 LayerId, const FWidgetStyle& InWidgetStyle, bool bParentEnabled) const
+	{
+		constexpr int8 TabSize = 16;
+		
+		TSharedPtr<ITableRow> RowPtr = Row.Pin();
+		if (!RowPtr.IsValid())
+		{
+			return LayerId;
+		}
+
+		const FSlateBrush* BackgroundBrush = FAppStyle::Get().GetBrush("DetailsView.CategoryMiddle");
+		const FSlateBrush* DropShadowBrush = FAppStyle::Get().GetBrush("DetailsView.ArrayDropShadow");
+
+		int32 IndentLevel = RowPtr->GetIndentLevel();
+		for (int32 i = 0; i < IndentLevel; ++i)
+		{
+			FSlateColor BackgroundColor = GetRowBackgroundColor(i);
+
+			FSlateDrawElement::MakeBox(
+				OutDrawElements,
+				LayerId,
+				AllottedGeometry.ToPaintGeometry(FVector2D(TabSize * i, 0), FVector2D(TabSize, AllottedGeometry.GetLocalSize().Y)),
+				BackgroundBrush,
+				ESlateDrawEffect::None,
+				BackgroundColor.GetColor(InWidgetStyle)
+			);
+
+			FSlateDrawElement::MakeBox(
+				OutDrawElements,
+				LayerId + 1,
+				AllottedGeometry.ToPaintGeometry(FVector2D(TabSize * i, 0), FVector2D(TabSize, AllottedGeometry.GetLocalSize().Y)),
+				DropShadowBrush
+			);
+		}
+			
+		return LayerId + 1;
+	}
+
+
+	void SIndent::Construct(const FArguments& InArgs, TSharedRef<ITableRow> DetailsRow)
+	{
+		Row = DetailsRow;
+
+		ChildSlot
+		[
+			SNew(SBox)
+				.WidthOverride(this, &SIndent::GetIndentWidth)
+		];
+	}
+
+	FOptionalSize SIndent::GetIndentWidth() const
+	{
+		int32 IndentLevel = 0;
+
+		TSharedPtr<ITableRow> RowPtr = Row.Pin();
+		if (RowPtr.IsValid())
+		{
+			IndentLevel = RowPtr->GetIndentLevel();
+		}
+		return IndentLevel * 16.0f;
+	}
+
+	FSlateColor SIndent::GetRowBackgroundColor(int32 IndentLevel) const
+	{
+		const bool IsHovered = Row.IsValid() && Row.Pin()->AsWidget()->IsHovered();
+			
+		return PropertyInfoViewStyle::GetIndentBackgroundColor(IndentLevel, IsHovered);
+	}
+
+	void SExpanderArrow::Construct(const FArguments& InArgs, TSharedRef<ITableRow> DetailsRow)
+	{
+		Row = DetailsRow;
+
+		HasChildren = InArgs._HasChildren;
+
+		ChildSlot
+		[
+			SNew(SOverlay)
+				+ SOverlay::Slot()
+				[
+					SNew(SBorder)
+						.BorderImage(FAppStyle::Get().GetBrush("DetailsView.CategoryMiddle"))
+						.BorderBackgroundColor_Static(
+					             PropertyInfoViewStyle::GetRowBackgroundColor,
+					             Row.Pin().Get()
+				             )
+					[
+						SNew(SBox)
+							.WidthOverride(20.0f)
+							.HeightOverride(16.0f)
+					]
+				]
+				+ SOverlay::Slot()
+				[
+					SAssignNew(ExpanderArrow, SButton)
+						.ButtonStyle(FCoreStyle::Get(), "NoBorder")
+						.VAlign(VAlign_Center)
+						.HAlign(HAlign_Center)
+						.ClickMethod(EButtonClickMethod::MouseDown)
+						.OnClicked(this, &SExpanderArrow::OnExpanderClicked)
+						.ContentPadding(FMargin(0.0f))
+						.IsFocusable(false)
+					[
+						SNew(SImage)
+							.Image(this, &SExpanderArrow::GetExpanderImage)
+							.ColorAndOpacity(FSlateColor::UseForeground())
+					]
+				]
+		];
+	}
+
+	EVisibility SExpanderArrow::GetExpanderVisibility() const
+	{
+		TSharedPtr<ITableRow> RowPtr = Row.Pin();
+		if (!RowPtr.IsValid())
+		{
+			return EVisibility::Collapsed;
+		}
+
+		return RowPtr->DoesItemHaveChildren() ? EVisibility::Visible : EVisibility::Collapsed;
+	}
+
+	const FSlateBrush* SExpanderArrow::GetExpanderImage() const
+	{
+		TSharedPtr<ITableRow> RowPtr = Row.Pin();
+		if (!RowPtr.IsValid() || !HasChildren.Get())
+		{
+			return FAppStyle::Get().GetBrush("NoBrush");
+		}
+
+		const bool bIsItemExpanded = RowPtr->IsItemExpanded();
+
+		FName ResourceName;
+		if (bIsItemExpanded)
+		{
+			if (ExpanderArrow->IsHovered())
+			{
+				ResourceName = TEXT("TreeArrow_Expanded_Hovered");
+			}
+			else
+			{
+				ResourceName = TEXT("TreeArrow_Expanded");
+			}
+		}
+		else
+		{
+			if (ExpanderArrow->IsHovered())
+			{
+				ResourceName = TEXT("TreeArrow_Collapsed_Hovered");
+			}
+			else
+			{
+				ResourceName = TEXT("TreeArrow_Collapsed");
+			}
+		}
+
+		return FAppStyle::Get().GetBrush(ResourceName);
+	}
+
+	FReply SExpanderArrow::OnExpanderClicked() const
+	{
+		TSharedPtr<ITableRow> RowPtr = Row.Pin();
+		if (!RowPtr.IsValid())
+		{
+			return FReply::Unhandled();
+		}
+
+		// Recurse the expansion if "shift" is being pressed
+		const FModifierKeysState ModKeyState = FSlateApplication::Get().GetModifierKeys();
+		if (ModKeyState.IsShiftDown())
+		{
+			RowPtr->Private_OnExpanderArrowShiftClicked();
+		}
+		else
+		{
+			RowPtr->ToggleExpansion();
+		}
+
+		return FReply::Handled();
+	}
+}
diff --git a/Engine/Source/Editor/UnrealEd/Public/Kismet2/KismetDebugUtilities.h b/Engine/Source/Editor/UnrealEd/Public/Kismet2/KismetDebugUtilities.h
--- a/Engine/Source/Editor/UnrealEd/Public/Kismet2/KismetDebugUtilities.h
+++ b/Engine/Source/Editor/UnrealEd/Public/Kismet2/KismetDebugUtilities.h
@@ -30,14 +30,56 @@ struct FKismetTraceSample
 //////////////////////////////////////////////////////////////////////////
 // FDebugInfo
 
-struct FDebugInfo
+// call FKismetDebugUtilities::GetDebugInfo to construct
+struct FPropertyInstanceInfo
 {
+	/**
+	 * used to determine whether an object's property has
+	 * been visited yet when generating FDebugInfo
+	 */
+	struct FPropertyInstance
+	{
+		FProperty* Property;
+		const void* Value;
+	};
+	
+	/**
+	* Helper constructor. call FKismetDebugUtilities::GetDebugInfo or
+	* FKismetDebugUtilities::GetDebugInfoInternal instead
+	*/
+    FPropertyInstanceInfo(FPropertyInstance PropertyInstance);
+
+	/**
+	 * looks for an existing FPropertyInstance inside VisitedNodes
+	 * if one isn't found, it makes one, stores it in VisitedNodes, and returns it
+	 */
+	static TSharedPtr<FPropertyInstanceInfo> FindOrMake(FPropertyInstance PropertyInstance,
+		TMap<FPropertyInstance, TSharedPtr<FPropertyInstanceInfo>> &VisitedNodes);
+
+	/**
+	* populates a FDebugInfo::Children with sub-properties
+	*/
+	void PopulateChildren(FPropertyInstance PropertyInstance,
+		TMap<FPropertyInstance, TSharedPtr<FPropertyInstanceInfo>> &VisitedNodes);
+	
 	FText DisplayName;
 	FText Value;
 	FText Type;
-	TArray<FDebugInfo> Children;
+	TFieldPath<FProperty> Property;
+	TArray<TSharedPtr<FPropertyInstanceInfo>> Children;
 };
 
+inline uint32 GetTypeHash(const FPropertyInstanceInfo::FPropertyInstance& PropertyInstance)
+{
+	return HashCombine(GetTypeHash(PropertyInstance.Property), GetTypeHash(PropertyInstance.Value));
+}
+
+inline bool operator==(const FPropertyInstanceInfo::FPropertyInstance& A, const FPropertyInstanceInfo::FPropertyInstance& B)
+{
+	return A.Property == B.Property && A.Value == B.Value;
+}
+
+
 //////////////////////////////////////////////////////////////////////////
 // FObjectsBeingDebuggedIterator
 
@@ -269,7 +311,10 @@ public:
 	static EWatchTextResult GetWatchText(FString& OutWatchText, UBlueprint* Blueprint, UObject* ActiveObject, const UEdGraphPin* WatchPin);
 
 	// Gets the debug info for a specified site
-	static EWatchTextResult GetDebugInfo(FDebugInfo& OutDebugInfo, UBlueprint* Blueprint, UObject* ActiveObject, const UEdGraphPin* WatchPin);
+	static EWatchTextResult GetDebugInfo(TSharedPtr<FPropertyInstanceInfo> &OutDebugInfo, UBlueprint* Blueprint, UObject* ActiveObject, const UEdGraphPin* WatchPin);
+
+	// Retrieves Debug info from a Property and a pointer to it's associated data
+	static void GetDebugInfoInternal(TSharedPtr<FPropertyInstanceInfo> &DebugInfo, FProperty* Property, const void* PropertyValue);
 
 	//@TODO: Pretty lame way to handle this messaging, ideally the entire Info object gets pushed into the editor when intraframe debugging is triggered!
 	// This doesn't work properly if there is more than one blueprint editor open at once either (one will consume it, the others will be left in the cold)
@@ -278,8 +323,7 @@ protected:
 	static void CheckBreakConditions(UEdGraphNode* NodeStoppedAt, bool bHitBreakpoint, int32 BreakpointOffset, bool& InOutBreakExecution);
 	static void AttemptToBreakExecution(UBlueprint* BlueprintObj, const UObject* ActiveObject, const FFrame& StackFrame, const FBlueprintExceptionInfo& Info, UEdGraphNode* NodeStoppedAt, int32 DebugOpcodeOffset);
 
-	static void GetDebugInfo_InContainer(int32 Index, FDebugInfo& DebugInfo, FProperty* Property, const void* Data);
-	static void GetDebugInfoInternal(FDebugInfo& DebugInfo, FProperty* Property, const void* PropertyValue);
+	static void GetDebugInfo_InContainer(int32 Index, TSharedPtr<FPropertyInstanceInfo> &DebugInfo, FProperty* Property, const void* Data);
 
 	/**
 	* @brief	Helper function for converting between blueprint and debuggable data
diff --git a/Engine/Source/Editor/UnrealEd/Public/PropertyInfoViewStyle.h b/Engine/Source/Editor/UnrealEd/Public/PropertyInfoViewStyle.h
new file mode 100644
--- /dev/null
+++ b/Engine/Source/Editor/UnrealEd/Public/PropertyInfoViewStyle.h
@@ -0,0 +1,114 @@
+﻿// Copyright Epic Games, Inc. All Rights Reserved.
+
+#pragma once
+
+#include "CoreMinimal.h"
+#include "Framework/Application/SlateApplication.h"
+#include "Widgets/Images/SImage.h"
+#include "Widgets/Input/SButton.h"
+#include "Widgets/Views/STableRow.h"
+
+
+/**
+* A collection of widgets and helpers used for the style of the trees in SKismetDebuggingView
+*/
+namespace PropertyInfoViewStyle
+{
+
+	/**
+	 * Used to indent within stylized details tree to achieve a layered effect
+	 * see KismetDebugViewStyle::SIndent for usage
+	 * @param IndentLevel depth of the tree
+	 * @param IsHovered will give a lighter color if this line in the tree is hovered
+	 * @return the color to set the indent to
+	 */
+	UNREALED_API FSlateColor GetIndentBackgroundColor(int32 IndentLevel, bool IsHovered);
+
+	/**
+	 * calls KismetDebugViewStyle::GetIndentBackgroundColor using the indent level
+	 * and hover state of the provided table row
+	 * @param Row a table row - likely nested within a tree
+	 * @return the color to set the background of the row to
+	 */
+	UNREALED_API FSlateColor GetRowBackgroundColor(ITableRow* Row);
+
+	/** Helper class to force a widget to fill in a space. Copied from SDetailSingleItemRow.cpp */
+	class UNREALED_API SConstrainedBox : public SCompoundWidget
+	{
+	public:
+		SLATE_BEGIN_ARGS(SConstrainedBox){}
+			SLATE_DEFAULT_SLOT(FArguments, Content)
+			SLATE_ATTRIBUTE(TOptional<float>, MinWidth)
+			SLATE_ATTRIBUTE(TOptional<float>, MaxWidth)
+		SLATE_END_ARGS()
+
+		void Construct(const FArguments& InArgs);
+
+		virtual FVector2D ComputeDesiredSize(float LayoutScaleMultiplier) const override;
+
+	private:
+		TAttribute< TOptional<float> > MinWidth;
+		TAttribute< TOptional<float> > MaxWidth;
+	};
+
+	/**
+	 * SIndent is a widget used to indent trees in a layered
+	 * style. It's based on SDetailRowIndent but generalized to allow for use with
+	 * any ITableRow
+	 * 
+	 * see SDebugLineItem::GenerateWidgetForColumn (SKismetDebuggingView.cpp) for
+	 * use example
+	 */
+	class UNREALED_API SIndent : public SCompoundWidget
+	{ 
+	public:
+			
+		SLATE_BEGIN_ARGS(SIndent) {}
+		SLATE_END_ARGS()
+
+		void Construct(const FArguments& InArgs, TSharedRef<ITableRow> DetailsRow);
+
+	private:
+		virtual int32 OnPaint(const FPaintArgs& Args, const FGeometry& AllottedGeometry,
+		                      const FSlateRect& MyCullingRect, FSlateWindowElementList& OutDrawElements,
+		                      int32 LayerId, const FWidgetStyle& InWidgetStyle,
+		                      bool bParentEnabled) const override;
+
+		virtual FOptionalSize GetIndentWidth() const;
+
+		virtual FSlateColor GetRowBackgroundColor(int32 IndentLevel) const;
+
+	private:
+		TWeakPtr<ITableRow> Row;
+	};
+
+	/**
+	 * SExpanderArrow is a widget intended to be used alongside SIndent.
+	 * It's based on SDetailExpanderArrow but generalized to allow for use with
+	 * any ITableRow
+	 * 
+	 * see SDebugLineItem::GenerateWidgetForColumn (SKismetDebuggingView.cpp) for
+	 * use example
+	 */
+	class UNREALED_API SExpanderArrow : public SCompoundWidget
+	{
+	public:
+		SLATE_BEGIN_ARGS(SExpanderArrow) {}
+			SLATE_ATTRIBUTE(bool, HasChildren)
+		SLATE_END_ARGS()
+
+		void Construct(const FArguments& InArgs, TSharedRef<ITableRow> DetailsRow);
+
+	private:
+		EVisibility GetExpanderVisibility() const;
+
+		const FSlateBrush* GetExpanderImage() const;
+
+		FReply OnExpanderClicked() const;
+
+	private:
+		TWeakPtr<ITableRow> Row;
+		TSharedPtr<SButton> ExpanderArrow;
+		TAttribute<bool> HasChildren;
+	};
+}
