BP Compilation Manager: added additional matching conditions to find old CDO during reinstancing for control rig bp use case

diff --git a/Engine/Source/Editor/Kismet/Private/BlueprintCompilationManager.cpp b/Engine/Source/Editor/Kismet/Private/BlueprintCompilationManager.cpp
--- a/Engine/Source/Editor/Kismet/Private/BlueprintCompilationManager.cpp
+++ b/Engine/Source/Editor/Kismet/Private/BlueprintCompilationManager.cpp
@@ -2179,7 +2179,15 @@ void FBlueprintCompilationManagerImpl::ReinstanceBatch(TArray<FReinstancingJob>&
 					for (int32 i = 0; i < CurrentLinker->ExportMap.Num(); i++)
 					{
 						FObjectExport& ThisExport = CurrentLinker->ExportMap[i];
-						if (ThisExport.ObjectFlags & RF_ClassDefaultObject)
+						// In addition to checking for RF_ClassDefaultObject, we also need to check if the object names
+						// match to for Control Rig BP use case
+						// 
+						// In a normal Blueprint, there is usually just 1 CDO in the package, so name checking was not necessary.
+						// 
+						// But in a Control Rig BP, multiple CDO can be in a package because custom UClasses are used to describe 
+						// the layout of RigVM Memory Storage. These UClasses and their CDO are serialzed with the package as well.
+						// Thus, we have to ensure that the export is not just a CDO, but also a CDO with the matching name
+						if ((ThisExport.ObjectFlags & RF_ClassDefaultObject) && (ThisExport.ObjectName == CurrentBP->GeneratedClass->ClassDefaultObject->GetFName()))
 						{
 							OldCDOIndex = i;
 							break;
