Extend prototype BP namespace import UI to include auto-completion when typing into the entry box.

diff --git a/Engine/Source/Editor/Kismet/Private/BlueprintDetailsCustomization.cpp b/Engine/Source/Editor/Kismet/Private/BlueprintDetailsCustomization.cpp
--- a/Engine/Source/Editor/Kismet/Private/BlueprintDetailsCustomization.cpp
+++ b/Engine/Source/Editor/Kismet/Private/BlueprintDetailsCustomization.cpp
@@ -62,6 +62,7 @@
 #include "Kismet2/Kismet2NameValidators.h"
 #include "Widgets/Layout/SWidgetSwitcher.h"
 #include "BlueprintNamespaceHelper.h"
+#include "Widgets/Input/SSuggestionTextBox.h"
 
 #include "Modules/ModuleManager.h"
 #include "ISequencerModule.h"
@@ -84,6 +85,7 @@
 #include "Subsystems/AssetEditorSubsystem.h"
 #include "SupportedRangeTypes.h"	// StructsSupportingRangeVisibility
 #include "IPropertyAccessEditor.h"
+#include "AssetRegistryModule.h"
 
 #define LOCTEXT_NAMESPACE "BlueprintDetailsCustomization"
 
@@ -5464,6 +5466,54 @@ FBlueprintImportsLayout::FBlueprintImportsLayout(TWeakPtr<class FBlueprintGlobal
 	DisplayOptions.NoItemsLabelText = LOCTEXT("NoBlueprintImports", "No Imports");
 	DisplayOptions.ItemRowFilterText = LOCTEXT("BlueprintImportsValue", "Imports Value");
 	DisplayOptions.AddItemRowFilterText = LOCTEXT("BlueprintAddImport", "Add Import");
+
+	// @todo_namespaces - Consider moving this into a namespace "registry" so that it can be cached and hooked into AssetRegistry events.
+
+	UBlueprint* Blueprint = GetBlueprintObjectChecked();
+
+	auto CheckAndAddToAvailableNamespaceList = [this, Blueprint](const FString& Namespace)
+	{
+		if (!Namespace.IsEmpty() && Namespace != Blueprint->BlueprintNamespace && !Blueprint->ImportedNamespaces.Contains(Namespace))
+		{
+			AvailableNamespaces.Add(Namespace);
+		}
+	};
+
+	// Add namespaces for loaded class objects (including native).
+	for (TObjectIterator<UClass> LoadedClassIt; LoadedClassIt; ++LoadedClassIt)
+	{
+		const UClass* LoadedClass = *LoadedClassIt;
+
+		// Blueprint-based classes will be added below, so skip over them here.
+		if (!LoadedClass->HasAnyClassFlags(CLASS_CompiledFromBlueprint))
+		{
+			if (const FString* Namespace = LoadedClass->FindMetaData(FBlueprintMetadata::MD_Namespace))
+			{
+				CheckAndAddToAvailableNamespaceList(*Namespace);
+			}
+		}
+	}
+
+	// Add namespaces for loaded Blueprint objects. This will include macro and function library assets.
+	for (TObjectIterator<UBlueprint> LoadedBlueprintIt; LoadedBlueprintIt; ++LoadedBlueprintIt)
+	{
+		const UBlueprint* LoadedBlueprint = *LoadedBlueprintIt;
+		CheckAndAddToAvailableNamespaceList(LoadedBlueprint->BlueprintNamespace);
+	}
+
+	// Add namespaces for unloaded Blueprint assets. This will include unloaded macro and function library assets.
+	TArray<FAssetData> AssetDataArray;
+	FAssetRegistryModule& AssetRegistryModule = FModuleManager::LoadModuleChecked<FAssetRegistryModule>(TEXT("AssetRegistry"));
+	AssetRegistryModule.Get().GetAssetsByClass(UBlueprint::StaticClass()->GetFName(), AssetDataArray, true);
+	for (const FAssetData& AssetData : AssetDataArray)
+	{
+		if (AssetData.IsValid() && !AssetData.IsAssetLoaded())
+		{
+			// Note: This property is assumed to be marked AssetRegistrySearchable. It's value will be automatically added as an asset tag.
+			FString Namespace = AssetData.GetTagValueRef<FString>(GET_MEMBER_NAME_STRING_CHECKED(UBlueprint, BlueprintNamespace));
+			CheckAndAddToAvailableNamespaceList(Namespace);
+		}
+	}
 }
 
 TSharedPtr<SWidget> FBlueprintImportsLayout::MakeAddItemRowWidget()
@@ -5472,10 +5522,11 @@ TSharedPtr<SWidget> FBlueprintImportsLayout::MakeAddItemRowWidget()
 	+ SHorizontalBox::Slot()
 	.FillWidth(1.0f)
 	[
-		SAssignNew(ImportEntryTextBox, SEditableTextBox)
+		SAssignNew(ImportEntryTextBox, SSuggestionTextBox)
 		.MinDesiredWidth(200.0f)
 		.Font(IDetailLayoutBuilder::GetDetailFont())
 		.OnTextCommitted(this, &FBlueprintImportsLayout::HandleImportEntryTextCommitted)
+		.OnShowingSuggestions(this, &FBlueprintImportsLayout::OnShowingImportSuggestions)
 	]
 	+ SHorizontalBox::Slot()
 	.AutoWidth()
@@ -5518,7 +5569,8 @@ void FBlueprintImportsLayout::OnRemoveItem(const FManagedListItem& Item)
 
 bool FBlueprintImportsLayout::IsAddImportEntryButtonEnabled() const
 {
-	return !ImportEntryTextBox->GetText().IsEmptyOrWhitespace();
+	FString CurrentEntryAsString = ImportEntryTextBox->GetText().ToString();
+	return AvailableNamespaces.Contains(CurrentEntryAsString);
 }
 
 FReply FBlueprintImportsLayout::OnAddImportEntryButtonClicked()
@@ -5558,6 +5610,17 @@ void FBlueprintImportsLayout::HandleImportEntryTextCommitted(const FText& NewLab
 	}
 }
 
+void FBlueprintImportsLayout::OnShowingImportSuggestions(const FString& InputText, TArray<FString>& OutSuggestions)
+{
+	for (const FString& Namespace : AvailableNamespaces)
+	{
+		if (InputText.IsEmpty() || Namespace.StartsWith(InputText))
+		{
+			OutSuggestions.Add(Namespace);
+		}
+	}
+}
+
 
 FBlueprintInterfaceLayout::FBlueprintInterfaceLayout(TWeakPtr<class FBlueprintGlobalOptionsDetails> InGlobalOptionsDetails, bool bInShowsInheritedInterfaces)
 	:FBlueprintManagedListDetails(InGlobalOptionsDetails)
diff --git a/Engine/Source/Editor/Kismet/Private/BlueprintDetailsCustomization.h b/Engine/Source/Editor/Kismet/Private/BlueprintDetailsCustomization.h
--- a/Engine/Source/Editor/Kismet/Private/BlueprintDetailsCustomization.h
+++ b/Engine/Source/Editor/Kismet/Private/BlueprintDetailsCustomization.h
@@ -30,6 +30,7 @@ class SMultiLineEditableTextBox;
 class UEdGraphNode_Documentation;
 class UK2Node_Variable;
 class FSubobjectEditorTreeNode;
+class SSuggestionTextBox;
 
 /**
  * Variable network replication options.
@@ -756,9 +757,11 @@ protected:
 	bool IsAddImportEntryButtonEnabled() const;
 	FReply OnAddImportEntryButtonClicked();
 	void HandleImportEntryTextCommitted(const FText& NewLabel, ETextCommit::Type CommitType);
+	void OnShowingImportSuggestions(const FString& InputText, TArray<FString>& OutSuggestions);
 
 private:
-	TSharedPtr<SEditableTextBox> ImportEntryTextBox;
+	TSet<FString> AvailableNamespaces;
+	TSharedPtr<SSuggestionTextBox> ImportEntryTextBox;
 };
 
 /** Blueprint Interface List Details */
