Minor fix for namespace queries on a loaded Blueprint asset.

diff --git a/Engine/Source/Editor/Kismet/Private/BlueprintNamespaceHelper.cpp b/Engine/Source/Editor/Kismet/Private/BlueprintNamespaceHelper.cpp
--- a/Engine/Source/Editor/Kismet/Private/BlueprintNamespaceHelper.cpp
+++ b/Engine/Source/Editor/Kismet/Private/BlueprintNamespaceHelper.cpp
@@ -257,10 +257,17 @@ bool FBlueprintNamespaceHelper::IsImportedObject(const UObject* InObject) const
 	{
 		return IsImportedType(Type);
 	}
-	else
+	else if (const UBlueprint* Blueprint = Cast<UBlueprint>(InObject))
+	{
+		return IsImportedType(Blueprint->GeneratedClass);
+	}
+	else if(InObject)
 	{
 		return IsImportedType(InObject->GetClass());
 	}
+
+	// Objects exist in the global scope if we can't determine otherwise, which means it's always imported.
+	return true;
 }
 
 bool FBlueprintNamespaceHelper::IsImportedObject(const FSoftObjectPath& InObjectPath) const
