Extend the prototype namespace import UI to list any project-wide global imports as non-removable default items.

diff --git a/Engine/Source/Editor/Kismet/Private/BlueprintDetailsCustomization.cpp b/Engine/Source/Editor/Kismet/Private/BlueprintDetailsCustomization.cpp
--- a/Engine/Source/Editor/Kismet/Private/BlueprintDetailsCustomization.cpp
+++ b/Engine/Source/Editor/Kismet/Private/BlueprintDetailsCustomization.cpp
@@ -86,6 +86,7 @@
 #include "SupportedRangeTypes.h"	// StructsSupportingRangeVisibility
 #include "IPropertyAccessEditor.h"
 #include "AssetRegistryModule.h"
+#include "Settings/EditorProjectSettings.h"
 
 #define LOCTEXT_NAMESPACE "BlueprintDetailsCustomization"
 
@@ -5473,7 +5474,10 @@ FBlueprintImportsLayout::FBlueprintImportsLayout(TWeakPtr<class FBlueprintGlobal
 
 	auto CheckAndAddToAvailableNamespaceList = [this, Blueprint](const FString& Namespace)
 	{
-		if (!Namespace.IsEmpty() && Namespace != Blueprint->BlueprintNamespace && !Blueprint->ImportedNamespaces.Contains(Namespace))
+		if (!Namespace.IsEmpty()
+			&& Namespace != Blueprint->BlueprintNamespace
+			&& !Blueprint->ImportedNamespaces.Contains(Namespace)
+			&& !GetDefault<UBlueprintEditorProjectSettings>()->NamespacesToAlwaysInclude.Contains(Namespace))
 		{
 			AvailableNamespaces.Add(Namespace);
 		}
@@ -5541,17 +5545,27 @@ TSharedPtr<SWidget> FBlueprintImportsLayout::MakeAddItemRowWidget()
 
 void FBlueprintImportsLayout::GetManagedListItems(TArray<FManagedListItem>& OutListItems) const
 {
-	UBlueprint* Blueprint = GetBlueprintObjectChecked();
-
-	for (const FString& ImportedNamespace : Blueprint->ImportedNamespaces)
+	auto AddNamespaceItemsToOutputList = [&OutListItems](const TArray<FString>& NamespaceArray, bool bIsRemovable)
 	{
-		FManagedListItem ItemDesc;
-		ItemDesc.ItemName = ImportedNamespace;
-		ItemDesc.DisplayName = FText::FromString(ImportedNamespace);
-		ItemDesc.bIsRemovable = true;
+		for (const FString& Namespace : NamespaceArray)
+		{
+			FManagedListItem ItemDesc;
+			ItemDesc.ItemName = Namespace;
+			ItemDesc.DisplayName = FText::FromString(Namespace);
+			ItemDesc.bIsRemovable = bIsRemovable;
 
-		OutListItems.Add(MoveTemp(ItemDesc));
-	}
+			OutListItems.Add(MoveTemp(ItemDesc));
+		}
+	};
+
+	// Project imports (static).
+	bool bIsRemovable = false;
+	AddNamespaceItemsToOutputList(GetDefault<UBlueprintEditorProjectSettings>()->NamespacesToAlwaysInclude, bIsRemovable);
+
+	// Blueprint imports (removable).
+	bIsRemovable = true;
+	const UBlueprint* Blueprint = GetBlueprintObjectChecked();
+	AddNamespaceItemsToOutputList(Blueprint->ImportedNamespaces.Array(), bIsRemovable);
 }
 
 void FBlueprintImportsLayout::OnRemoveItem(const FManagedListItem& Item)
