Updated the Timline Track Icons and Toolbar, and moved some icons used by multiple features to the common library

diff --git a/Engine/Source/Editor/EditorStyle/Private/StarshipStyle.cpp b/Engine/Source/Editor/EditorStyle/Private/StarshipStyle.cpp
--- a/Engine/Source/Editor/EditorStyle/Private/StarshipStyle.cpp
+++ b/Engine/Source/Editor/EditorStyle/Private/StarshipStyle.cpp
@@ -3996,8 +3996,8 @@ void FStarshipEditorStyle::FStyle::SetupGraphEditorStyles()
 
 		Set("ClassIcon.K2Node_CallFunction",	new IMAGE_BRUSH_SVG("Starship/Blueprints/icon_Blueprint_Function", Icon16x16));
 		Set("ClassIcon.K2Node_FunctionEntry",	new IMAGE_BRUSH_SVG("Starship/Blueprints/icon_Blueprint_Function", Icon16x16));
-		Set("ClassIcon.K2Node_CustomEvent",		new IMAGE_BRUSH_SVG("Starship/GraphEditors/Event", Icon16x16));
-		Set("ClassIcon.K2Node_Event",			new IMAGE_BRUSH_SVG("Starship/GraphEditors/Event", Icon16x16));
+		Set("ClassIcon.K2Node_CustomEvent",		new IMAGE_BRUSH_SVG("Starship/Common/Event", Icon16x16));
+		Set("ClassIcon.K2Node_Event",			new IMAGE_BRUSH_SVG("Starship/Common/Event", Icon16x16));
 		Set("ClassIcon.K2Node_Variable",		new IMAGE_BRUSH_SVG("Starship/GraphEditors/Node", Icon16x16, FLinearColor::White));
 		Set("ClassIcon.K2Node_VariableGet",		new IMAGE_BRUSH_SVG("Starship/GraphEditors/VarGet", Icon16x16, FLinearColor::White));
 		Set("ClassIcon.K2Node_VariableSet",		new IMAGE_BRUSH_SVG("Starship/GraphEditors/VarSet", Icon16x16, FLinearColor::White));
@@ -4029,8 +4029,8 @@ void FStarshipEditorStyle::FStyle::SetupGraphEditorStyles()
 		Set( "GraphEditor.Rule_16x", new IMAGE_BRUSH_SVG( "Starship/GraphEditors/Rule", Icon16x16 ) );
 		Set( "GraphEditor.State_16x", new IMAGE_BRUSH_SVG( "Starship/GraphEditors/State", Icon16x16 ) );
 		Set( "GraphEditor.StateMachine_16x", new IMAGE_BRUSH_SVG( "Starship/GraphEditors/StateMachine", Icon16x16 ) );
-		Set( "GraphEditor.Event_16x", new IMAGE_BRUSH_SVG( "Starship/GraphEditors/Event", Icon16x16 ) );
-		Set( "GraphEditor.CustomEvent_16x", new IMAGE_BRUSH_SVG( "Starship/GraphEditors/Event", Icon16x16 ) );
+		Set( "GraphEditor.Event_16x", new IMAGE_BRUSH_SVG( "Starship/Common/Event", Icon16x16 ) );
+		Set( "GraphEditor.CustomEvent_16x", new IMAGE_BRUSH_SVG( "Starship/Common/Event", Icon16x16 ) );
 		Set( "GraphEditor.CallInEditorEvent_16x", new IMAGE_BRUSH_SVG( "Starship/GraphEditors/CallInEditorEvent", Icon16x16 ) );
 		Set( "GraphEditor.Timeline_16x", new IMAGE_BRUSH_SVG("Starship/Common/Timecode", Icon16x16));
 		Set( "GraphEditor.Comment_16x", new IMAGE_BRUSH_SVG("Starship/GraphEditors/Comment", Icon16x16));
@@ -4052,7 +4052,7 @@ void FStarshipEditorStyle::FStyle::SetupGraphEditorStyles()
 		Set( "GraphEditor.Select_16x", new IMAGE_BRUSH_SVG("Starship/GraphEditors/Select", Icon16x16));
 		Set( "GraphEditor.Cast_16x", new IMAGE_BRUSH_SVG("Starship/GraphEditors/Cast", Icon16x16));
 
-		Set( "GraphEditor.Macro.Loop_16x", new IMAGE_BRUSH_SVG("Starship/GraphEditors/Loop", Icon16x16));
+		Set( "GraphEditor.Macro.Loop_16x", new IMAGE_BRUSH_SVG("Starship/Common/Loop", Icon16x16));
 		Set( "GraphEditor.Macro.Gate_16x", new IMAGE_BRUSH_SVG("Starship/GraphEditors/Gate", Icon16x16));
 		Set( "GraphEditor.Macro.DoN_16x", new IMAGE_BRUSH_SVG("Starship/GraphEditors/DoN", Icon16x16));
 		Set( "GraphEditor.Macro.DoOnce_16x", new IMAGE_BRUSH_SVG("Starship/GraphEditors/DoOnce", Icon16x16));
@@ -4077,7 +4077,7 @@ void FStarshipEditorStyle::FStyle::SetupGraphEditorStyles()
 		Set( "GraphEditor.StructGlyph", new IMAGE_BRUSH_SVG( "Starship/Common/Struct", Icon22x22, FLinearColor::White) );
 
 		// Find In Blueprints
-		Set( "GraphEditor.FIB_Event", new IMAGE_BRUSH_SVG( "Starship/GraphEditors/Event", Icon22x22, FLinearColor::White) );
+		Set( "GraphEditor.FIB_Event", new IMAGE_BRUSH_SVG( "Starship/Common/Event", Icon22x22, FLinearColor::White) );
 
 		Set( "GraphEditor.GetSequenceBinding", new IMAGE_BRUSH_SVG("Starship/GraphEditors/GetSequenceBinding", Icon16x16));
 
@@ -4145,12 +4145,17 @@ void FStarshipEditorStyle::FStyle::SetupGraphEditorStyles()
 
 		// Timeline Editor
 		{
-			Set( "TimelineEditor.AddFloatTrack", new IMAGE_BRUSH( "Icons/icon_TrackAddFloat_36x24px", Icon36x24 ) );
-			Set( "TimelineEditor.AddVectorTrack", new IMAGE_BRUSH( "Icons/icon_TrackAddVector_36x24px", Icon36x24 ) );
-			Set( "TimelineEditor.AddEventTrack", new IMAGE_BRUSH( "Icons/icon_TrackAddEvent_36x24px", Icon36x24 ) );
-			Set( "TimelineEditor.AddColorTrack", new IMAGE_BRUSH( "Icons/icon_TrackAddColor_36x24px", Icon36x24 ) );
-			Set( "TimelineEditor.AddCurveAssetTrack", new IMAGE_BRUSH( "Icons/icon_TrackAddCurve_36x24px", Icon36x24 ) );
-			Set( "TimelineEditor.DeleteTrack", new IMAGE_BRUSH( "Icons/icon_TrackDelete_36x24px", Icon36x24 ) );
+			Set( "TimelineEditor.AddFloatTrack", new IMAGE_BRUSH_SVG( "Starship/TimelineEditor/TrackTypeFloat", Icon16x16) );
+			Set( "TimelineEditor.AddVectorTrack", new IMAGE_BRUSH_SVG( "Starship/TimelineEditor/TrackTypeVector", Icon16x16) );
+			Set( "TimelineEditor.AddEventTrack", new IMAGE_BRUSH_SVG( "Starship/Common/Event", Icon16x16) );
+			Set( "TimelineEditor.AddColorTrack", new IMAGE_BRUSH_SVG( "Starship/TimelineEditor/TrackTypeColor", Icon16x16) );
+			Set( "TimelineEditor.AddCurveAssetTrack", new IMAGE_BRUSH_SVG("Starship/TimelineEditor/TrackTypeCurve", Icon16x16));
+			Set( "TimelineEditor.DeleteTrack", new IMAGE_BRUSH_SVG( "Icons/icon_TrackDelete_36x24px", Icon36x24 ) );
+			Set("TimelineEditor.AutoPlay", new IMAGE_BRUSH_SVG("Starship/TimelineEditor/TimelineAutoPlay", Icon20x20));
+			Set("TimelineEditor.IgnoreTimeDilation", new IMAGE_BRUSH_SVG("Starship/TimelineEditor/TimelineIgnoreTimeDilation", Icon20x20));
+			Set("TimelineEditor.Replicated", new IMAGE_BRUSH_SVG("Starship/TimelineEditor/TimelineReplicated", Icon20x20));
+			Set("TimelineEditor.UseLastKeyframe", new IMAGE_BRUSH_SVG("Starship/TimelineEditor/TimelineUseLastKeyframe", Icon20x20));
+			Set("TimelineEditor.Loop", new IMAGE_BRUSH_SVG("Starship/Common/Loop", Icon20x20));
 
 			Set("TimelineEditor.TrackRowSubtleHighlight", FTableRowStyle(NormalTableRowStyle)
 				.SetActiveBrush(FSlateColorBrush(FStyleColors::Panel))
diff --git a/Engine/Source/Editor/Kismet/Private/STimelineEditor.cpp b/Engine/Source/Editor/Kismet/Private/STimelineEditor.cpp
--- a/Engine/Source/Editor/Kismet/Private/STimelineEditor.cpp
+++ b/Engine/Source/Editor/Kismet/Private/STimelineEditor.cpp
@@ -13,6 +13,7 @@
 #include "Widgets/Input/SCheckBox.h"
 #include "Widgets/Input/SSlider.h"
 #include "Widgets/Input/STextComboBox.h"
+#include "SEditorHeaderButton.h"
 #include "Styling/CoreStyle.h"
 #include "EditorStyleSet.h"
 #include "Curves/CurveFloat.h"
@@ -1037,84 +1038,13 @@ void STimelineEditor::Construct(const FArguments& InArgs, TSharedPtr<FBlueprintE
 			SNew(SHorizontalBox)
 			+SHorizontalBox::Slot()
 			.AutoWidth()
-			.Padding(2.f)
-			[
-				// Add float track button
-				SNew(SButton)
-				.ButtonStyle(FAppStyle::Get(), "SimpleButton")
-				[
-					SNew(SImage)
-					.Image(FAppStyle::Get().GetBrush("TimelineEditor.AddFloatTrack"))
-					.ColorAndOpacity(FSlateColor::UseForeground())
-				]
-				.ToolTipText( LOCTEXT( "AddFloatTrack", "Add Float Track" ) )
-				.OnClicked( this, &STimelineEditor::CreateNewTrack, FTTTrackBase::TT_FloatInterp )
-				.AddMetaData<FTagMetaData>(TEXT("TimelineEditor.AddFloatTrack"))
-			]
-			+SHorizontalBox::Slot()
-			.AutoWidth()
-			.Padding(2.f)
-			[
-				// Add vector track button
-				SNew(SButton)
-				.ButtonStyle(FAppStyle::Get(), "SimpleButton")
-				[
-					SNew(SImage)
-					.Image(FAppStyle::Get().GetBrush("TimelineEditor.AddVectorTrack"))
-					.ColorAndOpacity(FSlateColor::UseForeground())
-				]
-				.ToolTipText( LOCTEXT( "AddVectorTrack", "Add Vector Track" ) )
-				.OnClicked( this, &STimelineEditor::CreateNewTrack, FTTTrackBase::TT_VectorInterp )
-				.AddMetaData<FTagMetaData>(TEXT("TimelineEditor.AddVectorTrack"))
-			]
-			+SHorizontalBox::Slot()
-			.AutoWidth()
-			.Padding(2.f)
-			[
-				// Add event track button
-				SNew(SButton)
-				.ButtonStyle(FAppStyle::Get(), "SimpleButton")
-				[
-					SNew(SImage)
-					.Image(FAppStyle::Get().GetBrush("TimelineEditor.AddEventTrack"))
-					.ColorAndOpacity(FSlateColor::UseForeground())
-				]
-				.ToolTipText( LOCTEXT( "AddEventTrack", "Add Event Track" ) )
-				.OnClicked( this, &STimelineEditor::CreateNewTrack, FTTTrackBase::TT_Event )
-				.AddMetaData<FTagMetaData>(TEXT("TimelineEditor.AddEventTrack"))
-			]
-			+SHorizontalBox::Slot()
-			.AutoWidth()
-			.Padding(2.f)
-			[
-				// Add color track button
-				SNew(SButton)
-				.ButtonStyle(FAppStyle::Get(), "SimpleButton")
-				[
-					SNew(SImage)
-					.Image(FAppStyle::Get().GetBrush("TimelineEditor.AddColorTrack"))
-					.ColorAndOpacity(FSlateColor::UseForeground())
-				]
-				.ToolTipText( LOCTEXT( "AddColorTrack", "Add Color Track" ) )
-				.OnClicked( this, &STimelineEditor::CreateNewTrack, FTTTrackBase::TT_LinearColorInterp )
-				.AddMetaData<FTagMetaData>(TEXT("TimelineEditor.AddColorTrack"))
-			]
-			+SHorizontalBox::Slot()
-			.AutoWidth()
-			.Padding(2.f)
+			.VAlign(VAlign_Center)
+			.Padding(6.f)
 			[
-				// Add external curve asset button
-				SNew(SButton)
-				.ButtonStyle(FAppStyle::Get(), "SimpleButton")
-				[
-					SNew(SImage)
-					.Image(FAppStyle::Get().GetBrush("TimelineEditor.AddCurveAssetTrack"))
-					.ColorAndOpacity(FSlateColor::UseForeground())
-				]
-				.ToolTipText( LOCTEXT( "AddExternalAsset", "Add Selected Curve Asset" ) )
-				.IsEnabled( this, &STimelineEditor::IsCurveAssetSelected )
-				.OnClicked( this, &STimelineEditor::CreateNewTrackFromAsset )
-				.AddMetaData<FTagMetaData>(TEXT("TimelineEditor.AddCurveAssetTrack"))
+				SNew(SEditorHeaderButton)
+				.OnGetMenuContent(this, &STimelineEditor::MakeAddButton)
+				.Icon(FAppStyle::Get().GetBrush("Icons.Plus"))
+				.Text(LOCTEXT("Track", "Track"))
 			]
 			+SHorizontalBox::Slot()
 			.AutoWidth()
@@ -1127,7 +1057,8 @@ void STimelineEditor::Construct(const FArguments& InArgs, TSharedPtr<FBlueprintE
 			]
 			+SHorizontalBox::Slot()
 			.AutoWidth()
-			.Padding(6.f)
+			.Padding(FMargin(6.0f, 2.0f, 2.0f, 2.0f))
+			.VAlign(VAlign_Center)
 			[
 				// Length edit box
 				SAssignNew(TimelineLengthEdit, SEditableTextBox)
@@ -1140,65 +1071,92 @@ void STimelineEditor::Construct(const FArguments& InArgs, TSharedPtr<FBlueprintE
 			+SHorizontalBox::Slot()
 			.AutoWidth()
 			.Padding(2.f)
+			.VAlign(VAlign_Center)
 			[
 				// Use last keyframe as length check box
 				SAssignNew(UseLastKeyframeCheckBox, SCheckBox)
 				.IsChecked( this, &STimelineEditor::IsUseLastKeyframeChecked )
 				.OnCheckStateChanged( this, &STimelineEditor::OnUseLastKeyframeChanged )
+				.Style(FAppStyle::Get(), "ToggleButtonCheckbox")
+				.ToolTipText(LOCTEXT("UseLastKeyframe", "Use Last Keyframe"))
 				[
-					SNew(STextBlock) .Text( LOCTEXT( "UseLastKeyframe", "Use Last Keyframe" ) )
+					SNew(SImage)
+					.ColorAndOpacity(FSlateColor::UseForeground())
+					.Image(FAppStyle::Get().GetBrush("TimelineEditor.UseLastKeyframe"))
 					.AddMetaData<FTagMetaData>(TEXT("TimelineEditor.UseLastKeyframe"))
 				]
 			]
 			+SHorizontalBox::Slot()
 			.AutoWidth()
 			.Padding(2.f)
+			.VAlign(VAlign_Center)
 			[
 				// Play check box
 				SAssignNew(PlayCheckBox, SCheckBox)
 				.IsChecked( this, &STimelineEditor::IsAutoPlayChecked )
 				.OnCheckStateChanged( this, &STimelineEditor::OnAutoPlayChanged )
+				.Style(FAppStyle::Get(), "ToggleButtonCheckbox")
+				.ToolTipText(LOCTEXT("AutoPlay", "AutoPlay"))
 				[
-					SNew(STextBlock) .Text( LOCTEXT( "AutoPlay", "AutoPlay" ) )
+					SNew(SImage)
+					.ColorAndOpacity(FSlateColor::UseForeground())
+					.Image(FAppStyle::Get().GetBrush("TimelineEditor.AutoPlay"))
 					.AddMetaData<FTagMetaData>(TEXT("TimelineEditor.AutoPlay"))
 				]
 			]
 			+SHorizontalBox::Slot()
 			.AutoWidth()
 			.Padding(2.f)
+			.VAlign(VAlign_Center)
 			[
 				// Loop check box
 				SAssignNew(LoopCheckBox, SCheckBox)
 				.IsChecked( this, &STimelineEditor::IsLoopChecked )
 				.OnCheckStateChanged( this, &STimelineEditor::OnLoopChanged )
+				.Style(FAppStyle::Get(), "ToggleButtonCheckbox")
+				.ToolTipText(LOCTEXT("Loop", "Loop"))
 				[
-					SNew(STextBlock) .Text( LOCTEXT( "Loop", "Loop" ) )
+					SNew(SImage)
+					.ColorAndOpacity(FSlateColor::UseForeground())
+					.Image(FAppStyle::Get().GetBrush("TimelineEditor.Loop"))
 					.AddMetaData<FTagMetaData>(TEXT("TimelineEditor.Loop"))
 				]
+				
 			]
 			+SHorizontalBox::Slot()
 			.AutoWidth()
 			.Padding(2.f)
+			.VAlign(VAlign_Center)
 			[
 				// Replicated check box
 				SAssignNew(ReplicatedCheckBox, SCheckBox)
 				.IsChecked( this, &STimelineEditor::IsReplicatedChecked )
 				.OnCheckStateChanged( this, &STimelineEditor::OnReplicatedChanged )
+				.Style(FAppStyle::Get(), "ToggleButtonCheckbox")
+				.ToolTipText(LOCTEXT("Replicated", "Replicated"))
 				[
-					SNew(STextBlock) .Text( LOCTEXT( "Replicated", "Replicated" ) )
+					SNew(SImage)
+					.ColorAndOpacity(FSlateColor::UseForeground())
+					.Image(FAppStyle::Get().GetBrush("TimelineEditor.Replicated"))
 					.AddMetaData<FTagMetaData>(TEXT("TimelineEditor.Replicated"))
 				]
+			
 			]
 			+SHorizontalBox::Slot()
 			.AutoWidth()
 			.Padding(2.f)
+			.VAlign(VAlign_Center)
 			[
 				// Ignore Time Dilation check box
 				SAssignNew(IgnoreTimeDilationCheckBox, SCheckBox)
 				.IsChecked( this, &STimelineEditor::IsIgnoreTimeDilationChecked )
 				.OnCheckStateChanged( this, &STimelineEditor::OnIgnoreTimeDilationChanged )
+				.Style(FAppStyle::Get(), "ToggleButtonCheckbox")
+				.ToolTipText(LOCTEXT("IgnoreTimeDilation", "Ignore Time Dilation"))
 				[
-					SNew(STextBlock) .Text( LOCTEXT( "IgnoreTimeDilation", "Ignore Time Dilation" ) )
+					SNew(SImage)
+					.ColorAndOpacity(FSlateColor::UseForeground())
+					.Image(FAppStyle::Get().GetBrush("TimelineEditor.IgnoreTimeDilation"))
 					.AddMetaData<FTagMetaData>(TEXT("TimelineEditor.IgnoreTimeDilation"))
 				]
 			]
@@ -1206,21 +1164,22 @@ void STimelineEditor::Construct(const FArguments& InArgs, TSharedPtr<FBlueprintE
 			+ SHorizontalBox::Slot()
 			.AutoWidth()
 			.Padding(2.f)
+			.VAlign(VAlign_Center)
 			[
-				SNew(STextComboBox)
-				.OptionsSource(&TickGroupNameStrings)
-				.InitiallySelectedItem(TickGroupNameStrings[CurrentTickGroupNameStringIndex])
-				.OnSelectionChanged(this, &STimelineEditor::OnTimelineTickGroupChanged)
-				.ToolTipText(LOCTEXT("TimelineTickGroupDropdownTooltip", "Select the TickGroup you want this timeline to run in.\nTo assign options use context menu on timelines."))
+				SNew(STextBlock)
+				.Text(LOCTEXT("TickGroupLabel", "Tick Group"))
+				.AddMetaData<FTagMetaData>(TEXT("TimelineEditor.TickGroup"))
 			]
 			+ SHorizontalBox::Slot()
 			.AutoWidth()
-			.Padding(2.f)
+			.Padding(6.f)
 			.VAlign(VAlign_Center)
 			[
-				SNew(STextBlock)
-				.Text(LOCTEXT("TickGroupLabel", "Tick Group"))
-				.AddMetaData<FTagMetaData>(TEXT("TimelineEditor.TickGroup"))
+				SNew(STextComboBox)
+				.OptionsSource(&TickGroupNameStrings)
+				.InitiallySelectedItem(TickGroupNameStrings[CurrentTickGroupNameStringIndex])
+				.OnSelectionChanged(this, &STimelineEditor::OnTimelineTickGroupChanged)
+				.ToolTipText(LOCTEXT("TimelineTickGroupDropdownTooltip", "Select the TickGroup you want this timeline to run in.\nTo assign options use context menu on timelines."))
 			]
 		]
 		+SVerticalBox::Slot()
@@ -1330,7 +1289,7 @@ TSharedRef<ITableRow> STimelineEditor::MakeTrackWidget( TSharedPtr<FTimelineEdTr
 	];
 }
 
-FReply STimelineEditor::CreateNewTrack(FTTTrackBase::ETrackType Type)
+void STimelineEditor::CreateNewTrack(FTTTrackBase::ETrackType Type)
 {
 	FName TrackName;
 	do
@@ -1429,8 +1388,6 @@ FReply STimelineEditor::CreateNewTrack(FTTTrackBase::ETrackType Type)
 			Notification->SetCompletionState( SNotificationItem::CS_Fail );
 		}
 	}
-
-	return FReply::Handled();
 }
 
 UCurveBase* STimelineEditor::CreateNewCurve(FTTTrackBase::ETrackType Type )
@@ -1845,7 +1802,7 @@ bool STimelineEditor::IsCurveAssetSelected() const
 }
 
 
-FReply STimelineEditor::CreateNewTrackFromAsset()
+void STimelineEditor::CreateNewTrackFromAsset()
 {
 	FEditorDelegates::LoadSelectedAssetsIfNeeded.Broadcast();
 	UCurveBase* SelectedObj = GEditor->GetSelectedObjects()->GetTop<UCurveBase>();
@@ -1906,7 +1863,6 @@ FReply STimelineEditor::CreateNewTrackFromAsset()
 		TimelineNode->ReconstructNode();
 		Kismet2->RefreshEditors();
 	}
-	return FReply::Handled();
 }
 
 bool STimelineEditor::CanRenameSelectedTrack() const
@@ -1950,6 +1906,46 @@ TSharedPtr< SWidget > STimelineEditor::MakeContextMenu() const
 	return MenuBuilder.MakeWidget();
 }
 
+
+TSharedRef<SWidget> STimelineEditor::MakeAddButton()
+{
+	FMenuBuilder MenuBuilder(true, nullptr);
+
+	MenuBuilder.AddMenuEntry(
+		LOCTEXT("AddFloatTrack", "Add Float Track"),
+		LOCTEXT("AddFloatTrackToolTip", "Adds a Float Track."),
+		FSlateIcon(FAppStyle::GetAppStyleSetName(), "TimelineEditor.AddFloatTrack"),
+		FUIAction(FExecuteAction::CreateRaw(this, &STimelineEditor::CreateNewTrack, FTTTrackBase::TT_FloatInterp)));
+
+	MenuBuilder.AddMenuEntry(
+		LOCTEXT("AddVectorTrack", "Add Vector Track"),
+		LOCTEXT("AddVectorTrackToolTip", "Adds a Vector Track."),
+		FSlateIcon(FAppStyle::GetAppStyleSetName(), "TimelineEditor.AddVectorTrack"),
+		FUIAction(FExecuteAction::CreateRaw(this, &STimelineEditor::CreateNewTrack, FTTTrackBase::TT_VectorInterp)));
+	
+	MenuBuilder.AddMenuEntry(
+		LOCTEXT("AddEventTrack", "Add Event Track"),
+		LOCTEXT("AddEventTrackToolTip", "Adds an Event Track."),
+		FSlateIcon(FAppStyle::GetAppStyleSetName(), "TimelineEditor.AddEventTrack"),
+		FUIAction(FExecuteAction::CreateRaw(this, &STimelineEditor::CreateNewTrack, FTTTrackBase::TT_Event)));
+
+	MenuBuilder.AddMenuEntry(
+		LOCTEXT("AddColorTrack", "Add Color Track"),
+		LOCTEXT("AddColorTrackToolTip", "Adds a Color Track."),
+		FSlateIcon(FAppStyle::GetAppStyleSetName(), "TimelineEditor.AddColorTrack"),
+		FUIAction(FExecuteAction::CreateRaw(this, &STimelineEditor::CreateNewTrack, FTTTrackBase::TT_LinearColorInterp)));
+
+	FUIAction AddCurveAssetAction(FExecuteAction::CreateRaw(this, &STimelineEditor::CreateNewTrackFromAsset), FCanExecuteAction::CreateRaw(this, &STimelineEditor::IsCurveAssetSelected));
+
+	MenuBuilder.AddMenuEntry(
+		LOCTEXT("AddExternalAsset", "Add Selected Curve Asset"),
+		LOCTEXT("AddExternalAssetToolTip", "Add the currently selected curve asset."),
+		FSlateIcon(FAppStyle::GetAppStyleSetName(), "TimelineEditor.AddCurveAssetTrack"),
+		AddCurveAssetAction);
+
+	return MenuBuilder.MakeWidget();
+}
+
 FVector2D STimelineEditor::GetTimelineDesiredSize() const
 {
 	return TimelineDesiredSize;
diff --git a/Engine/Source/Editor/Kismet/Private/STimelineEditor.h b/Engine/Source/Editor/Kismet/Private/STimelineEditor.h
--- a/Engine/Source/Editor/Kismet/Private/STimelineEditor.h
+++ b/Engine/Source/Editor/Kismet/Private/STimelineEditor.h
@@ -289,7 +289,7 @@ private:
 	/** Used by list view to create a track widget from the track item struct */
 	TSharedRef<ITableRow> MakeTrackWidget( TSharedPtr<FTimelineEdTrack> Track, const TSharedRef<STableViewBase>& OwnerTable );
 	/** Add a new track to the timeline */
-	FReply CreateNewTrack(FTTTrackBase::ETrackType Type );
+	void CreateNewTrack(FTTTrackBase::ETrackType Type );
 
 	/** Checks if the user can delete the selected tracks */
 	bool CanDeleteSelectedTracks() const;
@@ -333,7 +333,7 @@ private:
 	bool IsCurveAssetSelected() const;
 
 	/** Create new track from curve asset */
-	FReply CreateNewTrackFromAsset();
+	void CreateNewTrackFromAsset();
 
 	/** Callback when a track item is scrolled into view */
 	void OnItemScrolledIntoView( TSharedPtr<FTimelineEdTrack> InTrackNode, const TSharedPtr<ITableRow>& InWidget );
@@ -351,5 +351,7 @@ private:
 
 	void SetSizeScaleValue(float NewValue);
 	float GetSizeScaleValue() const;
+
+	TSharedRef<SWidget> MakeAddButton();
 };
 
diff --git a/Engine/Source/Runtime/AppFramework/Private/Framework/Testing/SStarshipGallery.cpp b/Engine/Source/Runtime/AppFramework/Private/Framework/Testing/SStarshipGallery.cpp
--- a/Engine/Source/Runtime/AppFramework/Private/Framework/Testing/SStarshipGallery.cpp
+++ b/Engine/Source/Runtime/AppFramework/Private/Framework/Testing/SStarshipGallery.cpp
@@ -470,6 +470,7 @@ public:
                 +SVerticalBox::Slot().AutoHeight()[ GenerateIconLibrary(NSLOCTEXT("StarshipGallery", "FractureIconTitle", "Fracture"), "/Plugins/Experimental/ChaosEditor/Content")]
                 +SVerticalBox::Slot().AutoHeight()[ GenerateIconLibrary(NSLOCTEXT("StarshipGallery", "CurveEditorIconTitle", "CurveEditor"), "Content/Editor/Slate/GenericCurveEditor/Icons")]
                 +SVerticalBox::Slot().AutoHeight()[ GenerateIconLibrary(NSLOCTEXT("StarshipGallery", "GeneralIconTitle", "General"), "Content/Editor/Slate/Icons/GeneralTools")]
+				+SVerticalBox::Slot().AutoHeight()[ GenerateIconLibrarySVG(NSLOCTEXT("StarshipGallery", "TimelineEditorIconTitle", "TimelineEditor"), "Content/Editor/Slate/Starship/TimelineEditor")]
             ]
         ];
     }
