Added search bar to BP Debugger

diff --git a/Engine/Source/Editor/Kismet/Private/Debugging/SKismetDebuggingView.cpp b/Engine/Source/Editor/Kismet/Private/Debugging/SKismetDebuggingView.cpp
--- a/Engine/Source/Editor/Kismet/Private/Debugging/SKismetDebuggingView.cpp
+++ b/Engine/Source/Editor/Kismet/Private/Debugging/SKismetDebuggingView.cpp
@@ -34,6 +34,7 @@
 #include "Styling/StyleColors.h"
 #include "Kismet2/BlueprintEditorUtils.h"
 #include "PropertyInfoViewStyle.h"
+#include "Widgets/Input/SSearchBox.h"
 
 #define LOCTEXT_NAMESPACE "DebugViewUI"
 
@@ -68,16 +69,49 @@ FText FDebugLineItem::GetDescription() const
 
 TSharedRef<SWidget> FDebugLineItem::GenerateNameWidget()
 {
-	return SNew(STextBlock)
-		.ToolTipText(this, &FDebugLineItem::GetDisplayName)
-		.Text(this, &FDebugLineItem::GetDisplayName);
+	return SNew(PropertyInfoViewStyle::STextHighlightOverlay)
+		.FullText(this, &FDebugLineItem::GetDisplayName)
+		.HighlightText(SearchBox.Get(), &SSearchBox::GetText)
+		[
+			SNew(STextBlock)
+				.ToolTipText(this, &FDebugLineItem::GetDisplayName)
+				.Text(this, &FDebugLineItem::GetDisplayName)
+		];
 }
 
 TSharedRef<SWidget> FDebugLineItem::GenerateValueWidget()
 {
-	return SNew(STextBlock)
-		.ToolTipText(this, &FDebugLineItem::GetDescription)
-		.Text(this, &FDebugLineItem::GetDescription);
+	return SNew(PropertyInfoViewStyle::STextHighlightOverlay)
+		.FullText(this, &FDebugLineItem::GetDescription)
+		.HighlightText(SearchBox.Get(), &SSearchBox::GetText)
+		[
+			SNew(STextBlock)
+				.ToolTipText(this, &FDebugLineItem::GetDescription)
+				.Text(this, &FDebugLineItem::GetDescription)
+		];
+}
+
+void FDebugLineItem::UpdateSearchFlags(bool bIsRootNode)
+{
+	const FString SearchString = SearchBox->GetText().ToString();
+	
+	bVisible = GetDisplayName().ToString().Contains(SearchString) || GetDescription().ToString().Contains(SearchString);
+
+	// for root nodes, bParentsMatchSearch always matches bVisible
+	if(bVisible || bIsRootNode)
+	{
+		bParentsMatchSearch = bVisible;
+	}
+}
+
+bool FDebugLineItem::IsVisible()
+{
+	return bVisible;
+}
+
+bool FDebugLineItem::DoParentsMatchSearch()
+{
+	return bParentsMatchSearch;
 }
 
 bool FDebugLineItem::HasChildren()
@@ -109,7 +143,7 @@ UBlueprint* FDebugLineItem::GetBlueprintForObject(UObject* ParentObject)
 		return ParentBlueprint;
 	}
 
-	if(UClass *ParentClass = ParentObject->GetClass())
+	if(UClass* ParentClass = ParentObject->GetClass())
 	{
 		if(UBlueprint* ParentBlueprint = Cast<UBlueprint>(ParentClass->ClassGeneratedBy))
 		{
@@ -172,8 +206,8 @@ void FDebugLineItem::OnDebugLineTypeActiveChanged(ECheckBoxState CheckState, EDe
 class FLineItemWithChildren : public FDebugLineItem
 {
 public:
-	FLineItemWithChildren(EDebugLineType InType) :
-		FDebugLineItem(InType)
+	FLineItemWithChildren(EDebugLineType InType, TSharedPtr<SSearchBox> InSearchBox) :
+		FDebugLineItem(InType, MoveTemp(InSearchBox))
 	{}
 	
 	virtual ~FLineItemWithChildren() override = default;
@@ -183,6 +217,86 @@ public:
 		return !ChildrenMirrors.IsEmpty();
 	}
 
+	virtual bool CanHaveChildren() override { return true; }
+
+	bool SearchRecursive(TArray<FLineItemWithChildren*>& Parents,
+		TSharedPtr<STreeView<FDebugTreeItemPtr>> DebugTreeView)
+	{
+		TSharedPtr<ITableRow> Row = DebugTreeView->WidgetFromItem(SharedThis(this));
+		bVisible = false;
+
+		UpdateSearchFlags();
+
+		bool bChildMatch = false;
+		Parents.Push(this);
+		
+		TArray<FDebugTreeItemPtr> Children;
+		GatherChildren(Children, false);
+		for(const FDebugTreeItemPtr& ChildRef : Children)
+		{
+			if(ChildRef->CanHaveChildren())
+			{
+				ChildRef->bParentsMatchSearch = bParentsMatchSearch;
+				FLineItemWithChildren* Child = StaticCast<FLineItemWithChildren*>(ChildRef.Get());
+				
+                // check if the child has been seen already in parents.
+                // if it has, skip it. (avoids stack overflows)
+                if(Parents.FindByPredicate(
+                [Child](const FLineItemWithChildren* Relative)
+                	{
+                		return (Relative->Type == Child->Type) && Relative->Compare(Child);
+                	}
+                ))
+                {
+                	continue;
+                }
+
+				// if any children need to expand, so should this
+				if(Child->SearchRecursive(Parents, DebugTreeView))
+				{
+					bVisible = true;
+					bChildMatch = true;
+
+					// exit early if children aren't in the tree yet anyway and
+					// we already know to expand this
+					if(!Row)
+					{
+						break;
+					}
+				}
+			}
+			else
+			{
+				ChildRef->UpdateSearchFlags();
+				
+				// if any children need to expand, so should this
+				if(ChildRef->IsVisible())
+				{
+					bVisible = true;
+					bChildMatch = true;
+
+					// exit early if children aren't in the tree yet anyway and
+					// we already know to expand this
+					if(!Row)
+					{
+						break;
+					}
+				}
+			}
+		}
+		
+		Parents.Pop(false);
+		if(bChildMatch)
+		{
+			if(Row && !Row->IsItemExpanded())
+			{
+				Row->ToggleExpansion();
+			}
+		}
+
+		return bVisible;
+	}
+
 protected:
 	// List of children 
 	TArray<FDebugTreeItemPtr> ChildrenMirrors;
@@ -193,18 +307,28 @@ protected:
 	 *
 	 * O( # Children )
 	 */
-	void EnsureChildIsAdded(TArray<FDebugTreeItemPtr>& OutChildren, const FDebugLineItem& Item)
+	void EnsureChildIsAdded(TArray<FDebugTreeItemPtr>& OutChildren, const FDebugLineItem& Item, bool bRespectSearch)
 	{
 		for (int32 i = 0; i < ChildrenMirrors.Num(); ++i)
 		{
 			TSharedPtr< FDebugLineItem > MirrorItem = ChildrenMirrors[i];
+			if(bParentsMatchSearch)
+			{
+				// propogate parents search state to children
+				MirrorItem->bParentsMatchSearch = true;
+			}
 
 			if (MirrorItem->Type == Item.Type)
 			{
 				if (Item.Compare(MirrorItem.Get()))
 				{
 					MirrorItem->UpdateData(Item);
-					OutChildren.Add(MirrorItem);
+					
+					// only add item if it matches search
+					if(!bRespectSearch || SearchBox->GetText().IsEmpty() || MirrorItem->IsVisible() || MirrorItem->DoParentsMatchSearch())
+					{
+						OutChildren.Add(MirrorItem);
+					}
 					return;
 				}
 			}
@@ -227,8 +351,8 @@ protected:
 	FString Message;
 public:
 	// Message line
-	FMessageLineItem(const FString& InMessage)
-		: FDebugLineItem(DLT_Message)
+	FMessageLineItem(const FString& InMessage, TSharedPtr<SSearchBox> InSearchBox)
+		: FDebugLineItem(DLT_Message, MoveTemp(InSearchBox))
 		, Message(InMessage)
 	{
 	}
@@ -241,7 +365,7 @@ protected:
 
 	virtual FDebugLineItem* Duplicate() const override
 	{
-		return new FMessageLineItem(Message);
+		return new FMessageLineItem(Message, SearchBox);
 	}
 
 	virtual FText GetDescription() const override
@@ -259,8 +383,8 @@ protected:
 	int32 UUID;
 	TWeakObjectPtr< UObject > ParentObjectRef;
 public:
-	FLatentActionLineItem(int32 InUUID, UObject* ParentObject)
-		: FDebugLineItem(DLT_LatentAction)
+	FLatentActionLineItem(int32 InUUID, UObject* ParentObject, TSharedPtr<SSearchBox> InSearchBox)
+		: FDebugLineItem(DLT_LatentAction, MoveTemp(InSearchBox))
 	{
 		UUID = InUUID;
 		check(UUID != INDEX_NONE);
@@ -277,7 +401,7 @@ protected:
 
 	virtual FDebugLineItem* Duplicate() const override
 	{
-		return new FLatentActionLineItem(UUID, ParentObjectRef.Get());
+		return new FLatentActionLineItem(UUID, ParentObjectRef.Get(), SearchBox);
 	}
 protected:
 	virtual TSharedRef<SWidget> GenerateNameWidget() override;
@@ -305,11 +429,16 @@ FText FLatentActionLineItem::GetDescription() const
 
 TSharedRef<SWidget> FLatentActionLineItem::GenerateNameWidget()
 {
-	return SNew(SHyperlink)
-		.Style(FEditorStyle::Get(), "HoverOnlyHyperlink")
-		.OnNavigate(this, &FLatentActionLineItem::OnNavigateToLatentNode)
-		.Text(this, &FLatentActionLineItem::GetDisplayName)
-		.ToolTipText( LOCTEXT("NavLatentActionLoc_Tooltip", "Navigate to the latent action location") );
+	return SNew(PropertyInfoViewStyle::STextHighlightOverlay)
+		.FullText(this, &FLatentActionLineItem::GetDisplayName)
+		.HighlightText(SearchBox.Get(), &SSearchBox::GetText)
+		[
+			SNew(SHyperlink)
+				.Style(FEditorStyle::Get(), "HoverOnlyHyperlink")
+				.OnNavigate(this, &FLatentActionLineItem::OnNavigateToLatentNode)
+				.Text(this, &FLatentActionLineItem::GetDisplayName)
+				.ToolTipText( LOCTEXT("NavLatentActionLoc_Tooltip", "Navigate to the latent action location") )
+		];
 }
 
 TSharedRef<SWidget> FLatentActionLineItem::GetNameIcon()
@@ -357,8 +486,8 @@ struct FWatchChildLineItem : public FLineItemWithChildren
 protected:
 	FPropertyInstanceInfo Data;
 public:
-	FWatchChildLineItem(const FPropertyInstanceInfo& Child) :
-		FLineItemWithChildren(DLT_WatchChild),
+	FWatchChildLineItem(const FPropertyInstanceInfo& Child, TSharedPtr<SSearchBox> InSearchBox) :
+		FLineItemWithChildren(DLT_WatchChild, MoveTemp(InSearchBox)),
 		Data(Child)
 	{}
 
@@ -381,7 +510,7 @@ public:
 
 	virtual FDebugLineItem* Duplicate() const override
 	{
-		return new FWatchChildLineItem(Data);
+		return new FWatchChildLineItem(Data, SearchBox);
 	}
 	
 	virtual FText GetDescription() const override
@@ -412,11 +541,11 @@ public:
 			.ColorAndOpacity(Color);
 	}
 
-	virtual void GatherChildren(TArray<FDebugTreeItemPtr>& OutChildren) override
+	virtual void GatherChildren(TArray<FDebugTreeItemPtr>& OutChildren, bool bRespectSearch) override
 	{
-		for (const TSharedPtr<FPropertyInstanceInfo> &ChildData : Data.Children)
+		for (const TSharedPtr<FPropertyInstanceInfo>& ChildData : Data.Children)
 		{
-			EnsureChildIsAdded(OutChildren, FWatchChildLineItem(*ChildData));
+			EnsureChildIsAdded(OutChildren, FWatchChildLineItem(*ChildData, SearchBox), bRespectSearch);
 		}
 	}
 };
@@ -429,8 +558,8 @@ protected:
 	// watches a UObject instead of a pin
 	TWeakObjectPtr<UObject> ObjectToWatch;
 public:
-	FSelfWatchLineItem(UObject* Object)	: 
-		FLineItemWithChildren(DLT_Watch),
+	FSelfWatchLineItem(UObject* Object, TSharedPtr<SSearchBox> InSearchBox)	: 
+		FLineItemWithChildren(DLT_Watch, MoveTemp(InSearchBox)),
 		ObjectToWatch(Object)
 	{}
 
@@ -442,10 +571,10 @@ public:
 
 	virtual FDebugLineItem* Duplicate() const override
 	{
-		return new FSelfWatchLineItem(ObjectToWatch.Get());
+		return new FSelfWatchLineItem(ObjectToWatch.Get(), SearchBox);
 	}
 	
-	virtual void GatherChildren(TArray<FDebugTreeItemPtr>& OutChildren) override
+	virtual void GatherChildren(TArray<FDebugTreeItemPtr>& OutChildren, bool bRespectSearch) override
 	{
 		if (UObject* Object = ObjectToWatch.Get())
 		{
@@ -458,7 +587,7 @@ public:
 					void* Value = Property->ContainerPtrToValuePtr<void*>(Object);
 					FKismetDebugUtilities::GetDebugInfoInternal(DebugInfo, Property, Value);
 				
-					EnsureChildIsAdded(OutChildren, FWatchChildLineItem(*DebugInfo));
+					EnsureChildIsAdded(OutChildren, FWatchChildLineItem(*DebugInfo, SearchBox), bRespectSearch);
 				}
 			}
 		}
@@ -487,8 +616,8 @@ protected:
 	TWeakObjectPtr< UObject > ParentObjectRef;
 	FEdGraphPinReference ObjectRef;
 public:
-	FWatchLineItem(UEdGraphPin* PinToWatch, UObject* ParentObject)
-		: FLineItemWithChildren(DLT_Watch)
+	FWatchLineItem(UEdGraphPin* PinToWatch, UObject* ParentObject, TSharedPtr<SSearchBox> InSearchBox)
+		: FLineItemWithChildren(DLT_Watch, MoveTemp(InSearchBox))
 	{
 		ObjectRef = PinToWatch;
 		ParentObjectRef = ParentObject;
@@ -503,7 +632,7 @@ public:
 
 	virtual FDebugLineItem* Duplicate() const override
 	{
-		return new FWatchLineItem(ObjectRef.Get(), ParentObjectRef.Get());
+		return new FWatchLineItem(ObjectRef.Get(), ParentObjectRef.Get(), SearchBox);
 	}	
 
 	virtual void MakeMenu(class FMenuBuilder& MenuBuilder) override
@@ -522,7 +651,7 @@ public:
 		}
 	}
 	
-	virtual void GatherChildren(TArray<FDebugTreeItemPtr>& OutChildren) override
+	virtual void GatherChildren(TArray<FDebugTreeItemPtr>& OutChildren, bool bRespectSearch) override
 	{
 		if (UEdGraphPin* PinToWatch = ObjectRef.Get())
 		{
@@ -539,9 +668,9 @@ public:
 				if(WatchStatus == FKismetDebugUtilities::EWTR_Valid)
 				{
 					check(DebugInfo);
-					for (const TSharedPtr<FPropertyInstanceInfo> &ChildData : DebugInfo->Children)
+					for (const TSharedPtr<FPropertyInstanceInfo>& ChildData : DebugInfo->Children)
 					{
-						EnsureChildIsAdded(OutChildren, FWatchChildLineItem(*ChildData));
+						EnsureChildIsAdded(OutChildren, FWatchChildLineItem(*ChildData, SearchBox), bRespectSearch);
 					}
 				}
 			}
@@ -621,11 +750,16 @@ FText FWatchLineItem::GetDescription() const
 
 TSharedRef<SWidget> FWatchLineItem::GenerateNameWidget()
 {
-	return SNew(SHyperlink)
-		.Style(FEditorStyle::Get(), "HoverOnlyHyperlink")
-		.OnNavigate(this, &FWatchLineItem::OnNavigateToWatchLocation)
-		.Text(this, &FWatchLineItem::GetDisplayName)
-		.ToolTipText( LOCTEXT("NavWatchLoc", "Navigate to the watch location") );
+	return SNew(PropertyInfoViewStyle::STextHighlightOverlay)
+		.FullText(this, &FWatchLineItem::GetDisplayName)
+		.HighlightText(SearchBox.Get(), &SSearchBox::GetText)
+		[
+			SNew(SHyperlink)
+				.Style(FEditorStyle::Get(), "HoverOnlyHyperlink")
+				.OnNavigate(this, &FWatchLineItem::OnNavigateToWatchLocation)
+				.Text(this, &FWatchLineItem::GetDisplayName)
+				.ToolTipText( LOCTEXT("NavWatchLoc", "Navigate to the watch location") )
+		];
 }
 
 // overlays the watch icon on top of a faded icon associated with the pin type
@@ -679,8 +813,8 @@ protected:
 	TWeakObjectPtr<UObject> ParentObjectRef;
 	TSoftObjectPtr<UEdGraphNode> BreakpointNode;
 public:
-	FBreakpointLineItem(TSoftObjectPtr<UEdGraphNode> BreakpointToWatch, UObject* ParentObject)
-		: FDebugLineItem(DLT_Breakpoint)
+	FBreakpointLineItem(TSoftObjectPtr<UEdGraphNode> BreakpointToWatch, UObject* ParentObject, TSharedPtr<SSearchBox> InSearchBox)
+		: FDebugLineItem(DLT_Breakpoint, MoveTemp(InSearchBox))
 	{
 		BreakpointNode = BreakpointToWatch;
 		ParentObjectRef = ParentObject;
@@ -695,7 +829,7 @@ public:
 
 	virtual FDebugLineItem* Duplicate() const override
 	{
-		return new FBreakpointLineItem(BreakpointNode, ParentObjectRef.Get());
+		return new FBreakpointLineItem(BreakpointNode, ParentObjectRef.Get(), SearchBox);
 	}	
 
 	virtual void MakeMenu(class FMenuBuilder& MenuBuilder) override
@@ -767,11 +901,16 @@ protected:
 	
 	virtual TSharedRef<SWidget> GenerateNameWidget() override
 	{
-		return SNew(SHyperlink)
-			.Style(FEditorStyle::Get(), "HoverOnlyHyperlink")
-			.Text(this, &FBreakpointLineItem::GetLocationDescription)
-			.ToolTipText( LOCTEXT("NavBreakpointLoc", "Navigate to the breakpoint location") )
-			.OnNavigate(this, &FBreakpointLineItem::OnNavigateToBreakpointLocation);
+		return SNew(PropertyInfoViewStyle::STextHighlightOverlay)
+			.FullText(this, &FBreakpointLineItem::GetDisplayName)
+			.HighlightText(SearchBox.Get(), &SSearchBox::GetText)
+			[
+				SNew(SHyperlink)
+					.Style(FEditorStyle::Get(), "HoverOnlyHyperlink")
+					.Text(this, &FBreakpointLineItem::GetDisplayName)
+					.ToolTipText( LOCTEXT("NavBreakpointLoc", "Navigate to the breakpoint location") )
+					.OnNavigate(this, &FBreakpointLineItem::OnNavigateToBreakpointLocation)
+			];
 	}
 
 	virtual TSharedRef<SWidget> GetNameIcon() override
@@ -788,8 +927,8 @@ protected:
 			];
 	}
 	
-	
-	FText GetLocationDescription() const;
+
+	virtual FText GetDisplayName() const override;
 	FReply OnUserToggledEnabled();
 
 	void OnNavigateToBreakpointLocation();
@@ -798,7 +937,7 @@ protected:
 	FText GetStatusTooltip() const;
 };
 
-FText FBreakpointLineItem::GetLocationDescription() const
+FText FBreakpointLineItem::GetDisplayName() const
 {
 	if (FBlueprintBreakpoint* MyBreakpoint = GetBreakpoint())
 	{
@@ -869,14 +1008,17 @@ public:
 	// The parent object
 	TWeakObjectPtr<UBlueprint> Blueprint;
 	
-	FBreakpointParentItem(TWeakObjectPtr<UBlueprint> Blueprint)
-		: FLineItemWithChildren(DLT_TraceStackParent),
+	FBreakpointParentItem(TWeakObjectPtr<UBlueprint> Blueprint, TSharedPtr<SSearchBox> InSearchBox)
+		: FLineItemWithChildren(DLT_TraceStackParent, MoveTemp(InSearchBox)),
 		  Blueprint(Blueprint)
 	{
 	}
 
-	virtual void GatherChildren(TArray<FDebugTreeItemPtr>& OutChildren) override
+	virtual void GatherChildren(TArray<FDebugTreeItemPtr>& OutChildren, bool bRespectSearch) override
 	{
+		// update search flags to match that of a root node
+		UpdateSearchFlags(/* bIsRootNode */ true);
+		
 		if(!Blueprint.IsValid())
 		{
 			return;
@@ -885,17 +1027,19 @@ public:
 		// Create children for each breakpoint
 		FKismetDebugUtilities::ForeachBreakpoint(
 			Blueprint.Get(),
-			[this, &OutChildren]
+			[this, &OutChildren, bRespectSearch]
 			(FBlueprintBreakpoint& Breakpoint)
 			{
-				EnsureChildIsAdded(OutChildren, FBreakpointLineItem(Breakpoint.GetLocation(), Blueprint.Get()));
+				EnsureChildIsAdded(OutChildren,
+					FBreakpointLineItem(Breakpoint.GetLocation(), Blueprint.Get(), SearchBox), bRespectSearch);
 			}
 		);
 
 		// Make sure there is something there, to let the user know if there is nothing
 		if (OutChildren.Num() == 0)
 		{
-			EnsureChildIsAdded(OutChildren, FMessageLineItem( LOCTEXT("NoBreakpoints", "No breakpoints").ToString() ));
+			EnsureChildIsAdded(OutChildren,
+				FMessageLineItem(LOCTEXT("NoBreakpoints", "No breakpoints").ToString(), SearchBox), bRespectSearch);
 		}
 	}
 
@@ -927,8 +1071,8 @@ protected:
 	// The parent object
 	TWeakObjectPtr<UObject> ObjectRef;
 public:
-	FParentLineItem(UObject* Object)
-		: FLineItemWithChildren(DLT_Parent)
+	FParentLineItem(UObject* Object, TSharedPtr<SSearchBox> InSearchBox)
+		: FLineItemWithChildren(DLT_Parent, MoveTemp(InSearchBox))
 	{
 		ObjectRef = Object;
 	}
@@ -938,12 +1082,15 @@ public:
 		return ObjectRef.Get();
 	}
 
-	virtual void GatherChildren(TArray<FDebugTreeItemPtr>& OutChildren) override
+	virtual void GatherChildren(TArray<FDebugTreeItemPtr>& OutChildren, bool bRespectSearch) override
 	{
+		// update search flags to match that of a root node
+		UpdateSearchFlags(/* bIsRootNode */ true);
+		
 		if (UObject* ParentObject = ObjectRef.Get())
 		{
 			// every instance should have an automatic watch for 'self'
-			EnsureChildIsAdded(OutChildren, FSelfWatchLineItem(ParentObject));
+			EnsureChildIsAdded(OutChildren, FSelfWatchLineItem(ParentObject, SearchBox), bRespectSearch);
 			
 			UBlueprint* ParentBP = FDebugLineItem::GetBlueprintForObject(ParentObject);
 			if (ParentBP != nullptr)
@@ -953,9 +1100,10 @@ public:
 				{
 					FKismetDebugUtilities::ForeachPinWatch(
 						ParentBP,
-						[this, &OutChildren, ParentObject](UEdGraphPin* WatchedPin)
+						[this, &OutChildren, ParentObject, bRespectSearch](UEdGraphPin* WatchedPin)
 						{
-							EnsureChildIsAdded(OutChildren, FWatchLineItem(WatchedPin, ParentObject));
+							EnsureChildIsAdded(OutChildren,
+								FWatchLineItem(WatchedPin, ParentObject, SearchBox), bRespectSearch);
 						}
 					);
 				}
@@ -975,7 +1123,8 @@ public:
 						for (TSet<int32>::TConstIterator RemainingIt(UUIDSet); RemainingIt; ++RemainingIt)
 						{
 							const int32 UUID = *RemainingIt;
-							EnsureChildIsAdded(OutChildren, FLatentActionLineItem(UUID, ParentObject));
+							EnsureChildIsAdded(OutChildren,
+								FLatentActionLineItem(UUID, ParentObject, SearchBox), bRespectSearch);
 						}
 					}
 				}
@@ -983,7 +1132,8 @@ public:
 				// Make sure there is something there, to let the user know if there is nothing
 				if (OutChildren.Num() == 0)
 				{
-					EnsureChildIsAdded(OutChildren, FMessageLineItem( LOCTEXT("NoDebugInfo", "No debugging info").ToString() ));
+					EnsureChildIsAdded(OutChildren,
+						FMessageLineItem(LOCTEXT("NoDebugInfo", "No debugging info").ToString(), SearchBox), bRespectSearch);
 				}
 			}
 			//@TODO: try to get at TArray<struct FDebugDisplayProperty> DebugProperties in UGameViewportClient, if available
@@ -1031,7 +1181,7 @@ protected:
 
 	virtual FDebugLineItem* Duplicate() const override
 	{
-		return new FParentLineItem(ObjectRef.Get());
+		return new FParentLineItem(ObjectRef.Get(), SearchBox);
 	}
 
 	virtual FText GetDisplayName() const override
@@ -1098,8 +1248,8 @@ class FTraceStackChildItem : public FDebugLineItem
 protected:
 	int32 StackIndex;
 public:
-	FTraceStackChildItem(int32 InStackIndex)
-		: FDebugLineItem(DLT_TraceStackChild)
+	FTraceStackChildItem(int32 InStackIndex, TSharedPtr<SSearchBox> InSearchBox)
+		: FDebugLineItem(DLT_TraceStackChild, MoveTemp(InSearchBox))
 	{
 		StackIndex = InStackIndex;
 	}
@@ -1153,11 +1303,16 @@ protected:
 
 	virtual TSharedRef<SWidget> GenerateNameWidget() override
 	{
-		return SNew(SHyperlink)
-			.Text(this, &FTraceStackChildItem::GetDisplayName)
-			.Style(FEditorStyle::Get(), "HoverOnlyHyperlink")
-			.ToolTipText(LOCTEXT("NavigateToDebugTraceLocationHyperlink_ToolTip", "Navigate to the trace location"))
-			.OnNavigate(this, &FTraceStackChildItem::OnNavigateToNode);
+		return SNew(PropertyInfoViewStyle::STextHighlightOverlay)
+			.FullText(this, &FTraceStackChildItem::GetDisplayName)
+			.HighlightText(SearchBox.Get(), &SSearchBox::GetText)
+			[
+				SNew(SHyperlink)
+					.Text(this, &FTraceStackChildItem::GetDisplayName)
+					.Style(FEditorStyle::Get(), "HoverOnlyHyperlink")
+					.ToolTipText(LOCTEXT("NavigateToDebugTraceLocationHyperlink_ToolTip", "Navigate to the trace location"))
+					.OnNavigate(this, &FTraceStackChildItem::OnNavigateToNode)
+			];
 	}
 
 	virtual TSharedRef<SWidget> GetNameIcon() override
@@ -1173,22 +1328,26 @@ protected:
 	// Visit time and actor name
 	virtual TSharedRef<SWidget> GenerateValueWidget() override
 	{
-		return SNew(SHorizontalBox)
-			+ SHorizontalBox::Slot()
-			.AutoWidth()
-			[
-				SNew(SHyperlink)
-					.Text(this, &FTraceStackChildItem::GetContextObjectName)
-					.Style(FEditorStyle::Get(), "HoverOnlyHyperlink")
-					.ToolTipText( LOCTEXT("SelectActor_Tooltip", "Select this actor") )
-					.OnNavigate(this, &FTraceStackChildItem::OnSelectContextObject)
-			]
-			
-			+ SHorizontalBox::Slot()
-			.AutoWidth()
+		return SNew(PropertyInfoViewStyle::STextHighlightOverlay)
+			.FullText(this, &FTraceStackChildItem::GetDescription)
+			.HighlightText(SearchBox.Get(), &SSearchBox::GetText)
 			[
-				SNew(STextBlock)
-					.Text(this, &FTraceStackChildItem::GetVisitTime)
+				SNew(SHorizontalBox)
+					+ SHorizontalBox::Slot()
+					.AutoWidth()
+					[
+						SNew(SHyperlink)
+							.Text(this, &FTraceStackChildItem::GetContextObjectName)
+							.Style(FEditorStyle::Get(), "HoverOnlyHyperlink")
+							.ToolTipText( LOCTEXT("SelectActor_Tooltip", "Select this actor") )
+							.OnNavigate(this, &FTraceStackChildItem::OnSelectContextObject)
+					]
+					+ SHorizontalBox::Slot()
+					.AutoWidth()
+					[
+						SNew(STextBlock)
+							.Text(this, &FTraceStackChildItem::GetVisitTime)
+					]
 			];
 	}
 
@@ -1215,6 +1374,11 @@ protected:
 		return (ObjectContext != nullptr) ? FText::FromString(ObjectContext->GetName()) : LOCTEXT("ObjectDoesNotExist", "(object no longer exists)");
 	}
 
+	virtual FText GetDescription() const override
+	{
+		return FText::FromString(GetContextObjectName().ToString() + GetVisitTime().ToString());
+	}
+
 	void OnNavigateToNode()
 	{
 		if (UEdGraphNode* Node = GetNode())
@@ -1248,20 +1412,23 @@ protected:
 class FTraceStackParentItem : public FLineItemWithChildren
 {
 public:
-	FTraceStackParentItem()
-		: FLineItemWithChildren(DLT_TraceStackParent)
+	FTraceStackParentItem(TSharedPtr<SSearchBox> InSearchBox)
+		: FLineItemWithChildren(DLT_TraceStackParent, MoveTemp(InSearchBox))
 	{
 	}
 
-	virtual void GatherChildren(TArray<FDebugTreeItemPtr>& OutChildren) override
+	virtual void GatherChildren(TArray<FDebugTreeItemPtr>& OutChildren, bool bRespectSearch) override
 	{
+		// update search flags to match that of a root node
+		UpdateSearchFlags(/* bIsRootNode */ true);
+		
 		const TSimpleRingBuffer<FKismetTraceSample>& TraceStack = FKismetDebugUtilities::GetTraceStack();
 		const int32 NumVisible = TraceStack.Num();
 
 		// Create any new stack entries that are needed
 		for (int32 i = ChildrenMirrors.Num(); i < NumVisible; ++i)
 		{
-			ChildrenMirrors.Add(MakeShareable( new FTraceStackChildItem(i) ));
+			ChildrenMirrors.Add(MakeShareable( new FTraceStackChildItem(i, SearchBox) ));
 		}
 
 		// Add the visible stack entries as children
@@ -1441,6 +1608,15 @@ TSharedPtr<SWidget> SKismetDebuggingView::OnMakeContextMenu()
 	return MenuBuilder.MakeWidget();
 }
 
+void SKismetDebuggingView::OnSearchTextCommitted(const FText& Text, ETextCommit::Type Type)
+{
+	if(Type == ETextCommit::OnEnter)
+	{
+		DebugTreeView->ClearExpandedItems();
+		OtherTreeView->ClearExpandedItems();
+	}
+}
+
 
 FText SKismetDebuggingView::GetTabLabel() const
 {
@@ -1507,7 +1683,7 @@ TSharedRef<SWidget> SKismetDebuggingView::ConstructBlueprintClassPicker()
 	Options.bIsBlueprintBaseOnly = true;
 	Options.bShowUnloadedBlueprints = false;
 	
-	FClassViewerModule &ClassViewerModule = FModuleManager::LoadModuleChecked<FClassViewerModule>("ClassViewer");
+	FClassViewerModule& ClassViewerModule = FModuleManager::LoadModuleChecked<FClassViewerModule>("ClassViewer");
 
 	FOnClassPicked OnClassPicked;
 	OnClassPicked.BindRaw(this, &SKismetDebuggingView::OnBlueprintClassPicked);
@@ -1611,6 +1787,12 @@ void SKismetDebuggingView::Construct(const FArguments& InArgs)
 					                .OnClicked( this, &SKismetDebuggingView::OnDisableAllBreakpointsClicked )
 							]
 					]
+					+ SVerticalBox::Slot()
+					.AutoHeight()
+					[
+						SAssignNew(SearchBox, SSearchBox)
+						.OnTextCommitted(this, &SKismetDebuggingView::OnSearchTextCommitted)
+					]
 			]
 			+SVerticalBox::Slot()
 			[
@@ -1680,8 +1862,8 @@ void SKismetDebuggingView::Construct(const FArguments& InArgs)
 			]
 	];
 
-	TraceStackItem = MakeShareable(new FTraceStackParentItem);
-	BreakpointParentItem = MakeShareable(new FBreakpointParentItem(BlueprintToWatchPtr));
+	TraceStackItem = MakeShareable(new FTraceStackParentItem(SearchBox));
+	BreakpointParentItem = MakeShareable(new FBreakpointParentItem(BlueprintToWatchPtr, SearchBox));
 }
 
 void SKismetDebuggingView::Tick( const FGeometry& AllottedGeometry, const double InCurrentTime, const float InDeltaTime )
@@ -1731,11 +1913,40 @@ void SKismetDebuggingView::Tick( const FGeometry& AllottedGeometry, const double
         }
 	}
 
+	auto TryAddToRoot = [](const TSharedPtr<FDebugLineItem>& Item, TArray<FDebugTreeItemPtr>& Root,
+		 const FText& SearchText, const TSharedPtr<STreeView<FDebugTreeItemPtr>> &Tree)
+	{
+		if(SearchText.IsEmpty())
+		{
+			Root.Add(Item);
+		}
+		else
+		{
+			if(Item->CanHaveChildren())
+			{
+				TArray<FLineItemWithChildren*> StackOverflowAvoidence;
+				FLineItemWithChildren* ItemWithChildren = StaticCast<FLineItemWithChildren*>(Item.Get());
+				if(ItemWithChildren->SearchRecursive(StackOverflowAvoidence, Tree))
+                {
+                	Root.Add(Item);
+                }
+			}
+			else
+			{
+				Item->UpdateSearchFlags();
+				if(Item->IsVisible())
+				{
+                	Root.Add(Item);
+				}
+			}
+		}
+	};
+
 	// This will pull anything out of Old that is also New (sticking around), so afterwards Old is a list of things to remove
 	RootTreeItems.Empty();
 	for (TSet<UObject*>::TIterator NewIter(NewRootSet); NewIter; ++NewIter)
 	{
-		TWeakObjectPtr ObjectToAdd = *NewIter;
+		TWeakObjectPtr<UObject> ObjectToAdd = *NewIter;
 
 		// destroyed objects can still appear if they haven't ben GCed yet.
 		// weak object pointers will detect it and return nullptr
@@ -1747,17 +1958,15 @@ void SKismetDebuggingView::Tick( const FGeometry& AllottedGeometry, const double
 		if (OldRootSet.Contains(ObjectToAdd.Get()))
 		{
 			OldRootSet.Remove(ObjectToAdd.Get());
-			RootTreeItems.Add(ObjectToTreeItemMap.FindChecked(ObjectToAdd.Get()));
+			
+			const TSharedPtr<FDebugLineItem>& Item = ObjectToTreeItemMap.FindChecked(ObjectToAdd.Get());
+			TryAddToRoot(Item, RootTreeItems, SearchBox->GetText(), DebugTreeView);
 		}
 		else
 		{
-			
-			FDebugTreeItemPtr NewPtr = FDebugTreeItemPtr(new FParentLineItem(ObjectToAdd.Get()));
+			FDebugTreeItemPtr NewPtr = FDebugTreeItemPtr(new FParentLineItem(ObjectToAdd.Get(), SearchBox));
 			ObjectToTreeItemMap.Add(ObjectToAdd.Get(), NewPtr);
-			RootTreeItems.Add(NewPtr);
-
-			// Autoexpand newly selected items
-			DebugTreeView->SetItemExpansion(NewPtr, true);
+			TryAddToRoot(NewPtr, RootTreeItems, SearchBox->GetText(), DebugTreeView);
 		}
 	}
 
@@ -1772,7 +1981,11 @@ void SKismetDebuggingView::Tick( const FGeometry& AllottedGeometry, const double
 	if (RootTreeItems.Num() == 0)
 	{
 		RootTreeItems.Add(MakeShareable(new FMessageLineItem(
-			bIsDebugging ? LOCTEXT("NoInstances", "No instances of this blueprint in existence").ToString() : LOCTEXT("NoPIEorSIE", "run PIE or SIE to see instance debug info").ToString())));
+			bIsDebugging ?
+				LOCTEXT("NoInstances", "No instances of this blueprint in existence").ToString() :
+				LOCTEXT("NoPIEorSIE", "run PIE or SIE to see instance debug info").ToString(),
+			SearchBox
+		)));
 	}
 
 	// Refresh the list
@@ -1783,13 +1996,13 @@ void SKismetDebuggingView::Tick( const FGeometry& AllottedGeometry, const double
 	// Show Breakpoints
 	if(FDebugLineItem::IsDebugLineTypeActive(FDebugLineItem::DLT_BreakpointParent))
 	{
-		OtherTreeItems.Add(BreakpointParentItem);
+		TryAddToRoot(BreakpointParentItem, OtherTreeItems, SearchBox->GetText(), OtherTreeView);
 	}
 
 	// Show the trace stack when debugging
 	if (bIsDebugging && FDebugLineItem::IsDebugLineTypeActive(FDebugLineItem::DLT_TraceStackParent))
 	{
-		OtherTreeItems.Add(TraceStackItem);
+		TryAddToRoot(TraceStackItem, OtherTreeItems, SearchBox->GetText(), OtherTreeView);
 	}
 	OtherTreeView->RequestTreeRefresh();
 }
diff --git a/Engine/Source/Editor/Kismet/Private/Debugging/SKismetDebuggingView.h b/Engine/Source/Editor/Kismet/Private/Debugging/SKismetDebuggingView.h
--- a/Engine/Source/Editor/Kismet/Private/Debugging/SKismetDebuggingView.h
+++ b/Engine/Source/Editor/Kismet/Private/Debugging/SKismetDebuggingView.h
@@ -52,11 +52,14 @@ public:
 	virtual void MakeMenu(class FMenuBuilder& MenuBuilder) { }
 
 	// Gather all of the children
-	virtual void GatherChildren(TArray<FDebugTreeItemPtr>& OutChildren) {}
+	virtual void GatherChildren(TArray<FDebugTreeItemPtr>& OutChildren, bool bRespectSearch = true) {}
 
 	// returns whether this tree node has children (used by drop down arrows)
 	virtual bool HasChildren();
 
+	// only line items inherriting from FLineItemWithChildren can have children
+	virtual bool CanHaveChildren() { return false; }
+
 	// @return The object that will act as a parent to more items in the tree, or NULL if this is a leaf node
 	virtual UObject* GetParentObject() { return NULL; }
 
@@ -81,11 +84,18 @@ public:
 
 	static bool IsDebugLineTypeActive(EDebugLineType Type);
 	static void OnDebugLineTypeActiveChanged(ECheckBoxState CheckState, EDebugLineType Type);
+
+	// updates bVisible and bParentsMatchSearch based on this node alone
+	void UpdateSearchFlags(bool bIsRootNode = false);
+
+	bool IsVisible();
+	bool DoParentsMatchSearch();
 protected:
+	TSharedPtr<SSearchBox> SearchBox;
 	
 	// Cannot create an instance of this class, it's just for use as a base class
-	FDebugLineItem(EDebugLineType InType)
-		: Type(InType)
+	FDebugLineItem(EDebugLineType InType, TSharedPtr<SSearchBox> InSearchBox)
+		: SearchBox(MoveTemp(InSearchBox)), Type(InType)
 	{
 	}
 
@@ -104,11 +114,16 @@ protected:
 
 	// @return The text to display in the value column, unless GenerateValueWidget is overridden
 	virtual FText GetDescription() const;
-private:
+protected:
 	// Type of action (poor mans RTTI for the tree, really only used to accelerate Compare checks)
 	EDebugLineType Type;
 
 	static uint16 ActiveTypeBitset;
+
+	// true if self or any recursive children match the search
+	bool bVisible = false;
+	// true if self or any recursive parents match the search
+	bool bParentsMatchSearch = false;
 };
 
 //////////////////////////////////////////////////////////////////////////
@@ -148,8 +163,11 @@ protected:
 	static TSharedRef<SHorizontalBox> GetDebugLineTypeToggle(FDebugLineItem::EDebugLineType Type, const FText& Text);
 	
 	TSharedPtr<SWidget> OnMakeContextMenu();
+
+	// called when SearchBox query is committed by the user
+	void OnSearchTextCommitted(const FText& Text, ETextCommit::Type);
 protected:
-	TSharedPtr< STreeView<FDebugTreeItemPtr> > DebugTreeView;
+	TSharedPtr<STreeView<FDebugTreeItemPtr>> DebugTreeView;
 	TMap<UObject*, FDebugTreeItemPtr> ObjectToTreeItemMap;
 	TArray<FDebugTreeItemPtr> RootTreeItems;
 
@@ -161,6 +179,10 @@ protected:
 	TSharedPtr< FTraceStackParentItem > TraceStackItem;
 	TSharedPtr< FBreakpointParentItem > BreakpointParentItem;
 
+	// Combo button for selecting which blueprint is being watched
 	TSharedPtr<SComboButton> DebugClassComboButton;
 	TWeakObjectPtr<UBlueprint> BlueprintToWatchPtr;
+
+	// Search Box for tree
+	TSharedPtr<SSearchBox> SearchBox;
 };
diff --git a/Engine/Source/Editor/UnrealEd/Private/PropertyInfoViewStyle.cpp b/Engine/Source/Editor/UnrealEd/Private/PropertyInfoViewStyle.cpp
--- a/Engine/Source/Editor/UnrealEd/Private/PropertyInfoViewStyle.cpp
+++ b/Engine/Source/Editor/UnrealEd/Private/PropertyInfoViewStyle.cpp
@@ -1,7 +1,10 @@
 ﻿// Copyright Epic Games, Inc. All Rights Reserved.
 
 #include "PropertyInfoViewStyle.h"
+
+#include "EditorStyleSet.h"
 #include "Math/Color.h"
+#include "Widgets/Text/STextBlock.h"
 
 namespace PropertyInfoViewStyle
 {
@@ -231,4 +234,29 @@ namespace PropertyInfoViewStyle
 
 		return FReply::Handled();
 	}
+
+	void STextHighlightOverlay::Construct(const FArguments& InArgs)
+	{
+		static FSlateBrush HighlightShape = FSlateBrush(*FEditorStyle::Get().GetBrush("TextBlock.HighlightShape"));
+		HighlightShape.TintColor = FLinearColor(0.f, 0.47f, 1.f, .3f);
+		
+		// uses overlay to create the highlight so that the text widget can be a button,
+		// hyperlink, textblock, etc
+		ChildSlot
+		[
+			SNew(SOverlay)
+				+SOverlay::Slot()
+				[
+					SNew(STextBlock)
+						.Text(InArgs._FullText)
+						.ColorAndOpacity(FLinearColor::Transparent)
+						.HighlightShape(&HighlightShape)
+						.HighlightText(InArgs._HighlightText)
+				]
+				+SOverlay::Slot()
+				[
+					InArgs._Content.Widget
+				]
+		];
+	}
 }
diff --git a/Engine/Source/Editor/UnrealEd/Public/PropertyInfoViewStyle.h b/Engine/Source/Editor/UnrealEd/Public/PropertyInfoViewStyle.h
--- a/Engine/Source/Editor/UnrealEd/Public/PropertyInfoViewStyle.h
+++ b/Engine/Source/Editor/UnrealEd/Public/PropertyInfoViewStyle.h
@@ -94,4 +94,25 @@ namespace PropertyInfoViewStyle
 		TSharedPtr<SButton> ExpanderArrow;
 		TAttribute<bool> HasChildren;
 	};
+
+	/**
+	 * allows you to highlight any arbitrary text based widget.
+	 * See FTraceStackChildItem::GenerateValueWidget (SKismetDebuggingView.cpp)
+	 * for an example.
+	 *
+	 * note: the text in Content must be scaled, positioned, etc exactly
+	 * the same as a standard STextBlock would be in order for the highlight
+	 * to match up properly
+	 */
+	class UNREALED_API STextHighlightOverlay : public SCompoundWidget
+	{
+	public:
+		SLATE_BEGIN_ARGS(STextHighlightOverlay) {}
+			SLATE_DEFAULT_SLOT( FArguments, Content )
+			SLATE_ATTRIBUTE(FText, FullText)
+			SLATE_ATTRIBUTE(FText, HighlightText)
+		SLATE_END_ARGS()
+
+		void Construct(const FArguments& InArgs);
+	};
 }
