quality of life changes for blueprint debugger

diff --git a/Engine/Source/Editor/Kismet/Private/Debugging/SKismetDebuggingView.cpp b/Engine/Source/Editor/Kismet/Private/Debugging/SKismetDebuggingView.cpp
--- a/Engine/Source/Editor/Kismet/Private/Debugging/SKismetDebuggingView.cpp
+++ b/Engine/Source/Editor/Kismet/Private/Debugging/SKismetDebuggingView.cpp
@@ -34,7 +34,9 @@
 #include "Styling/StyleColors.h"
 #include "Kismet2/BlueprintEditorUtils.h"
 #include "PropertyInfoViewStyle.h"
+#include "GenericPlatform/GenericPlatformApplicationMisc.h"
 #include "Widgets/Input/SSearchBox.h"
+#include "HAL/PlatformApplicationMisc.h"
 
 #define LOCTEXT_NAMESPACE "DebugViewUI"
 
@@ -67,6 +69,26 @@ FText FDebugLineItem::GetDescription() const
 	return FText::GetEmpty();
 }
 
+bool FDebugLineItem::HasName() const
+{
+	return !GetDisplayName().IsEmpty();
+}
+
+bool FDebugLineItem::HasValue() const
+{
+	return !GetDescription().IsEmpty();
+}
+
+void FDebugLineItem::CopyNameToClipboard() const
+{
+	FPlatformApplicationMisc::ClipboardCopy(ToCStr(GetDisplayName().ToString()));
+}
+
+void FDebugLineItem::CopyValueToClipboard() const
+{
+	FPlatformApplicationMisc::ClipboardCopy(ToCStr(GetDescription().ToString()));
+}
+
 TSharedRef<SWidget> FDebugLineItem::GenerateNameWidget()
 {
 	return SNew(PropertyInfoViewStyle::STextHighlightOverlay)
@@ -91,6 +113,33 @@ TSharedRef<SWidget> FDebugLineItem::GenerateValueWidget()
 		];
 }
 
+void FDebugLineItem::MakeMenu(FMenuBuilder& MenuBuilder)
+{
+	const FUIAction CopyName(
+		FExecuteAction::CreateRaw(this, &FDebugLineItem::CopyNameToClipboard),
+		FCanExecuteAction::CreateRaw(this, &FDebugLineItem::HasName)
+	);
+
+	MenuBuilder.AddMenuEntry(
+		LOCTEXT("CopyName", "Copy Name"),
+		LOCTEXT("CopyName_ToolTip", "Copy name to clipboard"),
+		FSlateIcon(),
+		CopyName
+	);
+	
+	const FUIAction CopyValue(
+		FExecuteAction::CreateRaw(this, &FDebugLineItem::CopyValueToClipboard),
+		FCanExecuteAction::CreateRaw(this, &FDebugLineItem::HasValue)
+	);
+
+	MenuBuilder.AddMenuEntry(
+		LOCTEXT("CopyValue", "Copy Value"),
+		LOCTEXT("CopyValue_ToolTip", "Copy value to clipboard"),
+		FSlateIcon(),
+		CopyValue
+	);
+}
+
 void FDebugLineItem::UpdateSearchFlags(bool bIsRootNode)
 {
 	const FString SearchString = SearchBox->GetText().ToString();
@@ -485,6 +534,8 @@ struct FWatchChildLineItem : public FLineItemWithChildren
 {
 protected:
 	FPropertyInstanceInfo Data;
+private:
+	bool bIconHovered = false;
 public:
 	FWatchChildLineItem(const FPropertyInstanceInfo& Child, TSharedPtr<SSearchBox> InSearchBox) :
 		FLineItemWithChildren(DLT_WatchChild, MoveTemp(InSearchBox)),
@@ -524,21 +575,113 @@ public:
 		return Data.DisplayName;
 	}
 
+	// if data is pointing to an asset, get it's UPackage
+	const UPackage *GetDataPackage() const
+	{
+		if(Data.Object.IsValid())
+		{
+			if(const UBlueprintGeneratedClass *GeneratedClass = Cast<UBlueprintGeneratedClass>(Data.Object->GetClass()))
+			{
+				if(const UPackage *Package = GeneratedClass->GetPackage())
+                {
+					return Package;
+                }
+			}
+			if(const UPackage *Package = Data.Object->GetPackage())
+			{
+				return Package;
+			}
+		}
+		return {};
+	}
+
+	// opens result of GetDataPackage in editor
+	FReply OnFocusAsset() const
+	{
+		const UPackage* Package = GetDataPackage();
+		if(!Package)
+		{
+			return FReply::Unhandled();
+		}
+		
+		const FString Path = Package->GetPathName();
+		if(Path.IsEmpty())
+		{
+			return FReply::Unhandled();
+		}
+ 
+		GEditor->GetEditorSubsystem<UAssetEditorSubsystem>()->OpenEditorForAsset(Path);
+		return FReply::Handled();
+	}
+
+	// returns the icon color given a precalculated color associated with this datatype.
+	// the color changes slightly based on whether it's null or a hovered button
+	FSlateColor ModifiedIconColor(FSlateColor BaseColor) const
+	{
+		// check if Data is a UObject
+		if (CastField<FObjectPropertyBase>(Data.Property.Get()))
+		{
+			FLinearColor LinearHSV = BaseColor.GetSpecifiedColor().LinearRGBToHSV();
+			
+			// if it's a null object, darken the icon so it's clear that it's not a button
+			if(Data.Object == nullptr)
+			{
+				LinearHSV.B *= 0.5f; // decrease value
+				LinearHSV.A *= 0.5f; // decrease alpha
+				return LinearHSV.HSVToLinearRGB();
+			}
+			
+			// if the icon is hovered, lighten the icon
+			if(bIconHovered)
+			{
+				LinearHSV.B *= 2.f;  // increase value
+				LinearHSV.G *= 0.8f; // decrease Saturation
+				return LinearHSV.HSVToLinearRGB();
+			}
+		}
+		return BaseColor;
+	}
+
+	FText IconTooltipText() const
+	{
+		const UPackage* Package = GetDataPackage();
+		if(Package)
+		{
+			return FText::Format(LOCTEXT("OpenPackage", "Open: {0}"), FText::FromString(Package->GetName()));
+		}
+		return Data.Type;
+	}
+
 	// uses the icon and color associated with the property type
 	virtual TSharedRef<SWidget> GetNameIcon() override
 	{
-		FSlateColor Color;
+		FSlateColor BaseColor;
 		FSlateColor SecondaryColor;
 		FSlateBrush const* SecondaryIcon;
 		const FSlateBrush* Icon = FBlueprintEditor::GetVarIconAndColorFromProperty(
 			Data.Property.Get(),
-			Color,
+			BaseColor,
 			SecondaryIcon,
 			SecondaryColor
 		);
-		return SNew(SImage)
-			.Image(Icon)
-			.ColorAndOpacity(Color);
+
+		// make the icon a button so the user can open the asset in editor if there is one
+		return SNew(SButton)
+			.OnClicked(this, &FWatchChildLineItem::OnFocusAsset)
+			.ButtonStyle(FEditorStyle::Get(), "NoBorder")
+			.ContentPadding(0.0f)
+			.OnHovered_Lambda(
+				[&bIconHovered = bIconHovered](){bIconHovered = true;}
+			)
+			.OnUnhovered_Lambda(
+				[&bIconHovered = bIconHovered](){bIconHovered = false;}
+			)
+			[
+				SNew(SImage)
+					.Image(Icon)
+					.ColorAndOpacity(this, &FWatchChildLineItem::ModifiedIconColor, BaseColor)
+					.ToolTipText(this, &FWatchChildLineItem::IconTooltipText)
+			];
 	}
 
 	virtual void GatherChildren(TArray<FDebugTreeItemPtr>& OutChildren, bool bRespectSearch) override
@@ -640,7 +783,8 @@ public:
 		if (UEdGraphPin* WatchedPin = ObjectRef.Get())
 		{
 			FUIAction ClearThisWatch(
-				FExecuteAction::CreateStatic( &FDebuggingActionCallbacks::ClearWatch, WatchedPin )
+				FExecuteAction::CreateStatic( &FDebuggingActionCallbacks::ClearWatch, WatchedPin ),
+				FCanExecuteAction() // always allow
 				);
 
 			MenuBuilder.AddMenuEntry(
@@ -649,6 +793,8 @@ public:
 				FSlateIcon(),
 				ClearThisWatch);
 		}
+		
+		FDebugLineItem::MakeMenu(MenuBuilder);
 	}
 	
 	virtual void GatherChildren(TArray<FDebugTreeItemPtr>& OutChildren, bool bRespectSearch) override
@@ -767,6 +913,7 @@ TSharedRef<SWidget> FWatchLineItem::GetNameIcon()
 {
 	const FSlateBrush* PinIcon;
 	FLinearColor PinIconColor;
+	FText Typename;
 	if (UEdGraphPin* ObjectToFocus = ObjectRef.Get())
 	{
 		PinIcon = FBlueprintEditorUtils::GetIconFromPin(ObjectToFocus->PinType);
@@ -774,6 +921,10 @@ TSharedRef<SWidget> FWatchLineItem::GetNameIcon()
 		const UEdGraphSchema* Schema = ObjectToFocus->GetSchema();
 		PinIconColor = Schema->GetPinTypeColor(ObjectToFocus->PinType);
 		PinIconColor.A = 0.3f;
+		
+		UBlueprint* ParentBlueprint = GetBlueprintForObject(ParentObjectRef.Get());
+		FProperty* Property = FKismetDebugUtilities::FindClassPropertyForPin(ParentBlueprint, ObjectToFocus);
+		Typename = UEdGraphSchema_K2::TypeToText(Property);
 	}
 	else
 	{
@@ -781,6 +932,7 @@ TSharedRef<SWidget> FWatchLineItem::GetNameIcon()
 	}
 	
 	return SNew(SOverlay)
+		.ToolTipText(Typename)
 		+ SOverlay::Slot()
 		.Padding(FMargin(10.f, 0.f, 0.f, 0.f))
 		[
@@ -885,6 +1037,8 @@ public:
 				FSlateIcon(),
 				ClearThisBreakpoint);
 		}
+		
+		FDebugLineItem::MakeMenu(MenuBuilder);
 	}
 protected:
 	FBlueprintBreakpoint* GetBreakpoint() const
@@ -1043,6 +1197,76 @@ public:
 		}
 	}
 
+	virtual void MakeMenu(FMenuBuilder& MenuBuilder) override
+	{
+		if (FKismetDebugUtilities::BlueprintHasBreakpoints(Blueprint.Get()))
+		{
+			const FUIAction ClearAllBreakpoints(
+				FExecuteAction::CreateStatic( &FDebuggingActionCallbacks::ClearBreakpoints, Blueprint.Get() ),
+				FCanExecuteAction() // always allow
+			);
+
+			MenuBuilder.AddMenuEntry(
+				LOCTEXT("ClearBreakpoints", "Remove all breakpoints"),
+				LOCTEXT("ClearBreakpoints_ToolTip", "Clear all breakpoints in this blueprint"),
+				FSlateIcon(),
+				ClearAllBreakpoints);
+
+			const bool bEnabledBreakpointExists = FKismetDebugUtilities::FindBreakpointByPredicate(
+				Blueprint.Get(),
+				[](const FBlueprintBreakpoint& Breakpoint)->bool
+				{
+					return Breakpoint.IsEnabled();
+				}
+			) != nullptr;
+
+			if(bEnabledBreakpointExists)
+			{
+				const FUIAction DisableAllBreakpoints(
+					FExecuteAction::CreateStatic(
+						&FDebuggingActionCallbacks::SetEnabledOnAllBreakpoints,
+						const_cast<const UBlueprint*>(Blueprint.Get()),
+						false
+					),
+					FCanExecuteAction() // always allow
+				);
+
+				MenuBuilder.AddMenuEntry(
+					LOCTEXT("DisableBreakpoints", "Disable all breakpoints"),
+					LOCTEXT("DisableBreakpoints_ToolTip", "Disable all breakpoints in this blueprint"),
+					FSlateIcon(),
+					DisableAllBreakpoints);
+			}
+
+			const bool bDisabledBreakpointExists = FKismetDebugUtilities::FindBreakpointByPredicate(
+				Blueprint.Get(),
+				[](const FBlueprintBreakpoint& Breakpoint)->bool
+				{
+					return !Breakpoint.IsEnabled();
+				}
+			) != nullptr;
+
+			if(bDisabledBreakpointExists)
+			{
+				const FUIAction EnableAllBreakpoints(
+					FExecuteAction::CreateStatic(
+						&FDebuggingActionCallbacks::SetEnabledOnAllBreakpoints,
+						const_cast<const UBlueprint*>(Blueprint.Get()),
+						true
+					),
+					FCanExecuteAction() // always allow
+				);
+
+				MenuBuilder.AddMenuEntry(
+					LOCTEXT("EnableBreakpoints", "Enable all breakpoints"),
+					LOCTEXT("EnableBreakpoints_ToolTip", "Enable all breakpoints in this blueprint"),
+					FSlateIcon(),
+					EnableAllBreakpoints);
+			}
+		}
+		FDebugLineItem::MakeMenu(MenuBuilder);
+	}
+
 protected:
 	virtual FText GetDisplayName() const override
 	{
@@ -1214,8 +1438,9 @@ protected:
 			if (FKismetDebugUtilities::BlueprintHasPinWatches(BP))
 			{
 				FUIAction ClearAllWatches(
-					FExecuteAction::CreateStatic( &FDebuggingActionCallbacks::ClearWatches, BP )
-					);
+					FExecuteAction::CreateStatic( &FDebuggingActionCallbacks::ClearWatches, BP ),
+					FCanExecuteAction() // always allow
+				);
 
 				MenuBuilder.AddMenuEntry(
 					LOCTEXT("ClearWatches", "Clear all watches"),
@@ -1223,20 +1448,9 @@ protected:
 					FSlateIcon(),
 					ClearAllWatches);
 			}
-
-			if (FKismetDebugUtilities::BlueprintHasBreakpoints(BP))
-			{
-				FUIAction ClearAllBreakpoints(
-					FExecuteAction::CreateStatic( &FDebuggingActionCallbacks::ClearBreakpoints, BP )
-					);
-
-				MenuBuilder.AddMenuEntry(
-					LOCTEXT("ClearBreakpoints", "Remove all breakpoints"),
-					LOCTEXT("ClearBreakpoints_ToolTip", "Clear all breakpoints in this blueprint"),
-					FSlateIcon(),
-					ClearAllBreakpoints);
-			}
 		}
+	
+		FDebugLineItem::MakeMenu(MenuBuilder);
 	}
 };
 
@@ -1589,13 +1803,23 @@ TSharedRef<SHorizontalBox> SKismetDebuggingView::GetDebugLineTypeToggle(FDebugLi
 		];
 }
 
-TSharedPtr<SWidget> SKismetDebuggingView::OnMakeContextMenu()
+TSharedPtr<SWidget> SKismetDebuggingView::OnMakeDebugTreeContextMenu() const
+{
+	return OnMakeTreeContextMenu(DebugTreeView);
+}
+
+TSharedPtr<SWidget> SKismetDebuggingView::OnMakeOtherTreeContextMenu() const
+{
+	return OnMakeTreeContextMenu(OtherTreeView);
+}
+
+TSharedPtr<SWidget> SKismetDebuggingView::OnMakeTreeContextMenu(const TSharedPtr<STreeView<FDebugTreeItemPtr>>& Tree)
 {
 	FMenuBuilder MenuBuilder(true, nullptr);
 
 	MenuBuilder.BeginSection("DebugActions", LOCTEXT("DebugActionsMenuHeading", "Debug Actions"));
 	{
-		const TArray<FDebugTreeItemPtr> SelectionList = DebugTreeView->GetSelectedItems();
+		TArray<FDebugTreeItemPtr> SelectionList = Tree->GetSelectedItems();
 
 		for (int32 SelIndex = 0; SelIndex < SelectionList.Num(); ++SelIndex)
 		{
@@ -1805,7 +2029,7 @@ void SKismetDebuggingView::Construct(const FArguments& InArgs)
 							.SelectionMode( ESelectionMode::Single )
 							.OnGetChildren( this, &SKismetDebuggingView::OnGetChildrenForWatchTree )
 							.OnGenerateRow( this, &SKismetDebuggingView::OnGenerateRowForWatchTree ) 
-							.OnContextMenuOpening( this, &SKismetDebuggingView::OnMakeContextMenu )
+							.OnContextMenuOpening( this, &SKismetDebuggingView::OnMakeDebugTreeContextMenu )
 							.TreeViewStyle(&FAppStyle::Get().GetWidgetStyle<FTableViewStyle>("PropertyTable.InViewport.ListView"))
 							.HeaderRow
 							(
@@ -1823,7 +2047,7 @@ void SKismetDebuggingView::Construct(const FArguments& InArgs)
 							.SelectionMode( ESelectionMode::Single )
 							.OnGetChildren( this, &SKismetDebuggingView::OnGetChildrenForWatchTree )
 							.OnGenerateRow( this, &SKismetDebuggingView::OnGenerateRowForWatchTree ) 
-							.OnContextMenuOpening( this, &SKismetDebuggingView::OnMakeContextMenu )
+							.OnContextMenuOpening( this, &SKismetDebuggingView::OnMakeOtherTreeContextMenu )
 							.TreeViewStyle(&FAppStyle::Get().GetWidgetStyle<FTableViewStyle>("PropertyTable.InViewport.ListView"))
 							.HeaderRow
 							(
diff --git a/Engine/Source/Editor/Kismet/Private/Debugging/SKismetDebuggingView.h b/Engine/Source/Editor/Kismet/Private/Debugging/SKismetDebuggingView.h
--- a/Engine/Source/Editor/Kismet/Private/Debugging/SKismetDebuggingView.h
+++ b/Engine/Source/Editor/Kismet/Private/Debugging/SKismetDebuggingView.h
@@ -49,7 +49,8 @@ public:
 	virtual TSharedRef<SWidget> GenerateValueWidget();
 
 	// Add any context menu items that can act on this node
-	virtual void MakeMenu(class FMenuBuilder& MenuBuilder) { }
+	// If you override, make sure FDebugLineItem::MakeMenu still gets called
+	virtual void MakeMenu(class FMenuBuilder& MenuBuilder);
 
 	// Gather all of the children
 	virtual void GatherChildren(TArray<FDebugTreeItemPtr>& OutChildren, bool bRespectSearch = true) {}
@@ -114,6 +115,11 @@ protected:
 
 	// @return The text to display in the value column, unless GenerateValueWidget is overridden
 	virtual FText GetDescription() const;
+
+	bool HasName() const;
+	bool HasValue() const;
+	void CopyNameToClipboard() const;
+	void CopyValueToClipboard() const;
 protected:
 	// Type of action (poor mans RTTI for the tree, really only used to accelerate Compare checks)
 	EDebugLineType Type;
@@ -162,7 +168,9 @@ protected:
 
 	static TSharedRef<SHorizontalBox> GetDebugLineTypeToggle(FDebugLineItem::EDebugLineType Type, const FText& Text);
 	
-	TSharedPtr<SWidget> OnMakeContextMenu();
+	TSharedPtr<SWidget> OnMakeDebugTreeContextMenu() const;
+	TSharedPtr<SWidget> OnMakeOtherTreeContextMenu() const;
+	static TSharedPtr<SWidget> OnMakeTreeContextMenu(const TSharedPtr<STreeView<FDebugTreeItemPtr>>& Tree);
 
 	// called when SearchBox query is committed by the user
 	void OnSearchTextCommitted(const FText& Text, ETextCommit::Type);
diff --git a/Engine/Source/Editor/UnrealEd/Private/Kismet2/KismetDebugUtilities.cpp b/Engine/Source/Editor/UnrealEd/Private/Kismet2/KismetDebugUtilities.cpp
--- a/Engine/Source/Editor/UnrealEd/Private/Kismet2/KismetDebugUtilities.cpp
+++ b/Engine/Source/Editor/UnrealEd/Private/Kismet2/KismetDebugUtilities.cpp
@@ -1829,9 +1829,10 @@ FPropertyInstanceInfo::FPropertyInstanceInfo(FPropertyInstance PropertyInstance)
 	}
 	else if (FObjectPropertyBase* ObjectPropertyBase = CastField<FObjectPropertyBase>(Property.Get()))
 	{
-		if (UObject* Obj = ObjectPropertyBase->GetObjectPropertyValue(PropertyInstance.Value))
+		Object = ObjectPropertyBase->GetObjectPropertyValue(PropertyInstance.Value);
+		if (Object.IsValid())
 		{
-			Value = FText::FromString(Obj->GetFullName());
+			Value = FText::FromString(Object->GetFullName());
 		}
 		else
 		{
diff --git a/Engine/Source/Editor/UnrealEd/Public/Kismet2/KismetDebugUtilities.h b/Engine/Source/Editor/UnrealEd/Public/Kismet2/KismetDebugUtilities.h
--- a/Engine/Source/Editor/UnrealEd/Public/Kismet2/KismetDebugUtilities.h
+++ b/Engine/Source/Editor/UnrealEd/Public/Kismet2/KismetDebugUtilities.h
@@ -65,6 +65,7 @@ struct FPropertyInstanceInfo
 	FText DisplayName;
 	FText Value;
 	FText Type;
+	TWeakObjectPtr<UObject> Object = nullptr; // only filled if property is a UObject
 	TFieldPath<FProperty> Property;
 	TArray<TSharedPtr<FPropertyInstanceInfo>> Children;
 };
