BP Debugger optimizations
- swapped children mirrors to a TSet for O(1) access
- tree only updates a maximum of 2 times a second <- this can be reduced further if needed

diff --git a/Engine/Source/Editor/Kismet/Private/Debugging/SKismetDebuggingView.cpp b/Engine/Source/Editor/Kismet/Private/Debugging/SKismetDebuggingView.cpp
--- a/Engine/Source/Editor/Kismet/Private/Debugging/SKismetDebuggingView.cpp
+++ b/Engine/Source/Editor/Kismet/Private/Debugging/SKismetDebuggingView.cpp
@@ -268,6 +268,12 @@ public:
 
 	virtual bool CanHaveChildren() override { return true; }
 
+	/**
+	* returns whether this node should be visible according to the users
+	* search query
+	*
+	* O( number of recursive children )
+	*/
 	bool SearchRecursive(TArray<FLineItemWithChildren*>& Parents,
 		TSharedPtr<STreeView<FDebugTreeItemPtr>> DebugTreeView)
 	{
@@ -280,7 +286,7 @@ public:
 		Parents.Push(this);
 		
 		TArray<FDebugTreeItemPtr> Children;
-		GatherChildren(Children, false);
+		GatherChildrenBase(Children, false);
 		for(const FDebugTreeItemPtr& ChildRef : Children)
 		{
 			if(ChildRef->CanHaveChildren())
@@ -346,47 +352,89 @@ public:
 		return bVisible;
 	}
 
+	// ensures that ChildrenMirrors are set up for calls to EnsureChildIsAdded
+	virtual void GatherChildrenBase(TArray<FDebugTreeItemPtr>& OutChildren, bool bRespectSearch) override
+	{
+		Swap(PrevChildrenMirrors, ChildrenMirrors);
+		ChildrenMirrors = {};
+		GatherChildren(OutChildren, bRespectSearch);
+	}
+
+	// allows FDebugTreeItemPtr to be stored in TSets 
+	class FDebugTreeItemKeyFuncs
+	{
+	public:
+		typedef FDebugTreeItemPtr ElementType;
+		typedef TTypeTraits<ElementType>::ConstPointerType KeyInitType;
+		typedef TCallTraits<ElementType>::ParamType ElementInitType;
+		enum { bAllowDuplicateKeys = false };
+		
+		/**
+		* @return The key used to index the given element.
+		*/
+		static FORCEINLINE KeyInitType GetSetKey(ElementInitType Element)
+		{
+			return Element;
+		}
+
+		/**
+		* @return True if the keys match.
+		*/
+		static FORCEINLINE bool Matches(KeyInitType A, KeyInitType B)
+		{
+			FDebugLineItem *APtr = A.Get();
+			FDebugLineItem *BPtr = B.Get();
+			if(APtr && BPtr)
+			{
+				return (APtr->Type == BPtr->Type) && APtr->Compare(BPtr);
+			}
+			return APtr == BPtr;
+		}
+
+		/** Calculates a hash index for a key. */
+		static FORCEINLINE uint32 GetKeyHash(KeyInitType Key)
+		{
+			if(FDebugLineItem *KeyPtr = Key.Get())
+			{
+				return KeyPtr->GetHash();
+			}
+			return GetTypeHash(Key);
+		}
+	};
 protected:
-	// List of children 
-	TArray<FDebugTreeItemPtr> ChildrenMirrors;
+	// Last frames cached children
+	TSet<FDebugTreeItemPtr, FDebugTreeItemKeyFuncs> PrevChildrenMirrors;
+	// This frames children
+	TSet<FDebugTreeItemPtr, FDebugTreeItemKeyFuncs> ChildrenMirrors;
+
+	virtual void GatherChildren(TArray<FDebugTreeItemPtr>& OutChildren, bool bRespectSearch) {}
 	
 	/**
 	 * Adds either Item or an identical node that was previously
 	 * created (present in ChildrenMirrors) as a child to OutChildren.
 	 *
-	 * O( # Children )
+	 * O( 1 )
 	 */
 	void EnsureChildIsAdded(TArray<FDebugTreeItemPtr>& OutChildren, const FDebugLineItem& Item, bool bRespectSearch)
 	{
-		for (int32 i = 0; i < ChildrenMirrors.Num(); ++i)
+		const FDebugTreeItemPtr Shareable = MakeShareable(Item.Duplicate());
+		if(FDebugTreeItemPtr *Found = PrevChildrenMirrors.Find(Shareable))
 		{
-			TSharedPtr< FDebugLineItem > MirrorItem = ChildrenMirrors[i];
-			if(bParentsMatchSearch)
-			{
-				// propogate parents search state to children
-				MirrorItem->bParentsMatchSearch = true;
-			}
-
-			if (MirrorItem->Type == Item.Type)
+			FDebugTreeItemPtr FoundItem = *Found;
+			FoundItem->UpdateData(Item);
+			ChildrenMirrors.Add(FoundItem);
+			
+			// only add item if it matches search
+			if(!bRespectSearch || SearchBox->GetText().IsEmpty() || FoundItem->IsVisible() || FoundItem->DoParentsMatchSearch())
 			{
-				if (Item.Compare(MirrorItem.Get()))
-				{
-					MirrorItem->UpdateData(Item);
-					
-					// only add item if it matches search
-					if(!bRespectSearch || SearchBox->GetText().IsEmpty() || MirrorItem->IsVisible() || MirrorItem->DoParentsMatchSearch())
-					{
-						OutChildren.Add(MirrorItem);
-					}
-					return;
-				}
+				OutChildren.Add(FoundItem);
 			}
 		}
-
-		TSharedPtr< FDebugLineItem > Result = MakeShareable(Item.Duplicate());
-		ChildrenMirrors.Add(Result);
-
-		OutChildren.Add(Result);
+		else
+		{
+			ChildrenMirrors.Add(Shareable);
+			OutChildren.Add(Shareable);
+		}
 	}
 };
 
@@ -421,6 +469,11 @@ protected:
 	{
 		return FText::FromString(Message);
 	}
+
+	virtual uint32 GetHash() override
+	{
+		return GetTypeHash(Message);
+	}
 };
 
 //////////////////////////////////////////////////////////////////////////
@@ -452,6 +505,11 @@ protected:
 	{
 		return new FLatentActionLineItem(UUID, ParentObjectRef.Get(), SearchBox);
 	}
+
+	virtual uint32 GetHash() override
+	{
+		return HashCombine(GetTypeHash(UUID), GetTypeHash(ParentObjectRef));
+	}
 protected:
 	virtual TSharedRef<SWidget> GenerateNameWidget() override;
 	virtual TSharedRef<SWidget> GetNameIcon() override;
@@ -563,6 +621,11 @@ public:
 	{
 		return new FWatchChildLineItem(Data, SearchBox);
 	}
+
+	virtual uint32 GetHash() override
+	{
+		return HashCombine(GetTypeHash(Data.Property), GetTypeHash(Data.DisplayName.ToString()));
+	}
 	
 	virtual FText GetDescription() const override
 	{
@@ -687,9 +750,9 @@ public:
 	virtual void GatherChildren(TArray<FDebugTreeItemPtr>& OutChildren, bool bRespectSearch) override
 	{
 		for (const TSharedPtr<FPropertyInstanceInfo>& ChildData : Data.Children)
-		{
-			EnsureChildIsAdded(OutChildren, FWatchChildLineItem(*ChildData, SearchBox), bRespectSearch);
-		}
+        {
+            EnsureChildIsAdded(OutChildren, FWatchChildLineItem(*ChildData, SearchBox), bRespectSearch);
+        }
 	}
 };
 
@@ -716,6 +779,11 @@ public:
 	{
 		return new FSelfWatchLineItem(ObjectToWatch.Get(), SearchBox);
 	}
+
+	virtual uint32 GetHash() override
+	{
+		return GetTypeHash(ObjectToWatch);
+	}
 	
 	virtual void GatherChildren(TArray<FDebugTreeItemPtr>& OutChildren, bool bRespectSearch) override
 	{
@@ -778,6 +846,11 @@ public:
 		return new FWatchLineItem(ObjectRef.Get(), ParentObjectRef.Get(), SearchBox);
 	}	
 
+	virtual uint32 GetHash() override
+	{
+		return HashCombine(GetTypeHash(ParentObjectRef), GetTypeHash(ObjectRef));
+	}
+
 	virtual void MakeMenu(class FMenuBuilder& MenuBuilder) override
 	{
 		if (UEdGraphPin* WatchedPin = ObjectRef.Get())
@@ -984,6 +1057,11 @@ public:
 		return new FBreakpointLineItem(BreakpointNode, ParentObjectRef.Get(), SearchBox);
 	}	
 
+	virtual uint32 GetHash() override
+	{
+		return HashCombine(GetTypeHash(ParentObjectRef), GetTypeHash(BreakpointNode));
+	}
+
 	virtual void MakeMenu(class FMenuBuilder& MenuBuilder) override
 	{
 		FBlueprintBreakpoint* Breakpoint = GetBreakpoint();
@@ -1284,6 +1362,12 @@ protected:
 		check(false);
 		return nullptr;
 	}
+
+	virtual uint32 GetHash() override
+	{
+		check(false);
+		return 0;
+	}
 };
 
 //////////////////////////////////////////////////////////////////////////
@@ -1408,6 +1492,11 @@ protected:
 		return new FParentLineItem(ObjectRef.Get(), SearchBox);
 	}
 
+	virtual uint32 GetHash() override
+	{
+		return GetTypeHash(ObjectRef);
+	}
+
 	virtual FText GetDisplayName() const override
 	{
 		UObject* Object = ObjectRef.Get();
@@ -1480,6 +1569,12 @@ protected:
 		return nullptr;
 	}
 
+	virtual uint32 GetHash() override
+	{
+		check(false);
+		return 0;
+	}
+
 	UEdGraphNode* GetNode() const
 	{
 		const TSimpleRingBuffer<FKismetTraceSample>& TraceStack = FKismetDebugUtilities::GetTraceStack();
@@ -1630,6 +1725,11 @@ public:
 		: FLineItemWithChildren(DLT_TraceStackParent, MoveTemp(InSearchBox))
 	{
 	}
+	
+	virtual bool HasChildren() override
+	{
+		return !ChildrenMirrorsArr.IsEmpty();
+	}
 
 	virtual void GatherChildren(TArray<FDebugTreeItemPtr>& OutChildren, bool bRespectSearch) override
 	{
@@ -1640,15 +1740,15 @@ public:
 		const int32 NumVisible = TraceStack.Num();
 
 		// Create any new stack entries that are needed
-		for (int32 i = ChildrenMirrors.Num(); i < NumVisible; ++i)
+		for (int32 i = ChildrenMirrorsArr.Num(); i < NumVisible; ++i)
 		{
-			ChildrenMirrors.Add(MakeShareable( new FTraceStackChildItem(i, SearchBox) ));
+			ChildrenMirrorsArr.Add(MakeShareable( new FTraceStackChildItem(i, SearchBox) ));
 		}
 
 		// Add the visible stack entries as children
 		for (int32 i = 0; i < NumVisible; ++i)
 		{
-			OutChildren.Add(ChildrenMirrors[i]);
+			OutChildren.Add(ChildrenMirrorsArr[i]);
 		}
 	}
 
@@ -1669,6 +1769,15 @@ protected:
 		check(false);
 		return nullptr;
 	}
+
+	virtual uint32 GetHash() override
+	{
+		check(false);
+		return 0;
+	}
+
+	// use an array to store children mirrors instead of a set so it's ordered
+	TArray<FDebugTreeItemPtr> ChildrenMirrorsArr;
 };
 
 //////////////////////////////////////////////////////////////////////////
@@ -1780,7 +1889,7 @@ TSharedRef<ITableRow> SKismetDebuggingView::OnGenerateRowForWatchTree(FDebugTree
 
 void SKismetDebuggingView::OnGetChildrenForWatchTree(FDebugTreeItemPtr InParent, TArray<FDebugTreeItemPtr>& OutChildren)
 {
-	InParent->GatherChildren(OutChildren);
+	InParent->GatherChildrenBase(OutChildren);
 }
 
 TSharedRef<SHorizontalBox> SKismetDebuggingView::GetDebugLineTypeToggle(FDebugLineItem::EDebugLineType Type, const FText& Text)
@@ -2092,6 +2201,20 @@ void SKismetDebuggingView::Construct(const FArguments& InArgs)
 
 void SKismetDebuggingView::Tick( const FGeometry& AllottedGeometry, const double InCurrentTime, const float InDeltaTime )
 {
+	// don't update during scroll. this will help make the scroll smoother
+	if(DebugTreeView->IsScrolling() || OtherTreeView->IsScrolling())
+	{
+		return;
+	}
+	
+	// update less often to avoid lag
+	TreeUpdateTimer += InDeltaTime;
+	if(TreeUpdateTimer < UpdateInterval)
+	{
+		return;
+	}
+	TreeUpdateTimer = 0.f;
+	
 	// Gather the old root set
 	TSet<UObject*> OldRootSet;
 	for (int32 i = 0; i < RootTreeItems.Num(); ++i)
diff --git a/Engine/Source/Editor/Kismet/Private/Debugging/SKismetDebuggingView.h b/Engine/Source/Editor/Kismet/Private/Debugging/SKismetDebuggingView.h
--- a/Engine/Source/Editor/Kismet/Private/Debugging/SKismetDebuggingView.h
+++ b/Engine/Source/Editor/Kismet/Private/Debugging/SKismetDebuggingView.h
@@ -53,7 +53,7 @@ public:
 	virtual void MakeMenu(class FMenuBuilder& MenuBuilder);
 
 	// Gather all of the children
-	virtual void GatherChildren(TArray<FDebugTreeItemPtr>& OutChildren, bool bRespectSearch = true) {}
+	virtual void GatherChildrenBase(TArray<FDebugTreeItemPtr>& OutChildren, bool bRespectSearch = true) {}
 
 	// returns whether this tree node has children (used by drop down arrows)
 	virtual bool HasChildren();
@@ -106,6 +106,9 @@ protected:
 	// Compare this item to another of the same type
 	virtual bool Compare(const FDebugLineItem* Other) const=0;
 
+	// used for sets
+	virtual uint32 GetHash() = 0;
+
 	// Used to update the state of a line item rather than replace it.
 	// called after Compare returns true
 	[[maybe_unused]] virtual void UpdateData(const FDebugLineItem& NewerData) {}
@@ -193,4 +196,10 @@ protected:
 
 	// Search Box for tree
 	TSharedPtr<SSearchBox> SearchBox;
+
+	// updating the tree every tick is slow. use this to
+	// update less frequently
+	static constexpr uint8 TreeUpdatesPerSecond = 2;
+	static constexpr float UpdateInterval = 1.f / TreeUpdatesPerSecond;
+	float TreeUpdateTimer = UpdateInterval;
 };
diff --git a/Engine/Source/Runtime/Slate/Private/Widgets/Views/STableViewBase.cpp b/Engine/Source/Runtime/Slate/Private/Widgets/Views/STableViewBase.cpp
--- a/Engine/Source/Runtime/Slate/Private/Widgets/Views/STableViewBase.cpp
+++ b/Engine/Source/Runtime/Slate/Private/Widgets/Views/STableViewBase.cpp
@@ -991,6 +991,11 @@ void STableViewBase::ScrollToBottom()
 	RequestLayoutRefresh();
 }
 
+bool STableViewBase::IsScrolling() const
+{
+	return ScrollBar->IsScrolling();
+}
+
 FVector2D STableViewBase::GetScrollDistance()
 {
 	return FVector2D( 0, ScrollBar->DistanceFromTop() );
diff --git a/Engine/Source/Runtime/Slate/Public/Widgets/Views/STableViewBase.h b/Engine/Source/Runtime/Slate/Public/Widgets/Views/STableViewBase.h
--- a/Engine/Source/Runtime/Slate/Public/Widgets/Views/STableViewBase.h
+++ b/Engine/Source/Runtime/Slate/Public/Widgets/Views/STableViewBase.h
@@ -141,6 +141,9 @@ public:
 	/** Scrolls the view to the bottom */
 	void ScrollToBottom();
 
+	/** Returns whether the attached scrollbar is scrolling */
+	bool IsScrolling() const;
+
 	/** Gets the scroll offset of this view (in items) */
 	float GetScrollOffset() const;
 
