Fix for an editor crash when expanding watched tunnel node pins in the consolidated BP debugging tree view.

diff --git a/Engine/Source/Editor/Kismet/Private/Debugging/SKismetDebuggingView.cpp b/Engine/Source/Editor/Kismet/Private/Debugging/SKismetDebuggingView.cpp
--- a/Engine/Source/Editor/Kismet/Private/Debugging/SKismetDebuggingView.cpp
+++ b/Engine/Source/Editor/Kismet/Private/Debugging/SKismetDebuggingView.cpp
@@ -995,9 +995,16 @@ TSharedRef<SWidget> FWatchLineItem::GetNameIcon()
 		PinIconColor = Schema->GetPinTypeColor(ObjectToFocus->PinType);
 		PinIconColor.A = 0.3f;
 		
+		// Note: Currently tunnel node pins (e.g. macro/consolidated node outputs) won't return a property.
 		UBlueprint* ParentBlueprint = GetBlueprintForObject(ParentObjectRef.Get());
-		FProperty* Property = FKismetDebugUtilities::FindClassPropertyForPin(ParentBlueprint, ObjectToFocus);
-		Typename = UEdGraphSchema_K2::TypeToText(Property);
+		if (FProperty* Property = FKismetDebugUtilities::FindClassPropertyForPin(ParentBlueprint, ObjectToFocus))
+		{
+			Typename = UEdGraphSchema_K2::TypeToText(Property);
+		}
+		else
+		{
+			Typename = UEdGraphSchema_K2::TypeToText(ObjectToFocus->PinType);
+		}
 	}
 	else
 	{
