Control Rig: Display local variables in details panel

diff --git a/Engine/Plugins/Animation/ControlRig/Source/ControlRigEditor/Private/ControlRigEditorModule.cpp b/Engine/Plugins/Animation/ControlRig/Source/ControlRigEditor/Private/ControlRigEditorModule.cpp
--- a/Engine/Plugins/Animation/ControlRig/Source/ControlRigEditor/Private/ControlRigEditorModule.cpp
+++ b/Engine/Plugins/Animation/ControlRig/Source/ControlRigEditor/Private/ControlRigEditorModule.cpp
@@ -71,6 +71,7 @@
 #include "ControlRigCompilerDetails.h"
 #include "ControlRigDrawingDetails.h"
 #include "ControlRigGraphDetails.h"
+#include "ControlRigLocalVariableDetails.h"
 #include "ControlRigInfluenceMapDetails.h"
 #include "Animation/AnimSequence.h"
 #include "Editor/SControlRigProfilingView.h"
@@ -235,6 +236,10 @@ void FControlRigEditorModule::StartupModule()
 	PropertiesToUnregisterOnShutdown.Add(FRigNullElement::StaticStruct()->GetFName());
 	PropertyEditorModule.RegisterCustomPropertyTypeLayout(PropertiesToUnregisterOnShutdown.Last(), FOnGetPropertyTypeCustomizationInstance::CreateStatic(&FRigNullElementDetails::MakeInstance));
 
+	PropertiesToUnregisterOnShutdown.Add(FRigVMGraphVariableDescription::StaticStruct()->GetFName());
+	PropertyEditorModule.RegisterCustomPropertyTypeLayout(PropertiesToUnregisterOnShutdown.Last(), FOnGetPropertyTypeCustomizationInstance::CreateStatic(&FRigVMLocalVariableDetails::MakeInstance));
+
+
 	// Register asset tools
 	auto RegisterAssetTypeAction = [this](const TSharedRef<IAssetTypeActions>& InAssetTypeAction)
 	{
@@ -383,6 +388,7 @@ void FControlRigEditorModule::ShutdownModule()
 		if (BlueprintEditorModule)
 		{
 			BlueprintEditorModule->UnregisterVariableCustomization(FProperty::StaticClass());
+			BlueprintEditorModule->UnregisterLocalVariableCustomization(FProperty::StaticClass());
 			BlueprintEditorModule->UnregisterGraphCustomization(GetDefault<UControlRigGraphSchema>());
 		}
 	}
diff --git a/Engine/Plugins/Animation/ControlRig/Source/ControlRigEditor/Private/ControlRigLocalVariableDetails.cpp b/Engine/Plugins/Animation/ControlRig/Source/ControlRigEditor/Private/ControlRigLocalVariableDetails.cpp
new file mode 100644
--- /dev/null
+++ b/Engine/Plugins/Animation/ControlRig/Source/ControlRigEditor/Private/ControlRigLocalVariableDetails.cpp
@@ -0,0 +1,257 @@
+﻿// Copyright Epic Games, Inc. All Rights Reserved.
+
+#include "ControlRigLocalVariableDetails.h"
+
+#include "DetailWidgetRow.h"
+#include "DetailLayoutBuilder.h"
+#include "IDetailChildrenBuilder.h"
+#include "SPinTypeSelector.h"
+#include "Graph/ControlRigGraphSchema.h"
+#include "Kismet2/BlueprintEditorUtils.h"
+#include "Widgets/Input/STextComboBox.h"
+
+#define LOCTEXT_NAMESPACE "LocalVariableDetails"
+
+void LocalVariableDetails_GetCustomizedInfo(TSharedRef<IPropertyHandle> InStructPropertyHandle, UControlRigBlueprint*& OutBlueprint)
+{
+	TArray<UObject*> Objects;
+	InStructPropertyHandle->GetOuterObjects(Objects);
+	for (UObject* Object : Objects)
+	{
+		if (Object->IsA<UControlRigBlueprint>())
+		{
+			OutBlueprint = Cast<UControlRigBlueprint>(Object);
+			if (OutBlueprint)
+			{
+				break;
+			}
+		}
+	}
+
+	if (OutBlueprint == nullptr)
+	{
+		TArray<UPackage*> Packages;
+		InStructPropertyHandle->GetOuterPackages(Packages);
+		for (UPackage* Package : Packages)
+		{
+			if (Package == nullptr)
+			{
+				continue;
+			}
+
+			TArray<UObject*> SubObjects;
+			Package->GetDefaultSubobjects(SubObjects);
+			for (UObject* SubObject : SubObjects)
+			{
+				if (UControlRig* Rig = Cast<UControlRig>(SubObject))
+				{
+					UControlRigBlueprint* Blueprint = Cast<UControlRigBlueprint>(Rig->GetClass()->ClassGeneratedBy);
+					if (Blueprint)
+					{
+						if(Blueprint->GetOutermost() == Package)
+						{
+							OutBlueprint = Blueprint;
+							break;
+						}
+					}
+				}
+			}
+
+			if (OutBlueprint)
+			{
+				break;
+			}
+		}
+	}
+}
+
+void FRigVMLocalVariableDetails::CustomizeHeader(TSharedRef<IPropertyHandle> InStructPropertyHandle, FDetailWidgetRow& HeaderRow, IPropertyTypeCustomizationUtils& StructCustomizationUtils)
+{
+	ObjectsBeingCustomized.Reset();
+	
+	TArray<UObject*> Objects;
+	InStructPropertyHandle->GetOuterObjects(Objects);
+	for (UObject* Object : Objects)
+	{
+		UDetailsViewWrapperObject* WrapperObject = CastChecked<UDetailsViewWrapperObject>(Object);
+		ObjectsBeingCustomized.Add(WrapperObject);
+	}
+
+	if(ObjectsBeingCustomized[0].IsValid())
+	{
+		VariableDescription = *ObjectsBeingCustomized[0]->GetContent<FRigVMGraphVariableDescription>();
+		GraphBeingCustomized = ObjectsBeingCustomized[0]->GetTypedOuter<URigVMGraph>();
+	}
+}
+
+void FRigVMLocalVariableDetails::CustomizeChildren(TSharedRef<IPropertyHandle> InStructPropertyHandle, IDetailChildrenBuilder& StructBuilder, IPropertyTypeCustomizationUtils& StructCustomizationUtils)
+{
+	StructBuilder.GetParentCategory().GetParentLayout().HideCategory("RigVMGraphVariableDescription");
+	IDetailCategoryBuilder& Category = StructBuilder.GetParentCategory().GetParentLayout().EditCategory("Local Variable");
+		
+	NameHandle = InStructPropertyHandle->GetChildHandle(TEXT("Name"));
+	TypeHandle = InStructPropertyHandle->GetChildHandle(TEXT("CPPType"));
+	TypeObjectHandle = InStructPropertyHandle->GetChildHandle(TEXT("CPPTypeObject"));
+	DefaultValueHandle = InStructPropertyHandle->GetChildHandle(TEXT("DefaultValue"));
+	
+	const UEdGraphSchema* Schema = GetDefault<UControlRigGraphSchema>();
+
+	const FSlateFontInfo DetailFontInfo = IDetailLayoutBuilder::GetDetailFont();
+	Category.AddCustomRow( LOCTEXT("LocalVariableName", "Variable Name") )
+	.NameContent()
+	[
+		SNew(STextBlock)
+		.Text(LOCTEXT("LocalVariableName", "Variable Name"))
+		.Font(IDetailLayoutBuilder::GetDetailFont())
+	]
+	.ValueContent()
+	.MaxDesiredWidth(250.0f)
+	[
+		NameHandle->CreatePropertyValueWidget()
+	];
+
+	TSharedPtr<IPinTypeSelectorFilter> CustomPinTypeFilter;
+	Category.AddCustomRow(LOCTEXT("VariableTypeLabel", "Variable Type"))
+		.NameContent()
+		[
+			SNew(STextBlock)
+			.Text(LOCTEXT("VariableTypeLabel", "Variable Type"))
+			.Font(DetailFontInfo)
+		]
+		.ValueContent()
+		.MaxDesiredWidth(980.f)
+		[
+			SNew(SPinTypeSelector, FGetPinTypeTree::CreateUObject(GetDefault<UEdGraphSchema_K2>(), &UEdGraphSchema_K2::GetVariableTypeTree))
+			.TargetPinType(this, &FRigVMLocalVariableDetails::OnGetPinInfo)
+			.OnPinTypeChanged(this, &FRigVMLocalVariableDetails::HandlePinInfoChanged)
+			.Schema(Schema)
+			.TypeTreeFilter(ETypeTreeFilter::None)
+			.Font(DetailFontInfo)
+			.CustomFilter(CustomPinTypeFilter)
+		];
+
+
+	IDetailCategoryBuilder& DefaultValueCategory = StructBuilder.GetParentCategory().GetParentLayout().EditCategory(TEXT("DefaultValueCategory"), LOCTEXT("DefaultValueCategoryHeading", "Default Value"));
+
+	if (!VariableDescription.CPPTypeObject)
+	{
+		if (VariableDescription.CPPType == "bool")
+		{
+			DefaultValueCategory.AddCustomRow(LOCTEXT("DefaultValue", "Default Value"))
+			.NameContent()
+			[
+				SNew(STextBlock)
+				.Text(LOCTEXT("DefaultValue", "Default Value"))
+				.Font(DetailFontInfo)
+			]
+			.ValueContent()
+			.MaxDesiredWidth(980.f)
+			[
+				SNew(SCheckBox)
+				.IsChecked(this, &FRigVMLocalVariableDetails::HandleBoolDefaultValueIsChecked)
+				.OnCheckStateChanged(this, &FRigVMLocalVariableDetails::OnBoolDefaultValueChanged)
+			];
+		}
+		else
+		{
+			DefaultValueCategory.AddCustomRow( LOCTEXT("LocalVariableDefaultValue", "Default Value") )
+			.NameContent()
+			[
+				SNew(STextBlock)
+				.Text(LOCTEXT("LocalVariableDefaultValue", "Default Value"))
+				.Font(IDetailLayoutBuilder::GetDetailFont())
+			]
+			.ValueContent()
+			.MaxDesiredWidth(250.0f)
+			[
+				DefaultValueHandle->CreatePropertyValueWidget()
+			];
+		}
+	}
+	else
+	{
+		if(UEnum* Enum = Cast<UEnum>(VariableDescription.CPPTypeObject))
+		{
+			int32 CurrentValueIndex = 0;
+			for (int32 i = 0; i < Enum->NumEnums()-1; i++)
+			{
+				if (!Enum->HasMetaData(TEXT("Hidden"), i))
+				{
+					FString DisplayName = Enum->GetDisplayNameTextByIndex(i).ToString();
+					EnumOptions.Add(MakeShareable(new FString(DisplayName)));
+					if (DisplayName == VariableDescription.DefaultValue)
+					{
+						CurrentValueIndex = i;
+					}					
+				}
+			}
+			
+			DefaultValueCategory.AddCustomRow(LOCTEXT("VariableReplicationConditionsLabel", "Replication Condition"))
+			.NameContent()
+			[
+				SNew(STextBlock)
+				.Text(LOCTEXT("LocalVariableDefaultValue", "Default Value"))
+				.Font(DetailFontInfo)
+			]
+			.ValueContent()
+			[
+				SNew(STextComboBox)
+				.OptionsSource(&EnumOptions)
+				.InitiallySelectedItem(EnumOptions[CurrentValueIndex])
+				.OnSelectionChanged_Lambda([this](TSharedPtr<FString> InValue, ESelectInfo::Type)
+				{
+					VariableDescription.DefaultValue = *InValue;
+					DefaultValueHandle->SetValue(VariableDescription.DefaultValue);					
+				})
+			];			
+		}
+		else if(UScriptStruct* ScriptStruct = Cast<UScriptStruct>(VariableDescription.CPPTypeObject))
+		{
+			TSharedPtr<FStructOnScope> StructOnScope = MakeShareable(new FStructOnScope(ScriptStruct));
+			IDetailPropertyRow* Row = DefaultValueCategory.AddExternalStructureProperty(StructOnScope, NAME_None);
+			Row->GetPropertyHandle()->SetValueFromFormattedString(VariableDescription.DefaultValue);
+			
+			Row->GetPropertyHandle()->SetOnChildPropertyValueChanged(FSimpleDelegate::CreateLambda(
+				[this, Row]()
+				{
+					VariableDescription.DefaultValue.Reset();
+					Row->GetPropertyHandle()->GetValueAsFormattedString(VariableDescription.DefaultValue);
+					DefaultValueHandle->SetValue(VariableDescription.DefaultValue);
+				}
+			));
+		}
+		else
+		{
+			checkNoEntry();
+		}
+	}
+}
+
+FEdGraphPinType FRigVMLocalVariableDetails::OnGetPinInfo() const
+{
+	if (!VariableDescription.Name.IsNone())
+	{
+		return VariableDescription.ToPinType();
+	}
+	return FEdGraphPinType();
+}
+
+void FRigVMLocalVariableDetails::HandlePinInfoChanged(const FEdGraphPinType& PinType)
+{
+	VariableDescription.ChangeType(PinType);
+	TypeHandle->SetValue(VariableDescription.CPPType);
+	TypeObjectHandle->SetValue(VariableDescription.CPPTypeObject);
+}
+
+ECheckBoxState FRigVMLocalVariableDetails::HandleBoolDefaultValueIsChecked() const
+{
+	return VariableDescription.DefaultValue == "1" ? ECheckBoxState::Checked : ECheckBoxState::Unchecked;
+}
+
+void FRigVMLocalVariableDetails::OnBoolDefaultValueChanged(ECheckBoxState InCheckBoxState)
+{
+	VariableDescription.DefaultValue = InCheckBoxState == ECheckBoxState::Checked ? "1" : "0";
+	DefaultValueHandle->SetValue(VariableDescription.DefaultValue);
+}
+
+#undef LOCTEXT_NAMESPACE
diff --git a/Engine/Plugins/Animation/ControlRig/Source/ControlRigEditor/Private/ControlRigLocalVariableDetails.h b/Engine/Plugins/Animation/ControlRig/Source/ControlRigEditor/Private/ControlRigLocalVariableDetails.h
new file mode 100644
--- /dev/null
+++ b/Engine/Plugins/Animation/ControlRig/Source/ControlRigEditor/Private/ControlRigLocalVariableDetails.h
@@ -0,0 +1,55 @@
+﻿// Copyright Epic Games, Inc. All Rights Reserved.
+
+#pragma once
+
+#include "CoreMinimal.h"
+#include "IDetailCustomization.h"
+#include "IPropertyTypeCustomization.h"
+#include "ControlRig.h"
+#include "ControlRigBlueprint.h"
+#include "DetailsViewWrapperObject.h"
+#include "Graph/ControlRigGraph.h"
+#include "Graph/SControlRigGraphPinNameListValueWidget.h"
+#include "Styling/SlateTypes.h"
+#include "IPropertyUtilities.h"
+#include "DetailsViewWrapperObject.h"
+
+class IPropertyHandle;
+
+class FRigVMLocalVariableDetails : public IPropertyTypeCustomization
+{
+public:
+
+	// Makes a new instance of this detail layout class for a specific detail view requesting it
+	static TSharedRef<IPropertyTypeCustomization> MakeInstance()
+	{
+		return MakeShareable(new FRigVMLocalVariableDetails);
+	}
+
+	/** IPropertyTypeCustomization interface */
+	virtual void CustomizeHeader(TSharedRef<class IPropertyHandle> InStructPropertyHandle, class FDetailWidgetRow& HeaderRow, IPropertyTypeCustomizationUtils& StructCustomizationUtils) override;
+	virtual void CustomizeChildren(TSharedRef<class IPropertyHandle> InStructPropertyHandle, class IDetailChildrenBuilder& StructBuilder, IPropertyTypeCustomizationUtils& StructCustomizationUtils) override;
+
+
+private:
+
+	URigVMGraph* GraphBeingCustomized;
+	FRigVMGraphVariableDescription VariableDescription;
+	TArray<TWeakObjectPtr<UDetailsViewWrapperObject>> ObjectsBeingCustomized;
+
+	TSharedPtr<IPropertyHandle> NameHandle;
+	TSharedPtr<IPropertyHandle> TypeHandle;
+	TSharedPtr<IPropertyHandle> TypeObjectHandle;
+	TSharedPtr<IPropertyHandle> DefaultValueHandle;
+
+	TArray<TSharedPtr<FString>> EnumOptions;
+	
+	FEdGraphPinType OnGetPinInfo() const;
+	void HandlePinInfoChanged(const FEdGraphPinType& PinType);
+
+	ECheckBoxState HandleBoolDefaultValueIsChecked( ) const;
+	void OnBoolDefaultValueChanged(ECheckBoxState InCheckBoxState);
+};
+
+
+
diff --git a/Engine/Plugins/Animation/ControlRig/Source/ControlRigEditor/Private/Editor/ControlRigEditor.cpp b/Engine/Plugins/Animation/ControlRig/Source/ControlRigEditor/Private/Editor/ControlRigEditor.cpp
--- a/Engine/Plugins/Animation/ControlRig/Source/ControlRigEditor/Private/Editor/ControlRigEditor.cpp
+++ b/Engine/Plugins/Animation/ControlRig/Source/ControlRigEditor/Private/Editor/ControlRigEditor.cpp
@@ -66,7 +66,7 @@
 #include "EdGraphNode_Comment.h"
 #include "HAL/PlatformApplicationMisc.h"
 #include "SNodePanel.h"
-#include "Kismet/Private/SMyBlueprint.h"
+#include "SMyBlueprint.h"
 #include "Kismet/Private/SBlueprintEditorSelectedDebugObjectWidget.h"
 #include "Exporters/Exporter.h"
 #include "UnrealExporter.h"
@@ -1329,6 +1329,11 @@ public:
 #endif
 
 void FControlRigEditor::SetDetailObjects(const TArray<UObject*>& InObjects)
+{
+	SetDetailObjects(InObjects, true);
+}
+
+void FControlRigEditor::SetDetailObjects(const TArray<UObject*>& InObjects, bool bChangeUISelectionState)
 {
 	if(bSuspendDetailsPanelRefresh)
 	{
@@ -1380,7 +1385,7 @@ void FControlRigEditor::SetDetailObjects(const TArray<UObject*>& InObjects)
 	}
 #endif
 
-	ClearDetailObject();
+	ClearDetailObject(bChangeUISelectionState);
 
 	if (InObjects.Num() == 1)
 	{
@@ -1653,7 +1658,7 @@ bool FControlRigEditor::DetailViewShowsRigUnit(URigVMNode* InNode) const
 	return false;
 }
 
-void FControlRigEditor::ClearDetailObject()
+void FControlRigEditor::ClearDetailObject(bool bChangeUISelectionState)
 {
 	if(bSuspendDetailsPanelRefresh)
 	{
@@ -1675,7 +1680,10 @@ void FControlRigEditor::ClearDetailObject()
 	Inspector->ShowDetailsForObjects(TArray<UObject*>());
 	Inspector->ShowSingleStruct(TSharedPtr<FStructOnScope>());
 
-	SetUISelectionState(FBlueprintEditor::SelectionState_Graph);
+	if (bChangeUISelectionState)
+	{
+		SetUISelectionState(FBlueprintEditor::SelectionState_Graph);
+	}
 }
 
 
@@ -2045,7 +2053,9 @@ void FControlRigEditor::OnAddNewLocalVariable()
 		return;
 	}
 
-	FRigVMGraphVariableDescription NewVar = GetFocusedController()->AddLocalVariable(TEXT("NewLocalVar"), TEXT("bool"), nullptr, TEXT("0"), true, true);
+	FRigVMGraphVariableDescription LastTypeVar;
+	LastTypeVar.ChangeType(MyBlueprintWidget->GetLastPinTypeUsed());
+	FRigVMGraphVariableDescription NewVar = GetFocusedController()->AddLocalVariable(TEXT("NewLocalVar"), LastTypeVar.CPPType, LastTypeVar.CPPTypeObject, LastTypeVar.DefaultValue, true, true);
 	if(NewVar.Name.IsNone())
 	{
 		LogSimpleMessage( LOCTEXT("AddLocalVariable_Error", "Adding new local variable failed.") );
@@ -4425,8 +4435,7 @@ void FControlRigEditor::OnWrappedPropertyChangedChainEvent(UDetailsViewWrapperOb
 			ControlRigBP->MarkPackageDirty();
 		}
 	}
-	
-	if(InWrapperObject->GetWrappedStruct()->IsChildOf(FRigUnit::StaticStruct()))
+	else if(InWrapperObject->GetWrappedStruct()->IsChildOf(FRigUnit::StaticStruct()))
 	{
 		check(InWrapperObject->GetWrappedStruct() == WrapperObjects[0]->GetWrappedStruct());
 		
@@ -4446,6 +4455,55 @@ void FControlRigEditor::OnWrappedPropertyChangedChainEvent(UDetailsViewWrapperOb
 			GetFocusedController()->SetPinDefaultValue(PinPath, DefaultValue, true, true, false, true);
 		}
 	}
+	else if(InWrapperObject->GetWrappedStruct()->IsChildOf(FRigVMGraphVariableDescription::StaticStruct()))
+	{
+		check(InWrapperObject->GetWrappedStruct() == WrapperObjects[0]->GetWrappedStruct());
+		check(PropertyPath.RemoveFromStart(StoredStructString));
+		
+		FRigVMGraphVariableDescription* VariableDescription = InWrapperObject->GetContent<FRigVMGraphVariableDescription>();
+		URigVMGraph* Graph = CastChecked<URigVMGraph>(InWrapperObject->GetOuter());
+		URigVMController* Controller = ControlRigBP->GetController(Graph);
+		if (PropertyPath == TEXT("Name") && MyBlueprintWidget.IsValid())
+		{
+			if (FEdGraphSchemaAction_BlueprintVariableBase* VariableAcion = MyBlueprintWidget->SelectionAsBlueprintVariable())
+			{
+				const FName OldVariableName = VariableAcion->GetVariableName();
+				if (!OldVariableName.IsNone())
+				{
+					for (FRigVMGraphVariableDescription& Variable : Graph->GetLocalVariables())
+					{
+						if (Variable.Name == OldVariableName)
+						{
+							Controller->RenameLocalVariable(OldVariableName, VariableDescription->Name);
+						}
+					}
+				}
+			}
+			RefreshMyBlueprint();	
+		}
+		else if (PropertyPath == TEXT("CPPType") || PropertyPath == TEXT("CPPTypeObject"))
+		{			
+			for (FRigVMGraphVariableDescription& Variable : Graph->GetLocalVariables())
+			{
+				if (Variable.Name == VariableDescription->Name)
+				{
+					Controller->SetLocalVariableType(Variable.Name, VariableDescription->CPPType, VariableDescription->CPPTypeObject);
+				}
+			}
+		}
+		else if (PropertyPath == TEXT("DefaultValue"))
+		{			
+			for (FRigVMGraphVariableDescription& Variable : Graph->GetLocalVariables())
+			{
+				if (Variable.Name == VariableDescription->Name)
+				{
+					Controller->SetLocalVariableDefaultValue(Variable.Name, VariableDescription->DefaultValue);
+				}
+			}
+		}
+		
+		
+	}
 }
 
 void FControlRigEditor::OnRequestLocalizeFunctionDialog(URigVMLibraryNode* InFunction, UControlRigBlueprint* InTargetBlueprint, bool bForce)
@@ -4532,6 +4590,31 @@ void FControlRigEditor::UpdateDefaultValueForVariable(FBPVariableDescription& In
 	}
 }
 
+bool FControlRigEditor::SelectLocalVariable(const UEdGraph* Graph, const FName& VariableName)
+{
+	if (const UControlRigGraph* ControlRigGraph = Cast<UControlRigGraph>(Graph))
+	{
+		if (URigVMGraph* RigVMGraph = ControlRigGraph->GetModel())
+		{
+			for (FRigVMGraphVariableDescription& Variable : RigVMGraph->GetLocalVariables())
+			{
+				if (Variable.Name == VariableName)
+				{
+					UDetailsViewWrapperObject* WrapperObject = UDetailsViewWrapperObject::MakeInstance(
+						Variable.StaticStruct(), (uint8*)&Variable, RigVMGraph);
+					WrapperObject->GetWrappedPropertyChangedChainEvent().AddSP(this, &FControlRigEditor::OnWrappedPropertyChangedChainEvent);
+					WrapperObject->AddToRoot();
+
+					TArray<UObject*> Objects = {WrapperObject};
+					SetDetailObjects(Objects, false);
+					return true;
+				}
+			}
+		}
+	}
+	return false;
+}
+
 void FControlRigEditor::OnCreateComment()
 {
 	TSharedPtr<SGraphEditor> GraphEditor = FocusedGraphEdPtr.Pin();
diff --git a/Engine/Plugins/Animation/ControlRig/Source/ControlRigEditor/Private/Editor/ControlRigEditor.h b/Engine/Plugins/Animation/ControlRig/Source/ControlRigEditor/Private/Editor/ControlRigEditor.h
--- a/Engine/Plugins/Animation/ControlRig/Source/ControlRigEditor/Private/Editor/ControlRigEditor.h
+++ b/Engine/Plugins/Animation/ControlRig/Source/ControlRigEditor/Private/Editor/ControlRigEditor.h
@@ -157,6 +157,7 @@ public:
 	UControlRigBlueprint* GetControlRigBlueprint() const;
 
 	void SetDetailObjects(const TArray<UObject*>& InObjects);
+	void SetDetailObjects(const TArray<UObject*>& InObjects, bool bChangeUISelectionState);
 	void SetDetailViewForRigElements();
 	void SetDetailViewForGraph(URigVMGraph* InGraph);
 	void SetDetailViewForFocusedGraph();
@@ -167,7 +168,7 @@ public:
 	bool DetailViewShowsRigElement(FRigElementKey InKey) const;
 	bool DetailViewShowsRigUnit(URigVMNode* InNode) const;
 
-	void ClearDetailObject();
+	void ClearDetailObject(bool bChangeUISelectionState = true);
 
 	/** Get the persona toolkit */
 	TSharedRef<IPersonaToolkit> GetPersonaToolkit() const { return PersonaToolkit.ToSharedRef(); }
@@ -195,6 +196,8 @@ public:
 	DECLARE_EVENT_OneParam(FControlRigEditor, FPreviewControlRigUpdated, FControlRigEditor*);
 	FPreviewControlRigUpdated& OnPreviewControlRigUpdated() { return PreviewControlRigUpdated;  }
 
+	virtual bool SelectLocalVariable(const UEdGraph* Graph, const FName& VariableName) override;
+
 private:
 
 	virtual void OnCreateComment() override;
diff --git a/Engine/Plugins/Animation/ControlRig/Source/ControlRigEditor/Private/Editor/DetailsViewWrapperObject.h b/Engine/Plugins/Animation/ControlRig/Source/ControlRigEditor/Private/Editor/DetailsViewWrapperObject.h
--- a/Engine/Plugins/Animation/ControlRig/Source/ControlRigEditor/Private/Editor/DetailsViewWrapperObject.h
+++ b/Engine/Plugins/Animation/ControlRig/Source/ControlRigEditor/Private/Editor/DetailsViewWrapperObject.h
@@ -33,7 +33,7 @@ public:
 	template<typename T>
 	void SetContent(const T* InStructMemory)
 	{
-		SetCOntent((const uint8*)InStructMemory, T::StaticStruct()->GetStructureSize());
+		SetContent((const uint8*)InStructMemory, T::StaticStruct()->GetStructureSize());
 	}
 
 	FWrappedPropertyChangedChainEvent& GetWrappedPropertyChangedChainEvent() { return WrappedPropertyChangedChainEvent; }
diff --git a/Engine/Source/Developer/RigVMDeveloper/Private/RigVMModel/RigVMController.cpp b/Engine/Source/Developer/RigVMDeveloper/Private/RigVMModel/RigVMController.cpp
--- a/Engine/Source/Developer/RigVMDeveloper/Private/RigVMModel/RigVMController.cpp
+++ b/Engine/Source/Developer/RigVMDeveloper/Private/RigVMModel/RigVMController.cpp
@@ -4269,35 +4269,48 @@ bool URigVMController::RemoveNode(URigVMNode* InNode, bool bSetupUndoRedo, bool
 			{
 				if(URigVMLibraryNode* OuterFunction = OuterLibrary->FindFunctionForNode(Graph->GetTypedOuter<URigVMCollapseNode>()))
 				{
-					TArray<FRigVMExternalVariable> ExternalVariablesWithoutVariableNode;
-					{
-						URigVMGraph* EditedGraph = InNode->GetGraph();
-						TGuardValue<TArray<URigVMNode*>> TemporaryRemoveNodes(EditedGraph->Nodes, TArray<URigVMNode*>());
-						ExternalVariablesWithoutVariableNode = EditedGraph->GetExternalVariables();
-					}
-
 					const FName VariableToRemove = VariableNode->GetVariableName();
-					bool bFoundExternalVariable = false;
-					for(const FRigVMExternalVariable& ExternalVariable : ExternalVariablesWithoutVariableNode)
+					bool bIsLocalVariable = false;
+					for (FRigVMGraphVariableDescription VariableDescription : OuterFunction->GetContainedGraph()->LocalVariables)
 					{
-						if(ExternalVariable.Name == VariableToRemove)
+						if (VariableDescription.Name == VariableToRemove)
 						{
-							bFoundExternalVariable = true;
+							bIsLocalVariable = true;
 							break;
 						}
 					}
 
-					if(!bFoundExternalVariable)
+					if (!bIsLocalVariable)
 					{
-						FRigVMControllerGraphGuard Guard(this, OuterFunction->GetContainedGraph(), false);
-						if(RequestBulkEditDialogDelegate.IsBound())
+						TArray<FRigVMExternalVariable> ExternalVariablesWithoutVariableNode;
 						{
-							FRigVMController_BulkEditResult Result = RequestBulkEditDialogDelegate.Execute(OuterFunction, ERigVMControllerBulkEditType::RemoveVariable);
-							if(Result.bCanceled)
+							URigVMGraph* EditedGraph = InNode->GetGraph();
+							TGuardValue<TArray<URigVMNode*>> TemporaryRemoveNodes(EditedGraph->Nodes, TArray<URigVMNode*>());
+							ExternalVariablesWithoutVariableNode = EditedGraph->GetExternalVariables();
+						}
+
+						bool bFoundExternalVariable = false;
+						for(const FRigVMExternalVariable& ExternalVariable : ExternalVariablesWithoutVariableNode)
+						{
+							if(ExternalVariable.Name == VariableToRemove)
 							{
-								return false;
+								bFoundExternalVariable = true;
+								break;
+							}
+						}
+
+						if(!bFoundExternalVariable)
+						{
+							FRigVMControllerGraphGuard Guard(this, OuterFunction->GetContainedGraph(), false);
+							if(RequestBulkEditDialogDelegate.IsBound())
+							{
+								FRigVMController_BulkEditResult Result = RequestBulkEditDialogDelegate.Execute(OuterFunction, ERigVMControllerBulkEditType::RemoveVariable);
+								if(Result.bCanceled)
+								{
+									return false;
+								}
+								bSetupUndoRedo = Result.bSetupUndoRedo;
 							}
-							bSetupUndoRedo = Result.bSetupUndoRedo;
 						}
 					}
 				}
@@ -8003,6 +8016,10 @@ bool URigVMController::SetLocalVariableType(const FName& InVariableName, const F
 		CreateDefaultValueForStructIfRequired(ScriptStruct, DefaultValue);
 		LocalVariables[FoundIndex].DefaultValue = DefaultValue;
 	}
+	else
+	{
+		LocalVariables[FoundIndex].DefaultValue = FString();
+	}
 
 	// Change pin types on variable nodes
 	TArray<URigVMNode*> Nodes = Graph->GetNodes();
@@ -8060,6 +8077,79 @@ bool URigVMController::SetLocalVariableTypeFromObjectPath(const FName& InVariabl
 	return SetLocalVariableType(InVariableName, InCPPType, CPPTypeObject, bSetupUndoRedo, bPrintPythonCommand);
 }
 
+bool URigVMController::SetLocalVariableDefaultValue(const FName& InVariableName, const FString& InDefaultValue, bool bSetupUndoRedo, bool bPrintPythonCommand)
+{
+	if (!IsValidGraph())
+	{
+		return false;
+	}
+
+	URigVMGraph* Graph = GetGraph();
+	check(Graph);
+
+	TArray<FRigVMGraphVariableDescription>& LocalVariables = Graph->LocalVariables;
+	int32 FoundIndex = INDEX_NONE;
+	for (int32 Index = 0; Index < LocalVariables.Num(); ++Index)
+	{
+		if (LocalVariables[Index].Name == InVariableName)
+		{
+			FoundIndex = Index;
+			break;
+		}
+	}
+
+	if (FoundIndex == INDEX_NONE)
+	{
+		return false;
+	}
+
+	if (bSetupUndoRedo)
+	{
+		FRigVMInverseAction InverseAction;
+		InverseAction.Title = FString::Printf(TEXT("Change Local Variable %s default value"), *InVariableName.ToString());
+
+		ActionStack->BeginAction(InverseAction);
+		ActionStack->AddAction(FRigVMChangeLocalVariableDefaultValueAction(LocalVariables[FoundIndex], InDefaultValue));
+		ActionStack->EndAction(InverseAction);
+	}	
+	
+	LocalVariables[FoundIndex].DefaultValue = InDefaultValue;
+	
+	// Refresh variable nodes to reflect change in default value
+	TArray<URigVMNode*> Nodes = Graph->GetNodes();
+	for (URigVMNode* Node : Nodes)
+	{
+		if (URigVMVariableNode* VariableNode = Cast<URigVMVariableNode>(Node))
+		{
+			if (URigVMPin* VariablePin = VariableNode->FindPin(URigVMVariableNode::VariableName))
+			{
+				if (VariablePin->GetDefaultValue() == InVariableName.ToString())
+				{
+					VariableNode->FindPin(URigVMVariableNode::ValueName)->DefaultValue = InDefaultValue;
+					Notify(ERigVMGraphNotifType::VariableRenamed, VariableNode);
+					continue;
+				}
+			}
+		}
+	}
+
+	if (!bSuspendNotifications)
+	{
+		Graph->MarkPackageDirty();
+	}
+
+	if (bPrintPythonCommand)
+	{
+		RigVMPythonUtils::Print(GetGraph()->GetRootGraph()->GetOuter()->GetFName().ToString(), 
+			FString::Printf(TEXT("blueprint.get_controller_by_name('%s').set_local_variable_default_value('%s', '%s')"),
+				*GetGraph()->GetGraphName(),
+				*InVariableName.ToString(),
+				*InDefaultValue));
+	}
+	
+	return true;
+}
+
 TArray<TSoftObjectPtr<URigVMFunctionReferenceNode>> URigVMController::GetAffectedReferences(ERigVMControllerBulkEditType InEditType, bool bForceLoad, bool bNotify)
 {
 	TArray<TSoftObjectPtr<URigVMFunctionReferenceNode>> FunctionReferencePtrs;
diff --git a/Engine/Source/Developer/RigVMDeveloper/Private/RigVMModel/RigVMControllerActions.cpp b/Engine/Source/Developer/RigVMDeveloper/Private/RigVMModel/RigVMControllerActions.cpp
--- a/Engine/Source/Developer/RigVMDeveloper/Private/RigVMModel/RigVMControllerActions.cpp
+++ b/Engine/Source/Developer/RigVMDeveloper/Private/RigVMModel/RigVMControllerActions.cpp
@@ -1891,6 +1891,36 @@ bool FRigVMChangeLocalVariableTypeAction::Redo(URigVMController* InController)
 	return FRigVMBaseAction::Redo(InController);
 }
 
+FRigVMChangeLocalVariableDefaultValueAction::FRigVMChangeLocalVariableDefaultValueAction()
+	: LocalVariable()
+	, DefaultValue()
+{
+}
+
+FRigVMChangeLocalVariableDefaultValueAction::FRigVMChangeLocalVariableDefaultValueAction(
+	const FRigVMGraphVariableDescription& InLocalVariable, const FString& InDefaultValue)
+		: LocalVariable(InLocalVariable), DefaultValue(InDefaultValue)
+{
+}
+
+bool FRigVMChangeLocalVariableDefaultValueAction::Undo(URigVMController* InController)
+{
+	if (!FRigVMBaseAction::Undo(InController))
+	{
+		return false;
+	}
+	return InController->SetLocalVariableDefaultValue(LocalVariable.Name, LocalVariable.DefaultValue, false);
+}
+
+bool FRigVMChangeLocalVariableDefaultValueAction::Redo(URigVMController* InController)
+{
+	if (!InController->SetLocalVariableDefaultValue(LocalVariable.Name, DefaultValue, false))
+	{
+		return false;		
+	}
+	return FRigVMBaseAction::Redo(InController);
+}
+
 FRigVMAddArrayNodeAction::FRigVMAddArrayNodeAction()
 	: OpCode(ERigVMOpCode::Invalid)
 	, CPPType()
diff --git a/Engine/Source/Developer/RigVMDeveloper/Public/RigVMModel/RigVMController.h b/Engine/Source/Developer/RigVMDeveloper/Public/RigVMModel/RigVMController.h
--- a/Engine/Source/Developer/RigVMDeveloper/Public/RigVMModel/RigVMController.h
+++ b/Engine/Source/Developer/RigVMDeveloper/Public/RigVMModel/RigVMController.h
@@ -654,6 +654,9 @@ public:
 	UFUNCTION(BlueprintCallable, Category = RigVMController)
 	bool SetLocalVariableTypeFromObjectPath(const FName& InVariableName, const FString& InCPPType, const FString& InCPPTypeObjectPath, bool bSetupUndoRedo = true, bool bPrintPythonCommand = false);
 
+	UFUNCTION(BlueprintCallable, Category = RigVMController)
+	bool SetLocalVariableDefaultValue(const FName& InVariableName, const FString& InDefaultValue, bool bSetupUndoRedo = true, bool bPrintPythonCommand = false);
+
 	// Determine affected function references for a potential bulk edit on a library node
 	TArray<TSoftObjectPtr<URigVMFunctionReferenceNode>> GetAffectedReferences(ERigVMControllerBulkEditType InEditType, bool bForceLoad = false, bool bNotify = true);
 
diff --git a/Engine/Source/Developer/RigVMDeveloper/Public/RigVMModel/RigVMControllerActions.h b/Engine/Source/Developer/RigVMDeveloper/Public/RigVMModel/RigVMControllerActions.h
--- a/Engine/Source/Developer/RigVMDeveloper/Public/RigVMModel/RigVMControllerActions.h
+++ b/Engine/Source/Developer/RigVMDeveloper/Public/RigVMModel/RigVMControllerActions.h
@@ -1586,6 +1586,28 @@ public:
 	TObjectPtr<UObject> CPPTypeObject;
 };
 
+/**
+* An action to change the default value of a local variable.
+*/
+USTRUCT()
+struct FRigVMChangeLocalVariableDefaultValueAction : public FRigVMBaseAction
+{
+	GENERATED_BODY()
+
+public:
+	FRigVMChangeLocalVariableDefaultValueAction();
+	FRigVMChangeLocalVariableDefaultValueAction(const FRigVMGraphVariableDescription& InLocalVariable, const FString& InDefaultValue);
+	virtual ~FRigVMChangeLocalVariableDefaultValueAction() {};
+	virtual bool Undo(URigVMController* InController) override;
+	virtual bool Redo(URigVMController* InController) override;
+
+	UPROPERTY()
+	FRigVMGraphVariableDescription LocalVariable;
+
+	UPROPERTY()
+	FString DefaultValue;
+};
+
 /**
 * An action adding a array node to the graph.
 */
diff --git a/Engine/Source/Developer/RigVMDeveloper/Public/RigVMModel/RigVMGraph.h b/Engine/Source/Developer/RigVMDeveloper/Public/RigVMModel/RigVMGraph.h
--- a/Engine/Source/Developer/RigVMDeveloper/Public/RigVMModel/RigVMGraph.h
+++ b/Engine/Source/Developer/RigVMDeveloper/Public/RigVMModel/RigVMGraph.h
@@ -121,7 +121,7 @@ public:
 
 	TArray<FRigVMExternalVariable> GetExternalVariables() const;
 
-	const TArray<FRigVMGraphVariableDescription>& GetLocalVariables() const { return LocalVariables; }
+	TArray<FRigVMGraphVariableDescription>& GetLocalVariables() { return LocalVariables; }
 
 	// Returns the modified event, which can be used to 
 	// subscribe to changes happening within the Graph.
diff --git a/Engine/Source/Developer/RigVMDeveloper/Public/RigVMModel/RigVMVariableDescription.h b/Engine/Source/Developer/RigVMDeveloper/Public/RigVMModel/RigVMVariableDescription.h
--- a/Engine/Source/Developer/RigVMDeveloper/Public/RigVMModel/RigVMVariableDescription.h
+++ b/Engine/Source/Developer/RigVMDeveloper/Public/RigVMModel/RigVMVariableDescription.h
@@ -26,19 +26,19 @@ public:
 	}
 
 	// The name of the variable
-	UPROPERTY(BlueprintReadOnly, Category = RigVMGraphVariableDescription)
+	UPROPERTY(EditAnywhere, BlueprintReadOnly, Category = RigVMGraphVariableDescription)
 	FName Name;
 
 	// The C++ data type of the variable
-	UPROPERTY(BlueprintReadOnly, Category = RigVMGraphVariableDescription)
+	UPROPERTY(EditAnywhere, BlueprintReadOnly, Category = RigVMGraphVariableDescription)
 	FString CPPType;
 
 	// The Struct of the C++ data type of the variable (or nullptr)
-	UPROPERTY(BlueprintReadOnly, Category = RigVMGraphVariableDescription)
+	UPROPERTY(EditAnywhere, BlueprintReadOnly, Category = RigVMGraphVariableDescription)
 	TObjectPtr<UObject> CPPTypeObject = nullptr;
 
 	// The default value of the variable
-	UPROPERTY(BlueprintReadOnly, Category = RigVMGraphVariableDescription)
+	UPROPERTY(EditAnywhere, BlueprintReadOnly, Category = RigVMGraphVariableDescription)
 	FString DefaultValue;
 
 	// Returns nullptr external variable matching this description
@@ -66,7 +66,7 @@ public:
 		return ExternalVariable;
 	}
 
-	FORCEINLINE FEdGraphPinType ToPinType()
+	FORCEINLINE FEdGraphPinType ToPinType() const
 	{
 		FEdGraphPinType PinType;
 		PinType.ResetToDefaults();
@@ -160,13 +160,24 @@ public:
 		}
 		else if (InPinType.PinCategory == UEdGraphSchema_K2::PC_Struct)
 		{
-			OutCPPType = Prefix + TEXT("struct") + Suffix;
-			OutCPPTypeObject = InPinType.PinSubCategoryObject.Get();
-		}
-		else if (InPinType.PinCategory == UEdGraphSchema_K2::PC_Byte)
-		{
-			OutCPPType = Prefix + TEXT("byte") + Suffix;
-			OutCPPTypeObject = InPinType.PinSubCategoryObject.Get();
+			if (UScriptStruct* Struct = Cast<UScriptStruct>(InPinType.PinSubCategoryObject))
+			{
+				OutCPPType = Prefix + *Struct->GetStructCPPName() + Suffix;
+				OutCPPTypeObject = Struct;
+			}
+		}
+		else if (InPinType.PinCategory == UEdGraphSchema_K2::PC_Byte ||
+			InPinType.PinCategory == UEdGraphSchema_K2::PC_Enum)
+		{
+			if (UEnum* Enum = Cast<UEnum>(InPinType.PinSubCategoryObject))
+			{
+				OutCPPType = Prefix + Enum->GetFName().ToString() + Suffix;
+				OutCPPTypeObject = Enum;
+			}
+			else
+			{
+				OutCPPType = Prefix + TEXT("uint8") + Suffix;
+			}
 		}
 		else
 		{
diff --git a/Engine/Source/Editor/BlueprintGraph/Classes/EdGraphSchema_K2_Actions.h b/Engine/Source/Editor/BlueprintGraph/Classes/EdGraphSchema_K2_Actions.h
--- a/Engine/Source/Editor/BlueprintGraph/Classes/EdGraphSchema_K2_Actions.h
+++ b/Engine/Source/Editor/BlueprintGraph/Classes/EdGraphSchema_K2_Actions.h
@@ -542,6 +542,7 @@ public:
 	virtual int32 GetReorderIndexInContainer() const override;
 	virtual bool ReorderToBeforeAction(TSharedRef<FEdGraphSchemaAction> OtherAction) override;
 	virtual FEdGraphSchemaActionDefiningObject GetPersistentItemDefiningObject() const override;
+	virtual bool IsAVariable() const { return true; }
 	// End of FEdGraphSchemaAction interface
 
 	UBlueprint* GetSourceBlueprint() const;
diff --git a/Engine/Source/Editor/Kismet/Private/BlueprintEditorModule.cpp b/Engine/Source/Editor/Kismet/Private/BlueprintEditorModule.cpp
--- a/Engine/Source/Editor/Kismet/Private/BlueprintEditorModule.cpp
+++ b/Engine/Source/Editor/Kismet/Private/BlueprintEditorModule.cpp
@@ -370,6 +370,16 @@ void FBlueprintEditorModule::UnregisterVariableCustomization(FFieldClass* InFiel
 	VariableCustomizations.Remove(InFieldClass);
 }
 
+void FBlueprintEditorModule::RegisterLocalVariableCustomization(FFieldClass* InFieldClass, FOnGetLocalVariableCustomizationInstance InOnGetLocalVariableCustomization)
+{
+	LocalVariableCustomizations.Add(InFieldClass, InOnGetLocalVariableCustomization);
+}
+
+void FBlueprintEditorModule::UnregisterLocalVariableCustomization(FFieldClass* InFieldClass)
+{
+	LocalVariableCustomizations.Remove(InFieldClass);
+}
+
 void FBlueprintEditorModule::RegisterGraphCustomization(const UEdGraphSchema* InGraphSchema, FOnGetGraphCustomizationInstance InOnGetGraphCustomization)
 {
 	GraphCustomizations.Add(InGraphSchema, InOnGetGraphCustomization);
diff --git a/Engine/Source/Editor/Kismet/Private/SMyBlueprint.cpp b/Engine/Source/Editor/Kismet/Private/SMyBlueprint.cpp
--- a/Engine/Source/Editor/Kismet/Private/SMyBlueprint.cpp
+++ b/Engine/Source/Editor/Kismet/Private/SMyBlueprint.cpp
@@ -1958,7 +1958,7 @@ void SMyBlueprint::OnActionSelected( const TArray< TSharedPtr<FEdGraphSchemaActi
 	OnActionSelectedHelper(InAction, BlueprintEditorPtr, Blueprint, CurrentInspector.ToSharedRef());
 }
 
-void SMyBlueprint::OnActionSelectedHelper(TSharedPtr<FEdGraphSchemaAction> InAction, TWeakPtr< FBlueprintEditor > InBlueprintEditor, UBlueprint* Blueprint, TSharedRef<SKismetInspector> Inspector)
+void SMyBlueprint::OnActionSelectedHelper(TSharedPtr<FEdGraphSchemaAction> InAction, TWeakPtr< FBlueprintEditor > InBlueprintEditor, UBlueprint* CurrentBlueprint, TSharedRef<SKismetInspector> CurrentInspector)
 {
 	if (InAction.IsValid())
 	{
@@ -1972,7 +1972,7 @@ void SMyBlueprint::OnActionSelectedHelper(TSharedPtr<FEdGraphSchemaAction> InAct
 				const UEdGraphSchema* Schema = GraphAction->EdGraph->GetSchema();
 				check(Schema != nullptr);
 				Schema->GetGraphDisplayInformation(*GraphAction->EdGraph, DisplayInfo);
-				Inspector->ShowDetailsForSingleObject(GraphAction->EdGraph, SKismetInspector::FShowDetailsOptions(DisplayInfo.PlainName));
+				CurrentInspector->ShowDetailsForSingleObject(GraphAction->EdGraph, SKismetInspector::FShowDetailsOptions(DisplayInfo.PlainName));
 			}
 		}
 		else if (InAction->GetTypeId() == FEdGraphSchemaAction_K2Delegate::StaticGetTypeId())
@@ -1980,7 +1980,7 @@ void SMyBlueprint::OnActionSelectedHelper(TSharedPtr<FEdGraphSchemaAction> InAct
 			FEdGraphSchemaAction_K2Delegate* DelegateAction = (FEdGraphSchemaAction_K2Delegate*)InAction.Get();
 			if (FMulticastDelegateProperty* Property = DelegateAction->GetDelegateProperty())
 			{
-				Inspector->ShowDetailsForSingleObject(Property->GetUPropertyWrapper(), SKismetInspector::FShowDetailsOptions(FText::FromString(Property->GetName())));
+				CurrentInspector->ShowDetailsForSingleObject(Property->GetUPropertyWrapper(), SKismetInspector::FShowDetailsOptions(FText::FromString(Property->GetName())));
 			}
 		}
 		else if (InAction->GetTypeId() == FEdGraphSchemaAction_K2Var::StaticGetTypeId())
@@ -1992,7 +1992,7 @@ void SMyBlueprint::OnActionSelectedHelper(TSharedPtr<FEdGraphSchemaAction> InAct
 
 			FProperty* Prop = VarAction->GetProperty();
 			UPropertyWrapper* PropWrap = (Prop ? Prop->GetUPropertyWrapper() : nullptr);
-			Inspector->ShowDetailsForSingleObject(PropWrap, Options);
+			CurrentInspector->ShowDetailsForSingleObject(PropWrap, Options);
 		}
 		else if (InAction->GetTypeId() == FEdGraphSchemaAction_K2LocalVar::StaticGetTypeId())
 		{
@@ -2002,7 +2002,18 @@ void SMyBlueprint::OnActionSelectedHelper(TSharedPtr<FEdGraphSchemaAction> InAct
 
 			FProperty* Prop = VarAction->GetProperty();
 			UPropertyWrapper* PropWrap = (Prop ? Prop->GetUPropertyWrapper() : nullptr);
-			Inspector->ShowDetailsForSingleObject(PropWrap, Options);
+			CurrentInspector->ShowDetailsForSingleObject(PropWrap, Options);
+		}
+		else if (InAction->IsAVariable())
+		{
+			FEdGraphSchemaAction_BlueprintVariableBase* VarAction = (FEdGraphSchemaAction_BlueprintVariableBase*)InAction.Get();
+			if (UEdGraph* Graph = GetFocusedGraph())
+			{
+				if (BlueprintEditorPtr.IsValid())
+				{
+					BlueprintEditorPtr.Pin()->SelectLocalVariable(Graph, VarAction->GetVariableName());	
+				}
+			}
 		}
 		else if (InAction->GetTypeId() == FEdGraphSchemaAction_K2Enum::StaticGetTypeId())
 		{
@@ -2011,7 +2022,7 @@ void SMyBlueprint::OnActionSelectedHelper(TSharedPtr<FEdGraphSchemaAction> InAct
 			SKismetInspector::FShowDetailsOptions Options(FText::FromName(EnumAction->GetPathName()));
 			Options.bForceRefresh = true;
 
-			Inspector->ShowDetailsForSingleObject(EnumAction->Enum, Options);
+			CurrentInspector->ShowDetailsForSingleObject(EnumAction->Enum, Options);
 		}
 		else if (InAction->GetTypeId() == FEdGraphSchemaAction_K2Struct::StaticGetTypeId())
 		{
@@ -2020,7 +2031,7 @@ void SMyBlueprint::OnActionSelectedHelper(TSharedPtr<FEdGraphSchemaAction> InAct
 			SKismetInspector::FShowDetailsOptions Options(FText::FromName(StructAction->GetPathName()));
 			Options.bForceRefresh = true;
 
-			Inspector->ShowDetailsForSingleObject(StructAction->Struct, Options);
+			CurrentInspector->ShowDetailsForSingleObject(StructAction->Struct, Options);
 		}
 		else if (InAction->GetTypeId() == FEdGraphSchemaAction_K2TargetNode::StaticGetTypeId() ||
 			InAction->GetTypeId() == FEdGraphSchemaAction_K2Event::StaticGetTypeId() ||
@@ -2032,16 +2043,16 @@ void SMyBlueprint::OnActionSelectedHelper(TSharedPtr<FEdGraphSchemaAction> InAct
 			check(NodeTemplate != nullptr)
 
 			SKismetInspector::FShowDetailsOptions Options(NodeTemplate->GetNodeTitle(ENodeTitleType::EditableTitle));
-			Inspector->ShowDetailsForSingleObject(NodeTemplate, Options);
+			CurrentInspector->ShowDetailsForSingleObject(NodeTemplate, Options);
 		}
 		else
 		{
-			Inspector->ShowDetailsForObjects(TArray<UObject*>());
+			CurrentInspector->ShowDetailsForObjects(TArray<UObject*>());
 		}
 	}
 	else
 	{
-		Inspector->ShowDetailsForObjects(TArray<UObject*>());
+		CurrentInspector->ShowDetailsForObjects(TArray<UObject*>());
 	}
 }
 
diff --git a/Engine/Source/Editor/Kismet/Private/SMyBlueprint.h b/Engine/Source/Editor/Kismet/Private/SMyBlueprint.h
deleted file mode 100644
diff --git a/Engine/Source/Editor/Kismet/Public/BlueprintEditor.h b/Engine/Source/Editor/Kismet/Public/BlueprintEditor.h
--- a/Engine/Source/Editor/Kismet/Public/BlueprintEditor.h
+++ b/Engine/Source/Editor/Kismet/Public/BlueprintEditor.h
@@ -717,6 +717,9 @@ public:
 	/** Returns an instanced namespace helper utility object that corresponds to the given Blueprint. */
 	TSharedRef<FBlueprintNamespaceHelper> GetOrCreateNamespaceHelperForBlueprint(const UBlueprint* InBlueprint);
 
+	/** Selects a local variable to load in the details panel. */
+	virtual bool SelectLocalVariable(const UEdGraph* Graph, const FName& VariableName) { return false; }
+
 protected:
 	UE_DEPRECATED(4.26, "Please do any validation inside the UBlueprint class during compilation, extra errors during compiling only supplied by the designer can lead to design time only errors being reported and being missed during cooks/content validation.")
 	virtual void AppendExtraCompilerResults(TSharedPtr<class IMessageLogListing> ResultsListing) {}
diff --git a/Engine/Source/Editor/Kismet/Public/BlueprintEditorModule.h b/Engine/Source/Editor/Kismet/Public/BlueprintEditorModule.h
--- a/Engine/Source/Editor/Kismet/Public/BlueprintEditorModule.h
+++ b/Engine/Source/Editor/Kismet/Public/BlueprintEditorModule.h
@@ -26,6 +26,9 @@ class FSubobjectEditorTreeNode;
 /** Delegate used to customize variable display */
 DECLARE_DELEGATE_RetVal_OneParam(TSharedPtr<IDetailCustomization>, FOnGetVariableCustomizationInstance, TSharedPtr<IBlueprintEditor> /*BlueprintEditor*/);
 
+/** Delegate used to customize local variable display */
+DECLARE_DELEGATE_RetVal_OneParam(TSharedPtr<IDetailCustomization>, FOnGetLocalVariableCustomizationInstance, TSharedPtr<IBlueprintEditor> /*BlueprintEditor*/);
+
 /** Delegate used to customize graph display */
 DECLARE_DELEGATE_RetVal_OneParam(TSharedPtr<IDetailCustomization>, FOnGetGraphCustomizationInstance, TSharedPtr<IBlueprintEditor> /*BlueprintEditor*/);
 
@@ -214,6 +217,19 @@ public:
 	 */
 	virtual void UnregisterVariableCustomization(FFieldClass* InFieldClass);
 
+	/** 
+	* Register a customization for for Blueprint local variables
+	* @param	InFieldClass				The type of the variable to create the customization for
+	* @param	InOnGetLocalVariableCustomization	The delegate used to create customization instances
+	*/
+	virtual void RegisterLocalVariableCustomization(FFieldClass* InFieldClass, FOnGetLocalVariableCustomizationInstance InOnGetLocalVariableCustomization);
+
+	/** 
+	* Unregister a previously registered customization for BP local variables
+	* @param	InFieldClass				The type to create the customization for
+	*/
+	virtual void UnregisterLocalVariableCustomization(FFieldClass* InFieldClass);
+
 	/** 
 	 * Register a customization for for Blueprint graphs
 	 * @param	InGraphSchema				The schema of the graph to create the customization for
@@ -279,6 +295,9 @@ private:
 	/** Customizations for Blueprint variables */
 	TMap<FFieldClass*, FOnGetVariableCustomizationInstance> VariableCustomizations;
 
+	/** Customizations for Blueprint local variables */
+	TMap<FFieldClass*, FOnGetLocalVariableCustomizationInstance> LocalVariableCustomizations;
+
 	/** Customizations for Blueprint graphs */
 	TMap<const UEdGraphSchema*, FOnGetGraphCustomizationInstance> GraphCustomizations;
 
diff --git a/Engine/Source/Editor/Kismet/Public/SMyBlueprint.h b/Engine/Source/Editor/Kismet/Public/SMyBlueprint.h
new file mode 100644
--- /dev/null
+++ b/Engine/Source/Editor/Kismet/Public/SMyBlueprint.h
@@ -0,0 +1,352 @@
+// Copyright Epic Games, Inc. All Rights Reserved.
+
+#pragma once
+
+#include "CoreMinimal.h"
+#include "EdGraph/EdGraphPin.h"
+#include "Layout/Visibility.h"
+#include "Widgets/DeclarativeSyntaxSupport.h"
+#include "Input/Reply.h"
+#include "Widgets/SWidget.h"
+#include "Widgets/SCompoundWidget.h"
+#include "SGraphActionMenu.h"
+#include "EditorStyleSet.h"
+#include "EdGraph/EdGraph.h"
+#include "EdGraphSchema_K2_Actions.h"
+#include "WorkflowOrientedApp/WorkflowTabManager.h"
+#include "BlueprintEditor.h"
+#include "Framework/Commands/Commands.h"
+
+class FMenuBuilder;
+class SComboButton;
+class SKismetInspector;
+class UUserDefinedEnum;
+class UUserDefinedStruct;
+struct FEdGraphSchemaAction_K2Struct;
+struct FGraphActionNode;
+struct FGraphActionSort;
+struct FReplaceNodeReferencesHelper;
+
+class FMyBlueprintCommands : public TCommands<FMyBlueprintCommands>
+{
+public:
+	/** Constructor */
+	FMyBlueprintCommands() 
+		: TCommands<FMyBlueprintCommands>(TEXT("MyBlueprint"), NSLOCTEXT("Contexts", "My Blueprint", "My Blueprint"), NAME_None, FEditorStyle::GetStyleSetName())
+	{
+	}
+
+	// Basic operations
+	TSharedPtr<FUICommandInfo> OpenGraph;
+	TSharedPtr<FUICommandInfo> OpenGraphInNewTab;
+	TSharedPtr<FUICommandInfo> OpenExternalGraph;
+	TSharedPtr<FUICommandInfo> FocusNode;
+	TSharedPtr<FUICommandInfo> FocusNodeInNewTab;
+	TSharedPtr<FUICommandInfo> ImplementFunction;
+	TSharedPtr<FUICommandInfo> DeleteEntry;
+	TSharedPtr<FUICommandInfo> PasteVariable;
+	TSharedPtr<FUICommandInfo> PasteLocalVariable;
+	TSharedPtr<FUICommandInfo> PasteFunction;
+	TSharedPtr<FUICommandInfo> PasteMacro;
+	TSharedPtr<FUICommandInfo> GotoNativeVarDefinition;
+	TSharedPtr<FUICommandInfo> MoveToParent;
+	// Add New Item
+	/** Initialize commands */
+	virtual void RegisterCommands() override;
+};
+
+class KISMET_API SMyBlueprint : public SCompoundWidget
+{
+public:
+	SLATE_BEGIN_ARGS( SMyBlueprint ) {}
+	SLATE_END_ARGS()
+
+	void Construct(const FArguments& InArgs, TWeakPtr<FBlueprintEditor> InBlueprintEditor, const UBlueprint* InBlueprint = nullptr);
+	~SMyBlueprint();
+
+	void SetInspector( TSharedPtr<SKismetInspector> InInspector ) { Inspector = InInspector ; }
+
+	/* SWidget interface */
+	virtual void Tick(const FGeometry& AllottedGeometry, const double InCurrentTime, const float InDeltaTime);
+	virtual FReply OnKeyDown(const FGeometry& MyGeometry, const FKeyEvent& InKeyEvent);
+
+	/* Reset the last pin type settings to default. */
+	void ResetLastPinType();
+
+	/** Refreshes the graph action menu */
+	void Refresh();
+	void SetFocusedGraph(UEdGraph* InEdGraph) { EdGraph = InEdGraph; }
+	
+	/** Accessor for getting the current selection as a K2 graph */
+	FEdGraphSchemaAction_K2Graph* SelectionAsGraph() const;
+
+	/** Accessor for getting the current selection as a K2 enum */
+	FEdGraphSchemaAction_K2Enum* SelectionAsEnum() const;
+
+	/** Accessor for getting the current selection as a K2 enum */
+	FEdGraphSchemaAction_K2Struct* SelectionAsStruct() const;
+
+	/** Accessor for getting the current selection as a K2 var */
+	FEdGraphSchemaAction_K2Var* SelectionAsVar() const;
+	
+	/** Accessor for getting the current selection as a K2 delegate */
+	FEdGraphSchemaAction_K2Delegate* SelectionAsDelegate() const;
+
+	/** Accessor for getting the current selection as a K2 event */
+	FEdGraphSchemaAction_K2Event* SelectionAsEvent() const;
+	
+	/** Accessor for getting the current selection as a K2 Input Action */
+	FEdGraphSchemaAction_K2InputAction* SelectionAsInputAction() const;
+
+	/** Accessor for getting the current selection as a K2 local var */
+	FEdGraphSchemaAction_K2LocalVar* SelectionAsLocalVar() const;
+
+	/** Accessor for getting the current selection as a K2 blueprint base variable */
+	FEdGraphSchemaAction_BlueprintVariableBase* SelectionAsBlueprintVariable() const;
+
+	/** Accessor for determining if the current selection is a category*/
+	bool SelectionIsCategory() const;
+	
+	void EnsureLastPinTypeValid();
+
+	/** Gets the last pin type selected by this widget, or by the function editor */
+	FEdGraphPinType& GetLastPinTypeUsed() {EnsureLastPinTypeValid(); return LastPinType;}
+	FEdGraphPinType& GetLastFunctionPinTypeUsed() {EnsureLastPinTypeValid(); return LastFunctionPinType;}
+
+	/** Accessor the blueprint object from the main editor */
+	UBlueprint* GetBlueprintObj() const {return Blueprint;}
+
+	/** Gets whether we are showing user variables only or not */
+	bool ShowUserVarsOnly() const { return !IsShowingInheritedVariables(); }
+
+	/** Gets our parent blueprint editor */
+	TWeakPtr<FBlueprintEditor> GetBlueprintEditor() {return BlueprintEditorPtr;}
+
+	/**
+	 * Fills the supplied array with the currently selected objects
+	 * @param OutSelectedItems The array to fill.
+	 */
+	void GetSelectedItemsForContextMenu(TArray<FComponentEventConstructionData>& OutSelectedItems) const;
+
+	/** Called to reset the search filter */
+	void OnResetItemFilter();
+
+	/** Selects an item by name in either the main graph action menu or the local one */
+	void SelectItemByName(const FName& ItemName, ESelectInfo::Type SelectInfo = ESelectInfo::Direct, int32 SectionId = INDEX_NONE, bool bIsCategory = false);
+
+	/** Clears the selection in the graph action menus */
+	void ClearGraphActionMenuSelection();
+
+	/** Initiates a rename on the selected action node, if possible */
+	void OnRequestRenameOnActionNode();
+
+	/** Expands any category with the associated name */
+	void ExpandCategory(const FText& CategoryName);
+
+	/** Move the category before the target category */
+	bool MoveCategoryBeforeCategory( const FText& CategoryToMove, const FText& TargetCategory );
+	
+	/** Callbacks for Paste Commands */
+	void OnPasteGeneric();
+	bool CanPasteGeneric();
+private:
+	/** Creates widgets for the graph schema actions */
+	TSharedRef<SWidget> OnCreateWidgetForAction(struct FCreateWidgetForActionData* const InCreateData);
+
+	/** Callback used to populate all actions list in SGraphActionMenu */
+	void CollectAllActions(FGraphActionListBuilderBase& OutAllActions);
+	void CollectStaticSections(TArray<int32>& StaticSectionIDs);
+	void AddEventForFunctionGraph(UEdGraph* EdGraph, int32 const SectionId, FGraphActionSort& SortList, const FText& ParentCategory = FText::GetEmpty(), bool bAddChildGraphs = true) const;
+	void GetChildGraphs(UEdGraph* EdGraph, int32 const SectionId, FGraphActionSort& SortList, const FText& ParentCategory = FText::GetEmpty()) const;
+	void GetChildEvents(UEdGraph const* EdGraph, int32 const SectionId, FGraphActionSort& SortList, const FText& ParentCategory = FText::GetEmpty(), bool bInAddChildGraphs = true) const;
+	void GetLocalVariables(FGraphActionSort& SortList) const;
+	
+	/** Handles the visibility of the local action list */
+	EVisibility GetLocalActionsListVisibility() const;
+
+	/** Callbacks for the graph action menu */
+	FReply OnActionDragged(const TArray< TSharedPtr<FEdGraphSchemaAction> >& InActions, const FPointerEvent& MouseEvent);
+	FReply OnCategoryDragged(const FText& InCategory, const FPointerEvent& MouseEvent);
+	void OnActionSelected(const TArray< TSharedPtr<FEdGraphSchemaAction> >& InActions);
+	void OnActionSelectedHelper(TSharedPtr<FEdGraphSchemaAction> InAction, TWeakPtr< FBlueprintEditor > InBlueprintEditor, UBlueprint* Blueprint, TSharedRef<SKismetInspector> Inspector);
+	void OnGlobalActionSelected(const TArray< TSharedPtr<FEdGraphSchemaAction> >& InActions, ESelectInfo::Type InSelectionType);
+	void OnActionDoubleClicked(const TArray< TSharedPtr<FEdGraphSchemaAction> >& InActions);
+	void ExecuteAction(TSharedPtr<FEdGraphSchemaAction> InAction);
+	TSharedPtr<SWidget> OnContextMenuOpening();
+
+	TSharedRef<SWidget> CreateAddNewMenuWidget();
+	void BuildAddNewMenu(FMenuBuilder& MenuBuilder);
+	TSharedRef<SWidget> CreateAddToSectionButton(int32 InSectionID, TWeakPtr<SWidget> WeakRowWidget, FText AddNewText, FName MetaDataTag);
+
+	void OnCategoryNameCommitted(const FText& InNewText, ETextCommit::Type InTextCommit, TWeakPtr< struct FGraphActionNode > InAction );
+	bool CanRequestRenameOnActionNode(TWeakPtr<struct FGraphActionNode> InSelectedNode) const;
+	FText OnGetSectionTitle( int32 InSectionID );
+	TSharedRef<SWidget> OnGetSectionWidget( TSharedRef<SWidget> RowWidget, int32 InSectionID );
+	EVisibility OnGetSectionTextVisibility(TWeakPtr<SWidget> RowWidget, int32 InSectionID) const;
+	TSharedRef<SWidget> OnGetFunctionListMenu();
+	void BuildOverridableFunctionsMenu(FMenuBuilder& MenuBuilder);
+	FReply OnAddButtonClickedOnSection(int32 InSectionID);
+	bool CanAddNewElementToSection(int32 InSectionID) const;
+
+	bool HandleActionMatchesName(FEdGraphSchemaAction* InAction, const FName& InName) const;
+
+	/** Support functions for checkbox to manage displaying user variables only */
+	bool IsShowingInheritedVariables() const;
+	void OnToggleShowInheritedVariables();
+
+	/** Support functions for view options for Show Empty Sections */
+	void OnToggleShowEmptySections();
+	bool IsShowingEmptySections() const;
+	
+	/** Support functions for view options for Show Replicated Variables only */
+	void OnToggleShowReplicatedVariablesOnly();
+	bool IsShowingReplicatedVariablesOnly() const;
+
+	/** Support functions for view options for bAlwaysShowInterfacesInOverrides blueprint editor setting */
+	void OnToggleAlwaysShowInterfacesInOverrides();
+	bool GetAlwaysShowInterfacesInOverrides() const;
+
+	/** Support functions for view options for bShowParentClassInOverrides blueprint editor setting */
+	void OnToggleShowParentClassInOverrides();
+	bool GetShowParentClassInOverrides() const;
+
+	/** Support functions for view options for bShowAccessSpecifier blueprint editor setting */
+	void OnToggleShowAccessSpecifier();
+	bool GetShowAccessSpecifier() const;
+
+	/** Helper function to open the selected graph */
+	void OpenGraph(FDocumentTracker::EOpenDocumentCause InCause, bool bOpenExternalGraphInNewEditor = false);
+
+	/**
+	* Check if the override of a given function is most likely desired as a blueprint function 
+	* or as an event. 
+	* 
+	* @param OverrideFunc	Desired function to override
+	* 
+	* @return	True if the function is desired as a function, false if desired as an event
+	*/
+	bool IsImplementationDesiredAsFunction(const UFunction* OverrideFunc) const;
+
+	/** Callbacks for commands */
+	void OnOpenGraph();
+	void OnOpenGraphInNewTab();
+	void OnOpenExternalGraph();
+	bool CanOpenGraph() const;
+	bool CanOpenExternalGraph() const;
+	bool CanFocusOnNode() const;
+	void OnFocusNode();
+	void OnFocusNodeInNewTab();
+	void OnImplementFunction();
+	void ImplementFunction(TSharedPtr<FEdGraphSchemaAction_K2Graph> GraphAction);
+	void ImplementFunction(FEdGraphSchemaAction_K2Graph* GraphAction);
+	bool CanImplementFunction() const;
+	void OnFindReference();
+	bool CanFindReference() const;
+	void OnFindAndReplaceReference();
+	bool CanFindAndReplaceReference() const;
+	void OnDeleteEntry();
+	bool CanDeleteEntry() const;
+	FReply OnAddNewLocalVariable();
+	bool CanRequestRenameOnActionNode() const;
+	bool IsDuplicateActionVisible() const;
+	bool CanDuplicateAction() const;
+	void OnDuplicateAction();
+	void GotoNativeCodeVarDefinition();
+	bool IsNativeVariable() const;
+	void OnMoveToParent();
+	bool CanMoveToParent() const;
+	void OnCopy();
+	bool CanCopy() const;
+	void OnCut();
+	bool CanCut() const;
+	void OnPasteVariable();
+	void OnPasteLocalVariable();
+	bool CanPasteVariable() const;
+	bool CanPasteLocalVariable() const;
+	void OnPasteFunction();
+	bool CanPasteFunction() const;
+	void OnPasteMacro();
+	bool CanPasteMacro() const;
+
+	/** Gets the currently selected Category or returns default category name */
+	FText GetPasteCategory() const;
+
+	/** Callback when the filter is changed, forces the action tree(s) to filter */
+	void OnFilterTextChanged( const FText& InFilterText );
+
+	/** Callback for the action trees to get the filter text */
+	FText GetFilterText() const;
+
+	/** Checks if the selected action has context menu */
+	bool SelectionHasContextMenu() const;
+
+	/** Update Node Create Analytic */
+	void UpdateNodeCreation();
+
+	/** Returns the displayed category, if any, of a graph */
+	FText GetGraphCategory(UEdGraph* InGraph) const;
+
+	/** Helper function to delete a graph in the MyBlueprint window */
+	void OnDeleteGraph(UEdGraph* InGraph, EEdGraphSchemaAction_K2Graph::Type);
+
+	/** Helper function to delete a delegate in the MyBlueprint window */
+	void OnDeleteDelegate(FEdGraphSchemaAction_K2Delegate* InDelegateAction);
+
+	UEdGraph* GetFocusedGraph() const;
+
+	/** Delegate to hook us into non-structural Blueprint object post-change events */
+	void OnObjectPropertyChanged(UObject* InObject, FPropertyChangedEvent& InPropertyChangedEvent);
+
+	/** Helper function indicating whehter we're in editing mode, and can modify the target blueprint */
+	bool IsEditingMode() const;
+
+	/** Determine whether an FEdGraphSchemaAction is associated with an event */
+	static bool IsAnInterfaceEvent(FEdGraphSchemaAction_K2Graph* InAction);
+
+private:
+	/** List of UI Commands for this scope */
+	TSharedPtr<FUICommandList> CommandList;
+
+	/** Pointer back to the blueprint editor that owns us */
+	TWeakPtr<FBlueprintEditor> BlueprintEditorPtr;
+	
+	/** Graph Action Menu for displaying all our variables and functions */
+	TSharedPtr<class SGraphActionMenu> GraphActionMenu;
+
+	/** The +Function button in the function section */
+	TSharedPtr<SComboButton> FunctionSectionButton;
+
+	/** When we rebuild the view of members, we cache (but don't display) any overridable functions for user in popup menus. */
+	TArray< TSharedPtr<FEdGraphSchemaAction_K2Graph> > OverridableFunctionActions;
+
+	/** When we refresh the list of functions we cache off the implemented ones to ask questions for overridable functions. */
+	TSet<FName> ImplementedFunctionCache;
+
+	/** The last pin type used (including the function editor last pin type) */
+	FEdGraphPinType LastPinType;
+	FEdGraphPinType LastFunctionPinType;
+
+	/** Enums created from 'blueprint' level */
+	TArray<TWeakObjectPtr<UUserDefinedEnum>> EnumsAddedToBlueprint;
+
+	/** The filter box that handles filtering for both graph action menus. */
+	TSharedPtr< SSearchBox > FilterBox;
+
+	/** Enums created from 'blueprint' level */
+	TArray<TWeakObjectPtr<UUserDefinedStruct>> StructsAddedToBlueprint;
+
+	/** The blueprint being displayed: */
+	UBlueprint* Blueprint;
+
+	/** The Ed Graph being displayed: */
+	UEdGraph* EdGraph;
+
+	/** The Kismet Inspector used to display properties: */
+	TWeakPtr<SKismetInspector> Inspector;
+
+	/** Flag to indicate whether or not we need to refresh the panel */
+	bool bNeedsRefresh;
+
+	/** If set we'll show only replicated variables (local to a particular blueprint view). */
+	bool bShowReplicatedVariablesOnly;
+};
diff --git a/Engine/Source/Runtime/Engine/Classes/EdGraph/EdGraphSchema.h b/Engine/Source/Runtime/Engine/Classes/EdGraph/EdGraphSchema.h
--- a/Engine/Source/Runtime/Engine/Classes/EdGraph/EdGraphSchema.h
+++ b/Engine/Source/Runtime/Engine/Classes/EdGraph/EdGraphSchema.h
@@ -295,6 +295,9 @@ public:
 		return InType == GetTypeId();
 	}
 
+	// Returns true if the action refers to a member or local variable
+	virtual bool IsAVariable() const { return false; }
+
 private:
 	void UpdateSearchText();
 };
