Drag-and-drop to reorder enumerators in FUserDefinedEnumEditor, replacing the existing move up/down buttons

diff --git a/Engine/Source/Editor/Kismet/Private/UserDefinedEnumEditor.cpp b/Engine/Source/Editor/Kismet/Private/UserDefinedEnumEditor.cpp
--- a/Engine/Source/Editor/Kismet/Private/UserDefinedEnumEditor.cpp
+++ b/Engine/Source/Editor/Kismet/Private/UserDefinedEnumEditor.cpp
@@ -7,6 +7,7 @@
 #include "PropertyEditorModule.h"
 #include "PropertyHandle.h"
 #include "IDetailChildrenBuilder.h"
+#include "IDetailDragDropHandler.h"
 #include "Modules/ModuleManager.h"
 #include "Widgets/Images/SImage.h"
 #include "Widgets/Layout/SBox.h"
@@ -15,6 +16,7 @@
 #include "Widgets/Input/SCheckBox.h"
 #include "DetailLayoutBuilder.h"
 #include "DetailCategoryBuilder.h"
+#include "DragAndDrop/DecoratedDragDropOp.h"
 
 #include "PropertyCustomizationHelpers.h"
 #include "Widgets/SToolTip.h"
@@ -222,8 +224,142 @@ private:
 };
 
 
+/** Drag-and-drop operation that stores data about the source enumerator being dragged */
+class FUserDefinedEnumIndexDragDropOp : public FDecoratedDragDropOp
+{
+public:
+	DRAG_DROP_OPERATOR_TYPE(FUserDefinedEnumIndexDragDropOp, FDecoratedDragDropOp);
 
+	FUserDefinedEnumIndexDragDropOp(UUserDefinedEnum* InTargetEnum, int32 InEnumeratorIndex)
+		: TargetEnum(InTargetEnum)
+		, EnumeratorIndex(InEnumeratorIndex)
+	{
+		check(InTargetEnum);
+		check(InEnumeratorIndex >= 0 && InEnumeratorIndex < InTargetEnum->NumEnums());
 
+		EnumDisplayText = InTargetEnum->GetDisplayNameTextByIndex(InEnumeratorIndex);
+	}
+
+	void Init()
+	{
+		SetValidTarget(false);
+		SetupDefaults();
+		Construct();
+	}
+
+	void SetValidTarget(bool IsValidTarget)
+	{
+		FFormatNamedArguments Args;
+		Args.Add(TEXT("EnumeratorName"), EnumDisplayText);
+
+		if (IsValidTarget)
+		{
+			CurrentHoverText = FText::Format(LOCTEXT("MoveEnumeratorHere", "Move '{EnumeratorName}' Here"), Args);
+			CurrentIconBrush = FEditorStyle::GetBrush("Graph.ConnectorFeedback.OK");
+		}
+		else
+		{
+			CurrentHoverText = FText::Format(LOCTEXT("CannotMoveEnumeratorHere", "Cannot Move '{EnumeratorName}' Here"), Args);
+			CurrentIconBrush = FEditorStyle::GetBrush("Graph.ConnectorFeedback.Error");
+		}
+	}
+
+	UUserDefinedEnum* GetTargetEnum() const
+	{
+		return TargetEnum;
+	}
+
+	int32 GetEnumeratorIndex() const
+	{
+		return EnumeratorIndex;
+	}
+
+private:
+	UUserDefinedEnum* TargetEnum;
+	int32 EnumeratorIndex;
+	FText EnumDisplayText;
+};
+
+
+/** Handler for customizing the drag-and-drop behavior for enum index rows, allowing enumerators to be reordered */
+class FUserDefinedEnumIndexDragDropHandler : public IDetailDragDropHandler
+{
+public:
+	FUserDefinedEnumIndexDragDropHandler(UUserDefinedEnum* InTargetEnum, int32 InEnumeratorIndex)
+		: TargetEnum(InTargetEnum)
+		, EnumeratorIndex(InEnumeratorIndex)
+	{
+		check(InTargetEnum);
+		check(InEnumeratorIndex >= 0 && InEnumeratorIndex < InTargetEnum->NumEnums());
+	}
+
+	virtual TSharedPtr<FDragDropOperation> CreateDragDropOperation() const override
+	{
+		TSharedPtr<FUserDefinedEnumIndexDragDropOp> DragOp = MakeShared<FUserDefinedEnumIndexDragDropOp>(TargetEnum, EnumeratorIndex);
+		DragOp->Init();
+		return DragOp;
+	}
+
+	/** Compute new target index for use with FEnumEditorUtils::MoveEnumeratorInUserDefinedEnum based on drop zone (above vs below) */
+	static int32 ComputeNewIndex(int32 OriginalIndex, int32 DropOntoIndex, EItemDropZone DropZone)
+	{
+		check(DropZone != EItemDropZone::OntoItem);
+
+		int32 NewIndex = DropOntoIndex;
+		if (DropZone == EItemDropZone::BelowItem)
+		{
+			// If the drop zone is below, then we actually move it to the next item's index
+			NewIndex++;
+		}
+		if (OriginalIndex < NewIndex)
+		{
+			// If the item is moved down the list, then all the other elements below it are shifted up one
+			NewIndex--;
+		}
+
+		return ensure(NewIndex >= 0) ? NewIndex : 0;
+	}
+
+	virtual bool AcceptDrop(const FDragDropEvent& DragDropEvent, EItemDropZone DropZone) const override
+	{
+		const TSharedPtr<FUserDefinedEnumIndexDragDropOp> DragOp = DragDropEvent.GetOperationAs<FUserDefinedEnumIndexDragDropOp>();
+		if (!DragOp.IsValid() || DragOp->GetTargetEnum() != TargetEnum || DropZone == EItemDropZone::OntoItem)
+		{
+			return false;
+		}
+
+		const int32 NewIndex = ComputeNewIndex(DragOp->GetEnumeratorIndex(), EnumeratorIndex, DropZone);
+		FEnumEditorUtils::MoveEnumeratorInUserDefinedEnum(TargetEnum, DragOp->GetEnumeratorIndex(), NewIndex);
+		return true;
+	}
+
+	virtual TOptional<EItemDropZone> CanAcceptDrop(const FDragDropEvent& DragDropEvent, EItemDropZone DropZone) const override
+	{
+		const TSharedPtr<FUserDefinedEnumIndexDragDropOp> DragOp = DragDropEvent.GetOperationAs<FUserDefinedEnumIndexDragDropOp>();
+		if (!DragOp.IsValid() || DragOp->GetTargetEnum() != TargetEnum)
+		{
+			return TOptional<EItemDropZone>();
+		}
+
+		// We're reordering, so there's no logical interpretation for dropping directly onto another enum.
+		// Just change it to a drop-above in this case.
+		const EItemDropZone OverrideZone = (DropZone == EItemDropZone::BelowItem) ? EItemDropZone::BelowItem : EItemDropZone::AboveItem;
+		const int32 NewIndex = ComputeNewIndex(DragOp->GetEnumeratorIndex(), EnumeratorIndex, OverrideZone);
+
+		// Make sure that the new index is valid *and* that it represents an actual move from the current position.
+		if (NewIndex < 0 || NewIndex >= TargetEnum->NumEnums() || NewIndex == DragOp->GetEnumeratorIndex())
+		{
+			return TOptional<EItemDropZone>();
+		}
+
+		DragOp->SetValidTarget(true);
+		return OverrideZone;
+	}
+
+private:
+	UUserDefinedEnum* TargetEnum;
+	int32 EnumeratorIndex;
+};
 
 
 void FUserDefinedEnumEditor::RegisterTabSpawners(const TSharedRef<class FTabManager>& InTabManager)
@@ -533,35 +669,6 @@ void FUserDefinedEnumIndexLayout::GenerateHeaderRowContent( FDetailWidgetRow& No
 				SNew(STextPropertyEditableTextBox, TooltipEditor.ToSharedRef())
 			]
 
-			+SHorizontalBox::Slot()
-			.AutoWidth()
-			.Padding(2, 0)
-			.VAlign(VAlign_Center)
-			[
-				SNew(SButton)
-				.ContentPadding(0)
-				.OnClicked(this, &FUserDefinedEnumIndexLayout::OnMoveEnumeratorUp)
-				.IsEnabled( bIsMoveUpEnabled )
-				[
-					SNew(SImage)
-					.Image(FEditorStyle::GetBrush("Icons.ChevronUp"))
-				]
-			]
-			+SHorizontalBox::Slot()
-			.AutoWidth()
-			.Padding(2, 0)
-			.VAlign(VAlign_Center)
-			[
-				SNew(SButton)
-				.ContentPadding(0)
-				.OnClicked(this, &FUserDefinedEnumIndexLayout::OnMoveEnumeratorDown)
-				.IsEnabled( bIsMoveDownEnabled )
-				[
-					SNew(SImage)
-					.Image(FEditorStyle::GetBrush("Icons.ChevronDown"))
-				]
-			]
-
 			+SHorizontalBox::Slot()
 			.AutoWidth()
 			.Padding(2, 0)
@@ -569,7 +676,8 @@ void FUserDefinedEnumIndexLayout::GenerateHeaderRowContent( FDetailWidgetRow& No
 			[
 				ClearButton
 			]
-		];
+		]
+		.DragDropHandler(MakeShared<FUserDefinedEnumIndexDragDropHandler>(TargetEnum, EnumeratorIndex));
 }
 
 void FUserDefinedEnumIndexLayout::OnEnumeratorRemove()
@@ -577,16 +685,4 @@ void FUserDefinedEnumIndexLayout::OnEnumeratorRemove()
 	FEnumEditorUtils::RemoveEnumeratorFromUserDefinedEnum(TargetEnum, EnumeratorIndex);
 }
 
-FReply FUserDefinedEnumIndexLayout::OnMoveEnumeratorUp()
-{
-	FEnumEditorUtils::MoveEnumeratorInUserDefinedEnum(TargetEnum, EnumeratorIndex, true);
-	return FReply::Handled();
-}
-
-FReply FUserDefinedEnumIndexLayout::OnMoveEnumeratorDown()
-{
-	FEnumEditorUtils::MoveEnumeratorInUserDefinedEnum(TargetEnum, EnumeratorIndex, false);
-	return FReply::Handled();
-}
-
 #undef LOCTEXT_NAMESPACE
diff --git a/Engine/Source/Editor/Kismet/Private/UserDefinedEnumEditor.h b/Engine/Source/Editor/Kismet/Private/UserDefinedEnumEditor.h
--- a/Engine/Source/Editor/Kismet/Private/UserDefinedEnumEditor.h
+++ b/Engine/Source/Editor/Kismet/Private/UserDefinedEnumEditor.h
@@ -157,12 +157,6 @@ private:
 	virtual bool InitiallyCollapsed() const override { return false; }
 
 private:
-	/** Moves the enumerator up in the list */
-	FReply OnMoveEnumeratorUp();
-	
-	/** Moves the enumerator down in the list */
-	FReply OnMoveEnumeratorDown();
-
 	/** Deletes the enumerator */
 	void OnEnumeratorRemove();
 
diff --git a/Engine/Source/Editor/UnrealEd/Private/Kismet2/EnumEditorUtils.cpp b/Engine/Source/Editor/UnrealEd/Private/Kismet2/EnumEditorUtils.cpp
--- a/Engine/Source/Editor/UnrealEd/Private/Kismet2/EnumEditorUtils.cpp
+++ b/Engine/Source/Editor/UnrealEd/Private/Kismet2/EnumEditorUtils.cpp
@@ -187,10 +187,9 @@ void FEnumEditorUtils::SetEnumeratorBitflagsTypeState(UUserDefinedEnum* Enum, bo
 	}
 }
 
-/** Reorder enumerators in enum. Swap an enumerator with given name, with previous or next (based on bDirectionUp parameter) */
-void FEnumEditorUtils::MoveEnumeratorInUserDefinedEnum(UUserDefinedEnum* Enum, int32 EnumeratorIndex, bool bDirectionUp)
+void FEnumEditorUtils::MoveEnumeratorInUserDefinedEnum(UUserDefinedEnum* Enum, int32 InitialEnumeratorIndex, int32 TargetIndex)
 {
-	if (!(Enum && (Enum->GetNameByIndex(EnumeratorIndex) != NAME_None)))
+	if (!Enum || Enum->GetNameByIndex(InitialEnumeratorIndex) == NAME_None || TargetIndex < 0 || TargetIndex >= Enum->NumEnums())
 	{
 		return;
 	}
@@ -203,17 +202,9 @@ void FEnumEditorUtils::MoveEnumeratorInUserDefinedEnum(UUserDefinedEnum* Enum, i
 	CopyEnumeratorsWithoutMax(Enum, OldNames);
 	Names = OldNames;
 
-	const bool bIsLast = ((Names.Num() - 1) == EnumeratorIndex);
-	const bool bIsFirst = (0 == EnumeratorIndex);
-
-	if (bDirectionUp && !bIsFirst)
-	{
-		Names.Swap(EnumeratorIndex, EnumeratorIndex - 1);
-	}
-	else if (!bDirectionUp && !bIsLast)
-	{
-		Names.Swap(EnumeratorIndex, EnumeratorIndex + 1);
-	}
+	TPair<FName, int64> EnumeratorToMove = Names[InitialEnumeratorIndex];
+	Names.RemoveAt(InitialEnumeratorIndex);
+	Names.Insert(EnumeratorToMove, TargetIndex);
 
 	// Clean up enum values.
 	for (int32 i = 0; i < Names.Num(); ++i)
diff --git a/Engine/Source/Editor/UnrealEd/Public/Kismet2/EnumEditorUtils.h b/Engine/Source/Editor/UnrealEd/Public/Kismet2/EnumEditorUtils.h
--- a/Engine/Source/Editor/UnrealEd/Public/Kismet2/EnumEditorUtils.h
+++ b/Engine/Source/Editor/UnrealEd/Public/Kismet2/EnumEditorUtils.h
@@ -51,8 +51,11 @@ public:
 	/** Removes enumerator from enum*/
 	static void RemoveEnumeratorFromUserDefinedEnum(class UUserDefinedEnum* Enum, int32 EnumeratorIndex);
 
-	/** Reorder enumerators in enum. Swap an enumerator with given name, with previous or next (based on bDirectionUp parameter) */
-	static void MoveEnumeratorInUserDefinedEnum(class UUserDefinedEnum* Enum, int32 EnumeratorIndex, bool bDirectionUp);
+	/**
+	 * Reorder enumerators in enum. Moves the enumerator at the given initial index to a new target index, shifting other enumerators as needed.
+	 * E.g. with enum [A, B, C, D, E], moving index 1 to index 3 results in [A, C, D, B, E].
+	 */
+	static void MoveEnumeratorInUserDefinedEnum(class UUserDefinedEnum* Enum, int32 InitialEnumeratorIndex, int32 TargetIndex);
 
 	/** Check if the enumerator-as-bitflags meta data is set */
 	static bool IsEnumeratorBitflagsType(class UUserDefinedEnum* Enum);
