UserDefinedEnumEditor: reskin, move "add enumerator" button to toolbar

diff --git a/Engine/Source/Editor/Kismet/Private/UserDefinedEnumEditor.cpp b/Engine/Source/Editor/Kismet/Private/UserDefinedEnumEditor.cpp
--- a/Engine/Source/Editor/Kismet/Private/UserDefinedEnumEditor.cpp
+++ b/Engine/Source/Editor/Kismet/Private/UserDefinedEnumEditor.cpp
@@ -17,6 +17,8 @@
 #include "DetailLayoutBuilder.h"
 #include "DetailCategoryBuilder.h"
 #include "DragAndDrop/DecoratedDragDropOp.h"
+#include "SEditorHeaderButton.h"
+#include "Styling/ToolBarStyle.h"
 
 #include "PropertyCustomizationHelpers.h"
 #include "Widgets/SToolTip.h"
@@ -383,7 +385,9 @@ void FUserDefinedEnumEditor::UnregisterTabSpawners(const TSharedRef<class FTabMa
 
 void FUserDefinedEnumEditor::InitEditor(const EToolkitMode::Type Mode, const TSharedPtr< class IToolkitHost >& InitToolkitHost, UUserDefinedEnum* EnumToEdit)
 {
-	const TSharedRef<FTabManager::FLayout> StandaloneDefaultLayout = FTabManager::NewLayout( "Standalone_UserDefinedEnumEditor_Layout_v2" )
+	TargetEnum = EnumToEdit;
+
+	const TSharedRef<FTabManager::FLayout> StandaloneDefaultLayout = FTabManager::NewLayout( "Standalone_UserDefinedEnumEditor_Layout_v3" )
 	->AddArea
 	(
 		FTabManager::NewPrimaryArea() ->SetOrientation(Orient_Vertical)
@@ -394,6 +398,7 @@ void FUserDefinedEnumEditor::InitEditor(const EToolkitMode::Type Mode, const TSh
 			(
 				FTabManager::NewStack()
 				->AddTab( EnumeratorsTabId, ETabState::OpenedTab )
+				->SetHideTabWell(true)
 			)
 		)
 	);
@@ -402,6 +407,12 @@ void FUserDefinedEnumEditor::InitEditor(const EToolkitMode::Type Mode, const TSh
 	const bool bCreateDefaultToolbar = true;
 	FAssetEditorToolkit::InitAssetEditor( Mode, InitToolkitHost, UserDefinedEnumEditorAppIdentifier, StandaloneDefaultLayout, bCreateDefaultStandaloneMenu, bCreateDefaultToolbar, EnumToEdit );
 
+	TSharedPtr<FExtender> Extender = MakeShared<FExtender>();
+	Extender->AddToolBarExtension("Asset", EExtensionHook::After, GetToolkitCommands(),
+		FToolBarExtensionDelegate::CreateSP(this, &FUserDefinedEnumEditor::FillToolbar));
+	AddToolbarExtender(Extender);
+	RegenerateMenusAndToolbars();
+
 	// @todo toolkit world centric editing
 	/*if (IsWorldCentricAssetEditor())
 	{
@@ -448,6 +459,35 @@ TSharedRef<SDockTab> FUserDefinedEnumEditor::SpawnEnumeratorsTab(const FSpawnTab
 		];
 }
 
+void FUserDefinedEnumEditor::FillToolbar(FToolBarBuilder& ToolbarBuilder)
+{
+	const FToolBarStyle& ToolBarStyle = ToolbarBuilder.GetStyleSet()->GetWidgetStyle<FToolBarStyle>(ToolbarBuilder.GetStyleName());
+
+	ToolbarBuilder.BeginSection("Enumerators");
+	ToolbarBuilder.AddWidget(
+		SNew(SBox)
+		.HAlign(HAlign_Center)
+		.VAlign(VAlign_Fill)
+		.Padding(ToolBarStyle.ButtonPadding)
+		[
+			SNew(SEditorHeaderButton)
+			.Text(LOCTEXT("AddEnumeratorButtonText", "Add Enumerator"))
+			.OnClicked(this, &FUserDefinedEnumEditor::OnAddNewEnumerator)
+		]);
+	ToolbarBuilder.EndSection();
+}
+
+FReply FUserDefinedEnumEditor::OnAddNewEnumerator()
+{
+	if (!ensure(TargetEnum.IsValid()))
+	{
+		return FReply::Handled();
+	}
+
+	FEnumEditorUtils::AddNewEnumeratorForUserDefinedEnum(TargetEnum.Get());
+	return FReply::Handled();
+}
+
 FUserDefinedEnumEditor::~FUserDefinedEnumEditor()
 {
 }
@@ -502,18 +542,9 @@ void FEnumDetails::CustomizeDetails( IDetailLayoutBuilder& DetailLayout )
 
 		const FString DocLink = TEXT("Shared/Editors/BlueprintEditor/EnumDetails");
 
-		IDetailCategoryBuilder& InputsCategory = DetailLayout.EditCategory("Enumerators", LOCTEXT("EnumDetailsEnumerators", "Enumerators"));
+		DetailLayout.EditCategory("Description"); // make Description category appear before Enumerators
 
-		InputsCategory.AddCustomRow( LOCTEXT("FunctionNewInputArg", "New") )
-			[
-				SNew(SBox)
-				.HAlign(HAlign_Right)
-				[
-					SNew(SButton)
-					.Text(LOCTEXT("FunctionNewInputArg", "New"))
-					.OnClicked(this, &FEnumDetails::OnAddNewEnumerator)
-				]
-			];
+		IDetailCategoryBuilder& InputsCategory = DetailLayout.EditCategory("Enumerators", LOCTEXT("EnumDetailsEnumerators", "Enumerators"));
 
 		Layout = MakeShareable( new FUserDefinedEnumLayout(TargetEnum.Get()) );
 		InputsCategory.AddCustomBuilder( Layout.ToSharedRef() );
@@ -525,6 +556,7 @@ void FEnumDetails::CustomizeDetails( IDetailLayoutBuilder& DetailLayout )
 		[
 			SNew(STextBlock)
 			.Text(LOCTEXT("BitmaskFlagsAttributeLabel", "Bitmask Flags"))
+			.Font(IDetailLayoutBuilder::GetDetailFont())
 			.ToolTip(BitmaskFlagsTooltip)
 		]
 		.ValueContent()
@@ -573,12 +605,6 @@ void FEnumDetails::PostChange(const class UUserDefinedEnum* Enum, FEnumEditorUti
 	}
 }
 
-FReply FEnumDetails::OnAddNewEnumerator()
-{
-	FEnumEditorUtils::AddNewEnumeratorForUserDefinedEnum(TargetEnum.Get());
-	return FReply::Handled();
-}
-
 ECheckBoxState FEnumDetails::OnGetBitmaskFlagsAttributeState() const
 {
 	return FEnumEditorUtils::IsEnumeratorBitflagsType(TargetEnum.Get()) ? ECheckBoxState::Checked : ECheckBoxState::Unchecked;
@@ -626,10 +652,10 @@ void FUserDefinedEnumIndexLayout::GenerateHeaderRowContent( FDetailWidgetRow& No
 	TooltipEditor = MakeShared<FEditableTextUserDefinedEnumTooltip>(TargetEnum, EnumeratorIndex);
 
 	const bool bIsEditable = !DisplayNameEditor->IsReadOnly();
-	const bool bIsMoveUpEnabled = (TargetEnum->NumEnums() != 1) && (EnumeratorIndex != 0) && bIsEditable;
-	const bool bIsMoveDownEnabled = (TargetEnum->NumEnums() != 1) && (EnumeratorIndex < TargetEnum->NumEnums() - 2) && bIsEditable;
 
-	TSharedRef< SWidget > ClearButton = PropertyCustomizationHelpers::MakeClearButton(FSimpleDelegate::CreateSP(this, &FUserDefinedEnumIndexLayout::OnEnumeratorRemove));
+	TSharedRef< SWidget > ClearButton = PropertyCustomizationHelpers::MakeEmptyButton(
+		FSimpleDelegate::CreateSP(this, &FUserDefinedEnumIndexLayout::OnEnumeratorRemove),
+		LOCTEXT("RemoveEnumToolTip", "Remove enumerator"));
 	ClearButton->SetEnabled(bIsEditable);
 
 	NodeRow
@@ -640,10 +666,11 @@ void FUserDefinedEnumIndexLayout::GenerateHeaderRowContent( FDetailWidgetRow& No
 			+SHorizontalBox::Slot()
 			.VAlign(VAlign_Center)
 			.AutoWidth()
-			.Padding(0, 0, 4, 0)
+			.Padding(0, 0, 12, 0)
 			[
 				SNew(STextBlock)
 				.Text(LOCTEXT("EnumDisplayNameLabel", "Display Name"))
+				.Font(IDetailLayoutBuilder::GetDetailFont())
 			]
 
 			+SHorizontalBox::Slot()
@@ -651,15 +678,17 @@ void FUserDefinedEnumIndexLayout::GenerateHeaderRowContent( FDetailWidgetRow& No
 			.FillWidth(1.0f)
 			[
 				SNew(STextPropertyEditableTextBox, DisplayNameEditor.ToSharedRef())
+				.Font(IDetailLayoutBuilder::GetDetailFont())
 			]
 
 			+SHorizontalBox::Slot()
 			.VAlign(VAlign_Center)
 			.AutoWidth()
-			.Padding(4, 0, 4, 0)
+			.Padding(24, 0, 12, 0)
 			[
 				SNew(STextBlock)
 				.Text(LOCTEXT("EnumTooltipLabel", "Description"))
+				.Font(IDetailLayoutBuilder::GetDetailFont())
 			]
 
 			+SHorizontalBox::Slot()
@@ -667,6 +696,7 @@ void FUserDefinedEnumIndexLayout::GenerateHeaderRowContent( FDetailWidgetRow& No
 			.FillWidth(1.0f)
 			[
 				SNew(STextPropertyEditableTextBox, TooltipEditor.ToSharedRef())
+				.Font(IDetailLayoutBuilder::GetDetailFont())
 			]
 
 			+SHorizontalBox::Slot()
diff --git a/Engine/Source/Editor/Kismet/Private/UserDefinedEnumEditor.h b/Engine/Source/Editor/Kismet/Private/UserDefinedEnumEditor.h
--- a/Engine/Source/Editor/Kismet/Private/UserDefinedEnumEditor.h
+++ b/Engine/Source/Editor/Kismet/Private/UserDefinedEnumEditor.h
@@ -59,6 +59,15 @@ public:
 
 protected:
 	TSharedRef<SDockTab> SpawnEnumeratorsTab(const FSpawnTabArgs& Args);
+
+private:
+	void FillToolbar(FToolBarBuilder& ToolbarBuilder);
+
+	/** Handles new enum element request */
+	FReply OnAddNewEnumerator();
+
+private:
+	TWeakObjectPtr<UUserDefinedEnum> TargetEnum;
 };
 
 /** Details customization for functions and graphs selected in the MyBlueprint panel */
@@ -89,9 +98,6 @@ public:
 	virtual void PostRedo(bool bSuccess) override { PostUndo(bSuccess); }
 
 private:
-	/** Handles new enum element request */
-	FReply OnAddNewEnumerator();
-
 	/** Handles the optional bitmask flags attribute */
 	ECheckBoxState OnGetBitmaskFlagsAttributeState() const;
 	void OnBitmaskFlagsAttributeStateChanged(ECheckBoxState InNewState);
