PR #7613: Sort categories (Contributed by [REDACTED_PERSON])

diff --git a/Engine/Source/Editor/Kismet/Private/BlueprintDetailsCustomization.cpp b/Engine/Source/Editor/Kismet/Private/BlueprintDetailsCustomization.cpp
--- a/Engine/Source/Editor/Kismet/Private/BlueprintDetailsCustomization.cpp
+++ b/Engine/Source/Editor/Kismet/Private/BlueprintDetailsCustomization.cpp
@@ -1467,7 +1467,7 @@ void FBlueprintVarActionDetails::PopulateCategories(SMyBlueprint* MyBlueprint, T
 	}
 
 	CategorySource.Reset();
-	CategorySource.Add(MakeShareable(new FText(LOCTEXT("Default", "Default"))));
+	CategorySource.Add(MakeShareable(new FText(UEdGraphSchema_K2::VR_DefaultCategory)));
 	for (int32 i = 0; i < VisibleVariables.Num(); ++i)
 	{
 		FText Category = FBlueprintEditorUtils::GetBlueprintVariableCategory(Blueprint, VisibleVariables[i], nullptr);
@@ -1570,6 +1570,16 @@ void FBlueprintVarActionDetails::PopulateCategories(SMyBlueprint* MyBlueprint, T
 			}
 		}
 	}
+
+	// Sort categories, but keep the default category listed first
+	CategorySource.Sort([](const TSharedPtr <FText> &LHS, const TSharedPtr <FText> &RHS)
+	{
+		if (LHS.IsValid() && RHS.IsValid())
+		{
+			return (LHS->EqualTo(UEdGraphSchema_K2::VR_DefaultCategory) || LHS->CompareToCaseIgnored(*RHS) <= 0);
+		}
+		return false;
+	});
 }
 
 UK2Node_Variable* FBlueprintVarActionDetails::EdGraphSelectionAsVar() const
@@ -6392,7 +6402,7 @@ void FBlueprintComponentDetails::PopulateVariableCategories()
 	FBlueprintEditorUtils::GetSCSVariableNameList(BlueprintObj, VisibleVariables);
 
 	VariableCategorySource.Empty();
-	VariableCategorySource.Add(MakeShareable(new FText(LOCTEXT("Default", "Default"))));
+	VariableCategorySource.Add(MakeShareable(new FText(UEdGraphSchema_K2::VR_DefaultCategory)));
 	for (const FName& VariableName : VisibleVariables)
 	{
 		FText Category = FBlueprintEditorUtils::GetBlueprintVariableCategory(BlueprintObj, VariableName, nullptr);
@@ -6409,6 +6419,16 @@ void FBlueprintComponentDetails::PopulateVariableCategories()
 			}
 		}
 	}
+
+	// Sort categories, but keep the default category listed first
+	VariableCategorySource.Sort([](const TSharedPtr <FText> &LHS, const TSharedPtr <FText> &RHS)
+	{
+		if (LHS.IsValid() && RHS.IsValid())
+		{
+			return (LHS->EqualTo(UEdGraphSchema_K2::VR_DefaultCategory) || LHS->CompareToCaseIgnored(*RHS) <= 0);
+		}
+		return false;
+	});
 }
 
 const UClass* FBlueprintComponentDetails::GetSelectedEntryClass() const
