UserDefinedStructureEditor reskin: split into structure/default values tabs, move "add variable" button to toolbar

- Also icon changes/tooltip text updates/small layout tweaks
- Remove bShowDefaultValuePropertyEditor config variable since default values tab can now be removed from the UI
- Cleanup: remove unused EMemberFieldPosition enum and PositionFlags member for FUserDefinedStructureFieldLayout

diff --git a/Engine/Config/BaseEditor.ini b/Engine/Config/BaseEditor.ini
--- a/Engine/Config/BaseEditor.ini
+++ b/Engine/Config/BaseEditor.ini
@@ -304,7 +304,6 @@ SupportWebsiteURL="http://unrealengine.com/support"
 
 [UserDefinedStructure]
 bUseUserDefinedStructure=true
-bShowDefaultValuePropertyEditor=true
 
 [CustomBlueprintFunctionLibrary]
 bCanCreateNew=true
diff --git a/Engine/Source/Editor/EditorStyle/Private/StarshipStyle.cpp b/Engine/Source/Editor/EditorStyle/Private/StarshipStyle.cpp
--- a/Engine/Source/Editor/EditorStyle/Private/StarshipStyle.cpp
+++ b/Engine/Source/Editor/EditorStyle/Private/StarshipStyle.cpp
@@ -5180,6 +5180,7 @@ void FStarshipEditorStyle::FStyle::SetupPersonaStyle()
 		Set( "Kismet.DeleteUnusedVariables", new IMAGE_BRUSH_SVG("/Starship/Blueprints/icon_kismet_findunused", Icon16x16) );
 		{
 			Set( "Kismet.Tabs.Variables", new IMAGE_BRUSH_SVG( "Starship/Blueprints/pill", Icon16x16 ) );
+			Set( "Kismet.Tabs.DefaultValues", new IMAGE_BRUSH_SVG( "Starship/Common/Details", Icon16x16 ) );
 			Set( "Kismet.Tabs.Palette", new IMAGE_BRUSH_SVG( "Starship/Blueprints/Palette", Icon16x16 ) );
 			Set( "Kismet.Tabs.CompilerResults", new IMAGE_BRUSH_SVG("Starship/Common/OutputLog", Icon16x16 ) );
 			Set( "Kismet.Tabs.FindResults", new CORE_IMAGE_BRUSH_SVG("Starship/Common/search", Icon16x16 ) );
diff --git a/Engine/Source/Editor/Kismet/Private/UserDefinedStructureEditor.cpp b/Engine/Source/Editor/Kismet/Private/UserDefinedStructureEditor.cpp
--- a/Engine/Source/Editor/Kismet/Private/UserDefinedStructureEditor.cpp
+++ b/Engine/Source/Editor/Kismet/Private/UserDefinedStructureEditor.cpp
@@ -26,6 +26,8 @@
 #include "Kismet2/BlueprintEditorUtils.h"
 #include "Widgets/Layout/SSplitter.h"
 #include "DragAndDrop/DecoratedDragDropOp.h"
+#include "SEditorHeaderButton.h"
+#include "Styling/ToolBarStyle.h"
 
 #include "PropertyCustomizationHelpers.h"
 #include "SPinTypeSelector.h"
@@ -239,9 +241,14 @@ class FUserDefinedStructureDetails : public IDetailCustomization, FStructureEdit
 {
 public:
 	/** Makes a new instance of this detail layout class for a specific detail view requesting it */
-	static TSharedRef<IDetailCustomization> MakeInstance()
+	static TSharedRef<IDetailCustomization> MakeInstance(TWeakPtr<FUserDefinedStructureEditor> InStructureEditor)
+	{
+		return MakeShareable(new FUserDefinedStructureDetails(InStructureEditor));
+	}
+
+	FUserDefinedStructureDetails(TWeakPtr<FUserDefinedStructureEditor> InStructureEditor)
+		: StructureEditor(InStructureEditor)
 	{
-		return MakeShareable(new FUserDefinedStructureDetails);
 	}
 
 	~FUserDefinedStructureDetails()
@@ -262,6 +269,11 @@ public:
 		return NULL;
 	}
 
+	const TWeakPtr<FUserDefinedStructureEditor>& GetStructureEditor() const
+	{
+		return StructureEditor;
+	}
+
 	/** IDetailCustomization interface */
 	virtual void CustomizeDetails(class IDetailLayoutBuilder& DetailLayout) override;
 
@@ -272,12 +284,14 @@ public:
 private:
 	TWeakObjectPtr<UUserDefinedStruct> UserDefinedStruct;
 	TSharedPtr<class FUserDefinedStructureLayout> Layout;
+	TWeakPtr<FUserDefinedStructureEditor> StructureEditor;
 };
 
 ///////////////////////////////////////////////////////////////////////////////////////
 // FUserDefinedStructureEditor
 
 const FName FUserDefinedStructureEditor::MemberVariablesTabId( TEXT( "UserDefinedStruct_MemberVariablesEditor" ) );
+const FName FUserDefinedStructureEditor::DefaultValuesTabId(TEXT("UserDefinedStruct_DefaultValuesEditor"));
 const FName FUserDefinedStructureEditor::UserDefinedStructureEditorAppIdentifier( TEXT( "UserDefinedStructEditorApp" ) );
 
 void FUserDefinedStructureEditor::RegisterTabSpawners(const TSharedRef<class FTabManager>& InTabManager)
@@ -287,9 +301,14 @@ void FUserDefinedStructureEditor::RegisterTabSpawners(const TSharedRef<class FTa
 	FAssetEditorToolkit::RegisterTabSpawners(InTabManager);
 
 	InTabManager->RegisterTabSpawner( MemberVariablesTabId, FOnSpawnTab::CreateSP(this, &FUserDefinedStructureEditor::SpawnStructureTab) )
-		.SetDisplayName( LOCTEXT("MemberVariablesEditor", "Structure Editor") )
+		.SetDisplayName( LOCTEXT("MemberVariablesEditor", "Structure") )
 		.SetGroup(WorkspaceMenuCategory.ToSharedRef())
 		.SetIcon(FSlateIcon(FEditorStyle::GetStyleSetName(), "Kismet.Tabs.Variables"));
+
+	InTabManager->RegisterTabSpawner( DefaultValuesTabId, FOnSpawnTab::CreateSP(this, &FUserDefinedStructureEditor::SpawnStructureDefaultValuesTab))
+		.SetDisplayName(LOCTEXT("DefaultValuesEditor", "Default Values"))
+		.SetGroup(WorkspaceMenuCategory.ToSharedRef())
+		.SetIcon(FSlateIcon(FEditorStyle::GetStyleSetName(), "Kismet.Tabs.DefaultValues"));
 }
 
 void FUserDefinedStructureEditor::UnregisterTabSpawners(const TSharedRef<class FTabManager>& InTabManager)
@@ -301,7 +320,10 @@ void FUserDefinedStructureEditor::UnregisterTabSpawners(const TSharedRef<class F
 
 void FUserDefinedStructureEditor::InitEditor(const EToolkitMode::Type Mode, const TSharedPtr< class IToolkitHost >& InitToolkitHost, UUserDefinedStruct* Struct)
 {
-	const TSharedRef<FTabManager::FLayout> StandaloneDefaultLayout = FTabManager::NewLayout( "Standalone_UserDefinedStructureEditor_Layout_v2" )
+	UserDefinedStruct = Struct;
+	InitialPinType = FEdGraphPinType(UEdGraphSchema_K2::PC_Boolean, NAME_None, nullptr, EPinContainerType::None, false, FEdGraphTerminalType());
+
+	const TSharedRef<FTabManager::FLayout> StandaloneDefaultLayout = FTabManager::NewLayout( "Standalone_UserDefinedStructureEditor_Layout_v3" )
 	->AddArea
 	(
 		FTabManager::NewPrimaryArea() 
@@ -312,8 +334,9 @@ void FUserDefinedStructureEditor::InitEditor(const EToolkitMode::Type Mode, cons
 			->Split
 			(
 				FTabManager::NewStack()
-				->SetHideTabWell( true )
 				->AddTab( MemberVariablesTabId, ETabState::OpenedTab )
+				->AddTab( DefaultValuesTabId, ETabState::OpenedTab )
+				->SetForegroundTab(MemberVariablesTabId)
 			)
 		)
 	);
@@ -321,6 +344,17 @@ void FUserDefinedStructureEditor::InitEditor(const EToolkitMode::Type Mode, cons
 	const bool bCreateDefaultStandaloneMenu = true;
 	const bool bCreateDefaultToolbar = true;
 	FAssetEditorToolkit::InitAssetEditor( Mode, InitToolkitHost, UserDefinedStructureEditorAppIdentifier, StandaloneDefaultLayout, bCreateDefaultStandaloneMenu, bCreateDefaultToolbar, Struct );
+
+	TSharedPtr<FExtender> Extender = MakeShared<FExtender>();
+	Extender->AddToolBarExtension("Asset", EExtensionHook::After, GetToolkitCommands(),
+		FToolBarExtensionDelegate::CreateSP(this, &FUserDefinedStructureEditor::FillToolbar));
+	AddToolbarExtender(Extender);
+	RegenerateMenusAndToolbars();
+}
+
+void FUserDefinedStructureEditor::SetInitialPinType(FEdGraphPinType PinType)
+{
+	InitialPinType = PinType;
 }
 
 FUserDefinedStructureEditor::~FUserDefinedStructureEditor()
@@ -376,128 +410,142 @@ TSharedRef<SDockTab> FUserDefinedStructureEditor::SpawnStructureTab(const FSpawn
 		EditedStruct = Cast<UUserDefinedStruct>(EditingObjs[ 0 ]);
 	}
 
-	TSharedRef<SSplitter> Splitter = SNew(SSplitter)
-		.Orientation(Orient_Vertical)
-		.PhysicalSplitterHandleSize(2.0f)
-		.HitDetectionSplitterHandleSize(8.f)
-		.ResizeMode(ESplitterResizeMode::FixedPosition);
-	{
-		// Create a property view
-		FPropertyEditorModule& EditModule = FModuleManager::Get().GetModuleChecked<FPropertyEditorModule>("PropertyEditor");
-		FDetailsViewArgs DetailsViewArgs;
-		DetailsViewArgs.bAllowSearch = false;
-		DetailsViewArgs.NameAreaSettings = FDetailsViewArgs::HideNameArea;
-		DetailsViewArgs.bHideSelectionTip = true;
-		DetailsViewArgs.bShowOptions = false;
-		PropertyView = EditModule.CreateDetailView(DetailsViewArgs);
-		FOnGetDetailCustomizationInstance LayoutStructDetails = FOnGetDetailCustomizationInstance::CreateStatic(&FUserDefinedStructureDetails::MakeInstance);
-		PropertyView->RegisterInstancedCustomPropertyLayout(UUserDefinedStruct::StaticClass(), LayoutStructDetails);
-		PropertyView->SetObject(EditedStruct);
-		Splitter->AddSlot()
-		.Value(0.25f)
+	// Create a property view
+	FPropertyEditorModule& EditModule = FModuleManager::Get().GetModuleChecked<FPropertyEditorModule>("PropertyEditor");
+	FDetailsViewArgs DetailsViewArgs;
+	DetailsViewArgs.bAllowSearch = false;
+	DetailsViewArgs.NameAreaSettings = FDetailsViewArgs::HideNameArea;
+	DetailsViewArgs.bHideSelectionTip = true;
+	DetailsViewArgs.bShowOptions = false;
+	PropertyView = EditModule.CreateDetailView(DetailsViewArgs);
+	TWeakPtr<FUserDefinedStructureEditor> LocalWeakThis = SharedThis(this);
+	FOnGetDetailCustomizationInstance LayoutStructDetails = FOnGetDetailCustomizationInstance::CreateStatic(&FUserDefinedStructureDetails::MakeInstance, LocalWeakThis);
+	PropertyView->RegisterInstancedCustomPropertyLayout(UUserDefinedStruct::StaticClass(), LayoutStructDetails);
+	PropertyView->SetObject(EditedStruct);
+
+	return SNew(SDockTab)
+		.Label( LOCTEXT("UserDefinedStructureEditor", "Structure") )
+		.TabColorScale( GetTabColorScale() )
 		[
-			SNew(SBox)
-			.Padding(FMargin(0.0f, 0.0f, 0.0f, 3.0f))
-			[
-				PropertyView.ToSharedRef()
-			]
+			PropertyView.ToSharedRef()
 		];
-	}
+}
 
-	DefaultValueView = NULL;
+TSharedRef<SDockTab> FUserDefinedStructureEditor::SpawnStructureDefaultValuesTab(const FSpawnTabArgs& Args)
+{
+	check(Args.GetTabId() == DefaultValuesTabId);
 
-	static FBoolConfigValueHelper ShowDefaultValuePropertyEditor(TEXT("UserDefinedStructure"), TEXT("bShowDefaultValuePropertyEditor"));
-	if (ShowDefaultValuePropertyEditor)
+	UUserDefinedStruct* EditedStruct = NULL;
+	const TArray<UObject*>& EditingObjs = GetEditingObjects();
+	if (EditingObjs.Num())
 	{
-		DefaultValueView = MakeShareable(new FStructureDefaultValueView(EditedStruct));
-		DefaultValueView->Initialize();
-		auto DefaultValueWidget = DefaultValueView->GetWidget();
-		if (DefaultValueWidget.IsValid())
-		{
-			Splitter->AddSlot()
-			[
-				SNew(SBox)
-				.Padding(FMargin(0.0f, 3.0f, 0.0f, 0.0f))
-				[
-					DefaultValueWidget.ToSharedRef()
-				]
-			];
-		}
+		EditedStruct = Cast<UUserDefinedStruct>(EditingObjs[0]);
 	}
 
+	DefaultValueView = MakeShareable(new FStructureDefaultValueView(EditedStruct));
+	DefaultValueView->Initialize();
+
 	return SNew(SDockTab)
-		.Label( LOCTEXT("UserDefinedStructureEditor", "Structure Editor") )
-		.TabColorScale( GetTabColorScale() )
+		.Label(LOCTEXT("UserDefinedStructureDefaultValuesEditor", "Default Values"))
+		.TabColorScale(GetTabColorScale())
 		[
-			Splitter
+			DefaultValueView->GetWidget().ToSharedRef()
 		];
 }
 
-///////////////////////////////////////////////////////////////////////////////////////
-// FUserDefinedStructureLayout
-
-//Represents single structure (List of fields)
-class FUserDefinedStructureLayout : public IDetailCustomNodeBuilder, public TSharedFromThis<FUserDefinedStructureLayout>
+void FUserDefinedStructureEditor::FillToolbar(FToolBarBuilder& ToolbarBuilder)
 {
-public:
-	FUserDefinedStructureLayout(TWeakPtr<class FUserDefinedStructureDetails> InStructureDetails)
-		: StructureDetails(InStructureDetails)
-		, InitialPinType(UEdGraphSchema_K2::PC_Boolean, NAME_None, nullptr, EPinContainerType::None, false, FEdGraphTerminalType())
-	{}
+	const FToolBarStyle& ToolBarStyle = ToolbarBuilder.GetStyleSet()->GetWidgetStyle<FToolBarStyle>(ToolbarBuilder.GetStyleName());
 
-	void OnChanged()
+	ToolbarBuilder.BeginSection("UserDefinedStructure");
+
+	TSharedPtr<SLayeredImage> CompileStatusImage;
+	ToolbarBuilder.AddWidget(
+		SNew(SBox)
+		.HAlign(HAlign_Center)
+		.VAlign(VAlign_Center)
+		.Padding(ToolBarStyle.ButtonPadding)
+		[
+			SAssignNew(CompileStatusImage, SLayeredImage)
+			.Image(FAppStyle::Get().GetBrush("Blueprint.CompileStatus.Background"))
+			.ToolTipText(this, &FUserDefinedStructureEditor::OnGetStatusTooltip)
+		]);
+	CompileStatusImage->AddLayer(TAttribute<const FSlateBrush*>::CreateSP(this, &FUserDefinedStructureEditor::OnGetStructureStatus));
+
+	ToolbarBuilder.AddWidget(
+		SNew(SBox)
+		.HAlign(HAlign_Center)
+		.VAlign(VAlign_Fill)
+		.Padding(ToolBarStyle.ButtonPadding)
+		[
+			SNew(SEditorHeaderButton)
+			.Text(LOCTEXT("AddStructVariable", "Add Variable"))
+			.ToolTipText(LOCTEXT("AddStructVariableToolTip", "Adds a new member variable to the end of this structure"))
+			.OnClicked(this, &FUserDefinedStructureEditor::OnAddNewField)
+		]);
+
+	ToolbarBuilder.EndSection();
+}
+
+FReply FUserDefinedStructureEditor::OnAddNewField()
+{
+	if (UserDefinedStruct.IsValid())
 	{
-		OnRegenerateChildren.ExecuteIfBound();
+		FStructureEditorUtils::AddVariable(UserDefinedStruct.Get(), InitialPinType);
+
+		// Ensure the member variables tab is topmost so the user can edit the newly-added variable
+		InvokeTab(MemberVariablesTabId);
 	}
 
-	FReply OnAddNewField()
+	return FReply::Handled();
+}
+
+const FSlateBrush* FUserDefinedStructureEditor::OnGetStructureStatus() const
+{
+	if (UserDefinedStruct.IsValid())
 	{
-		auto StructureDetailsSP = StructureDetails.Pin();
-		if(StructureDetailsSP.IsValid())
+		switch (UserDefinedStruct->Status.GetValue())
 		{
-			FStructureEditorUtils::AddVariable(StructureDetailsSP->GetUserDefinedStruct(), InitialPinType);
+		case EUserDefinedStructureStatus::UDSS_Error:
+			return FEditorStyle::GetBrush("Blueprint.CompileStatus.Overlay.Error");
+		case EUserDefinedStructureStatus::UDSS_UpToDate:
+			return FEditorStyle::GetBrush("Blueprint.CompileStatus.Overlay.Good");
+		default:
+			return FEditorStyle::GetBrush("Blueprint.CompileStatus.Overlay.Unknown");
 		}
-
-		return FReply::Handled();
 	}
+	return nullptr;
+}
 
-	const FSlateBrush* OnGetStructureStatus() const
+FText FUserDefinedStructureEditor::OnGetStatusTooltip() const
+{
+	if (UserDefinedStruct.IsValid())
 	{
-		auto StructureDetailsSP = StructureDetails.Pin();
-		if(StructureDetailsSP.IsValid())
+		switch (UserDefinedStruct->Status.GetValue())
 		{
-			if(auto Struct = StructureDetailsSP->GetUserDefinedStruct())
-			{
-				switch(Struct->Status.GetValue())
-				{
-				case EUserDefinedStructureStatus::UDSS_Error:
-					return FEditorStyle::GetBrush("Blueprint.CompileStatus.Overlay.Error");
-				case EUserDefinedStructureStatus::UDSS_UpToDate:
-					return FEditorStyle::GetBrush("Blueprint.CompileStatus.Overlay.Good");
-				default:
-					return FEditorStyle::GetBrush("Blueprint.CompileStatus.Overlay.Unknown");
-				}
-				
-			}
+		case EUserDefinedStructureStatus::UDSS_Error:
+			return FText::FromString(UserDefinedStruct->ErrorMessage);
+		default:
+			return LOCTEXT("GoodToGo_Status", "Good to go");
 		}
-		return nullptr;
 	}
+	return FText::GetEmpty();
+}
+
+///////////////////////////////////////////////////////////////////////////////////////
+// FUserDefinedStructureLayout
+
+//Represents single structure (List of fields)
+class FUserDefinedStructureLayout : public IDetailCustomNodeBuilder, public TSharedFromThis<FUserDefinedStructureLayout>
+{
+public:
+	FUserDefinedStructureLayout(TWeakPtr<class FUserDefinedStructureDetails> InStructureDetails)
+		: StructureDetails(InStructureDetails)
+	{}
 
-	FText GetStatusTooltip() const
+	void OnChanged()
 	{
-		auto StructureDetailsSP = StructureDetails.Pin();
-		if (StructureDetailsSP.IsValid())
-		{
-			if (auto Struct = StructureDetailsSP->GetUserDefinedStruct())
-			{
-				switch (Struct->Status.GetValue())
-				{
-				case EUserDefinedStructureStatus::UDSS_Error:
-					return FText::FromString(Struct->ErrorMessage);
-				}
-			}
-		}
-		return FText::GetEmpty();
+		OnRegenerateChildren.ExecuteIfBound();
 	}
 
 	FText OnGetTooltipText() const
@@ -525,12 +573,6 @@ public:
 		}
 	}
 
-	/** Callback when a pin type is selected to cache the value so new variables in the struct will be set to the cached type */
-	void OnPinTypeSelected(const FEdGraphPinType& InPinType)
-	{
-		InitialPinType = InPinType;
-	}
-
 	/** IDetailCustomNodeBuilder Interface*/
 	virtual void SetOnRebuildChildren( FSimpleDelegate InOnRegenerateChildren ) override 
 	{
@@ -559,15 +601,6 @@ public:
 private:
 	TWeakPtr<class FUserDefinedStructureDetails> StructureDetails;
 	FSimpleDelegate OnRegenerateChildren;
-
-	/** Cached value of the last pin type the user selected, used as the initial value for new struct members */
-	FEdGraphPinType InitialPinType;
-};
-
-enum EMemberFieldPosition
-{
-	MFP_First	=	0x1,
-	MFP_Last	=	0x2,
 };
 
 ///////////////////////////////////////////////////////////////////////////////////////
@@ -711,11 +744,11 @@ private:
 class FUserDefinedStructureFieldLayout : public IDetailCustomNodeBuilder, public TSharedFromThis<FUserDefinedStructureFieldLayout>
 {
 public:
-	FUserDefinedStructureFieldLayout(TWeakPtr<class FUserDefinedStructureDetails> InStructureDetails, TWeakPtr<class FUserDefinedStructureLayout> InStructureLayout, FGuid InFieldGuid, uint32 InPositionFlags)
+	FUserDefinedStructureFieldLayout(TWeakPtr<class FUserDefinedStructureDetails> InStructureDetails, TWeakPtr<class FUserDefinedStructureLayout> InStructureLayout, FGuid InFieldGuid)
 		: StructureDetails(InStructureDetails)
 		, StructureLayout(InStructureLayout)
 		, FieldGuid(InFieldGuid)
-		, PositionFlags(InPositionFlags) {}
+	{}
 
 	void OnChanged()
 	{
@@ -761,10 +794,9 @@ public:
 		if(StructureDetailsSP.IsValid())
 		{
 			FStructureEditorUtils::ChangeVariableType(StructureDetailsSP->GetUserDefinedStruct(), FieldGuid, PinType);
-			auto StructureLayoutPin = StructureLayout.Pin();
-			if (StructureLayoutPin.IsValid())
+			if (TSharedPtr<FUserDefinedStructureEditor> StructureEditorSP = StructureDetailsSP->GetStructureEditor().Pin())
 			{
-				StructureLayoutPin->OnPinTypeSelected(PinType);
+				StructureEditorSP->SetInitialPinType(PinType);
 			}
 		}
 	}
@@ -1002,6 +1034,7 @@ public:
 			[
 				SAssignNew(ErrorIcon, SImage)
 				.Image( FEditorStyle::GetBrush("Icons.Error") )
+				.ToolTipText(LOCTEXT("MemberVariableErrorToolTip", "Member variable is invalid"))
 			]
 
 			+SHorizontalBox::Slot()
@@ -1036,7 +1069,7 @@ public:
 			.HAlign(HAlign_Right)
 			.VAlign(VAlign_Center)
 			[
-				PropertyCustomizationHelpers::MakeClearButton(
+				PropertyCustomizationHelpers::MakeEmptyButton(
 					FSimpleDelegate::CreateSP(this, &FUserDefinedStructureFieldLayout::OnRemovField),
 					LOCTEXT("RemoveVariable", "Remove member variable"),
 					TAttribute<bool>::Create(TAttribute<bool>::FGetter::CreateSP(this, &FUserDefinedStructureFieldLayout::IsRemoveButtonEnabled)))
@@ -1146,73 +1179,27 @@ private:
 	FGuid FieldGuid;
 
 	FSimpleDelegate OnRegenerateChildren;
-
-	uint32 PositionFlags;
 };
 
 ///////////////////////////////////////////////////////////////////////////////////////
 // FUserDefinedStructureLayout
 void FUserDefinedStructureLayout::GenerateChildContent( IDetailChildrenBuilder& ChildrenBuilder ) 
 {
-	const float NameWidth = 80.0f;
-	const float ContentWidth = 130.0f;
-
-	TSharedPtr<SLayeredImage> StatusImage;
 	ChildrenBuilder.AddCustomRow(FText::GetEmpty())
-	[
-		SNew(SHorizontalBox)
-		+ SHorizontalBox::Slot()
-		.MaxWidth(NameWidth)
-		.HAlign(HAlign_Left)
-		.VAlign(VAlign_Center)
-		[
-			SAssignNew(StatusImage, SLayeredImage)
-			.Image(FAppStyle::Get().GetBrush("Blueprint.CompileStatus.Background"))
-			.ToolTipText(this, &FUserDefinedStructureLayout::GetStatusTooltip)
-			.DesiredSizeOverride(FVector2D(16,16))
-		]
-		+ SHorizontalBox::Slot()
-		.AutoWidth()
-		.HAlign(HAlign_Left)
-		[
-			SNew(SBox)
-			.WidthOverride(ContentWidth)
-			[
-				SNew(SButton)
-				.HAlign(HAlign_Center)
-				.Text(LOCTEXT("NewStructureField", "New Variable"))
-				.OnClicked(this, &FUserDefinedStructureLayout::OnAddNewField)
-			]
-		]
-	];
-
-	ChildrenBuilder.AddCustomRow(FText::GetEmpty())
-	[
-		SNew(SHorizontalBox)
-		+ SHorizontalBox::Slot()
-		.MaxWidth(NameWidth)
-		.HAlign(HAlign_Left)
+		.NameContent()
 		[
 			SNew(STextBlock)
 			.Text(LOCTEXT("Tooltip", "Tooltip"))
 			.Font(IDetailLayoutBuilder::GetDetailFont())
 		]
-		+ SHorizontalBox::Slot()
-		.AutoWidth()
-		.HAlign(HAlign_Left)
+		.ValueContent()
+		.MinDesiredWidth(400.0f)
 		[
-			SNew(SBox)
-			.WidthOverride(ContentWidth)
-			[
-				SNew(SEditableTextBox)
-				.Text(this, &FUserDefinedStructureLayout::OnGetTooltipText)
-				.OnTextCommitted(this, &FUserDefinedStructureLayout::OnTooltipCommitted)
-				.Font(IDetailLayoutBuilder::GetDetailFont())
-			]
-		]
-	];
-
-	StatusImage->AddLayer(MakeAttributeSP(this, &FUserDefinedStructureLayout::OnGetStructureStatus));
+			SNew(SEditableTextBox)
+			.Text(this, &FUserDefinedStructureLayout::OnGetTooltipText)
+			.OnTextCommitted(this, &FUserDefinedStructureLayout::OnTooltipCommitted)
+			.Font(IDetailLayoutBuilder::GetDetailFont())
+		];
 
 	auto StructureDetailsSP = StructureDetails.Pin();
 	if(StructureDetailsSP.IsValid())
@@ -1223,10 +1210,7 @@ void FUserDefinedStructureLayout::GenerateChildContent( IDetailChildrenBuilder&
 			for (int32 Index = 0; Index < VarDescArrayRef.Num(); ++Index)
 			{
 				auto& VarDesc = VarDescArrayRef[Index];
-				uint32 PositionFlag = 0;
-				PositionFlag |= (0 == Index) ? EMemberFieldPosition::MFP_First : 0;
-				PositionFlag |= ((VarDescArrayRef.Num() - 1) == Index) ? EMemberFieldPosition::MFP_Last : 0;
-				TSharedRef<class FUserDefinedStructureFieldLayout> VarLayout = MakeShareable(new FUserDefinedStructureFieldLayout(StructureDetails,  SharedThis(this), VarDesc.VarGuid, PositionFlag));
+				TSharedRef<class FUserDefinedStructureFieldLayout> VarLayout = MakeShareable(new FUserDefinedStructureFieldLayout(StructureDetails,  SharedThis(this), VarDesc.VarGuid));
 				ChildrenBuilder.AddCustomBuilder(VarLayout);
 			}
 		}
diff --git a/Engine/Source/Editor/Kismet/Private/UserDefinedStructureEditor.h b/Engine/Source/Editor/Kismet/Private/UserDefinedStructureEditor.h
--- a/Engine/Source/Editor/Kismet/Private/UserDefinedStructureEditor.h
+++ b/Engine/Source/Editor/Kismet/Private/UserDefinedStructureEditor.h
@@ -16,6 +16,7 @@ class FUserDefinedStructureEditor : public IUserDefinedStructureEditor
 
 	/**	The tab ids for all the tabs used */
 	static const FName MemberVariablesTabId;
+	static const FName DefaultValuesTabId;
 	
 	/** Property viewing widget */
 	TSharedPtr<class IDetailsView> PropertyView;
@@ -30,6 +31,9 @@ public:
 	 */
 	void InitEditor(const EToolkitMode::Type Mode, const TSharedPtr< class IToolkitHost >& InitToolkitHost, class UUserDefinedStruct* EnumToEdit);
 
+	/** Sets the pin type for new struct members added via the Add Variable toolbar button */
+	void SetInitialPinType(FEdGraphPinType PinType);
+
 	/** Destructor */
 	virtual ~FUserDefinedStructureEditor();
 
@@ -46,5 +50,24 @@ public:
 
 protected:
 	TSharedRef<SDockTab> SpawnStructureTab(const FSpawnTabArgs& Args);
+	TSharedRef<SDockTab> SpawnStructureDefaultValuesTab(const FSpawnTabArgs& Args);
+
+private:
+	void FillToolbar(FToolBarBuilder& ToolbarBuilder);
+
+	/** Handles adding a new member variable from the toolbar button */
+	FReply OnAddNewField();
+
+	/** Returns the overlay image to indicate the compile status */
+	const FSlateBrush* OnGetStructureStatus() const;
+
+	/** Returns the tooltip describing the compile status */
+	FText OnGetStatusTooltip() const;
+
+private:
+	TWeakObjectPtr<UUserDefinedStruct> UserDefinedStruct;
+
+	/** Cached value of the last pin type the user selected, used as the initial value for new struct members */
+	FEdGraphPinType InitialPinType;
 };
 
