Fix a regression which caused newly-created child Blueprint assets to no longer inherit default value overrides from variables defined by the parent class.

diff --git a/Engine/Source/Editor/Kismet/Private/BlueprintCompilationManager.cpp b/Engine/Source/Editor/Kismet/Private/BlueprintCompilationManager.cpp
--- a/Engine/Source/Editor/Kismet/Private/BlueprintCompilationManager.cpp
+++ b/Engine/Source/Editor/Kismet/Private/BlueprintCompilationManager.cpp
@@ -1096,7 +1096,7 @@ void FBlueprintCompilationManagerImpl::FlushCompilationQueueImpl(bool bSuppressB
 			for (FCompilerData& CompilerData : CurrentlyCompilingBPs)
 			{
 				if (CompilerData.JobType == ECompilationManagerJobType::Normal &&
-					!CompilerData.BP->bIsRegeneratingOnLoad &&
+					 CompilerData.BP->bHasBeenRegenerated &&		// Note: This ensures that we'll only do this after the Blueprint has been loaded/created; otherwise the class may not contain any properties to find.
 					 CompilerData.BP->GeneratedClass)
 				{
 					const UClass* ParentClass = CompilerData.BP->ParentClass;
