Fix non-unity build errors

diff --git a/Engine/Source/Editor/Kismet/Private/Debugging/SKismetDebuggingView.cpp b/Engine/Source/Editor/Kismet/Private/Debugging/SKismetDebuggingView.cpp
--- a/Engine/Source/Editor/Kismet/Private/Debugging/SKismetDebuggingView.cpp
+++ b/Engine/Source/Editor/Kismet/Private/Debugging/SKismetDebuggingView.cpp
@@ -35,6 +35,7 @@
 #include "Kismet2/BlueprintEditorUtils.h"
 #include "PropertyInfoViewStyle.h"
 #include "GenericPlatform/GenericPlatformApplicationMisc.h"
+#include "Widgets/Input/SCheckBox.h"
 #include "Widgets/Input/SSearchBox.h"
 #include "HAL/PlatformApplicationMisc.h"
 
diff --git a/Engine/Source/Editor/Kismet/Private/Debugging/SKismetDebuggingView.h b/Engine/Source/Editor/Kismet/Private/Debugging/SKismetDebuggingView.h
--- a/Engine/Source/Editor/Kismet/Private/Debugging/SKismetDebuggingView.h
+++ b/Engine/Source/Editor/Kismet/Private/Debugging/SKismetDebuggingView.h
@@ -16,6 +16,10 @@ class UBlueprint;
 class UBlueprintGeneratedClass;
 class FBreakpointParentItem;
 
+class SSearchBox;
+class SComboButton;
+class SCheckBox;
+
 //////////////////////////////////////////////////////////////////////////
 // FDebugLineItem
 
