Drag-and-drop to reorder blueprint function inputs/outputs (replaces up/down buttons)

- Made minor tweak to row layout so that pin type selector resizes when there's limited space

diff --git a/Engine/Source/Editor/Kismet/Private/BlueprintDetailsCustomization.cpp b/Engine/Source/Editor/Kismet/Private/BlueprintDetailsCustomization.cpp
--- a/Engine/Source/Editor/Kismet/Private/BlueprintDetailsCustomization.cpp
+++ b/Engine/Source/Editor/Kismet/Private/BlueprintDetailsCustomization.cpp
@@ -51,6 +51,7 @@
 #include "IDetailPropertyRow.h"
 #include "DetailCategoryBuilder.h"
 #include "IDetailsView.h"
+#include "IDetailDragDropHandler.h"
 #include "Widgets/Colors/SColorPicker.h"
 #include "SKismetInspector.h"
 #include "Widgets/SToolTip.h"
@@ -63,6 +64,7 @@
 #include "Widgets/Layout/SWidgetSwitcher.h"
 #include "BlueprintNamespaceHelper.h"
 #include "Widgets/Input/SSuggestionTextBox.h"
+#include "DragAndDrop/DecoratedDragDropOp.h"
 
 #include "Modules/ModuleManager.h"
 #include "ISequencerModule.h"
@@ -2755,14 +2757,191 @@ bool FBlueprintVarActionDetails::IsVariableDeprecated() const
 	return Property && Property->HasAnyPropertyFlags(CPF_Deprecated);
 }
 
-static FDetailWidgetRow& AddRow( TArray<TSharedRef<FDetailWidgetRow> >& OutChildRows )
+static TArray<UK2Node_EditablePinBase*> GatherAllResultNodes(UK2Node_EditablePinBase* TargetNode)
 {
-	TSharedRef<FDetailWidgetRow> NewRow( new FDetailWidgetRow );
-	OutChildRows.Add( NewRow );
-
-	return *NewRow;
+	if (UK2Node_FunctionResult* ResultNode = Cast<UK2Node_FunctionResult>(TargetNode))
+	{
+		return (TArray<UK2Node_EditablePinBase*>)ResultNode->GetAllResultNodes();
+	}
+	TArray<UK2Node_EditablePinBase*> Result;
+	if (TargetNode)
+	{
+		Result.Add(TargetNode);
+	}
+	return Result;
 }
 
+/** Drag-and-drop operation that stores data about the function parameter pin being dragged */
+class FBlueprintGraphArgumentDragDropOp : public FDecoratedDragDropOp
+{
+public:
+	DRAG_DROP_OPERATOR_TYPE(FBlueprintGraphArgumentDragDropOp, FDecoratedDragDropOp);
+
+	FBlueprintGraphArgumentDragDropOp(UK2Node_EditablePinBase* InTargetNode, TWeakPtr<FUserPinInfo> InParamItemPtr)
+		: TargetNode(InTargetNode)
+		, ParamItemPtr(InParamItemPtr)
+	{
+	}
+
+	void Init()
+	{
+		SetValidTarget(false);
+		SetupDefaults();
+		Construct();
+	}
+
+	void SetValidTarget(bool IsValidTarget)
+	{
+		FText PinName = FText::FromName(ParamItemPtr.IsValid() ? ParamItemPtr.Pin()->PinName : NAME_None);
+		FFormatNamedArguments Args;
+		Args.Add(TEXT("PinName"), PinName);
+
+		if (IsValidTarget)
+		{
+			CurrentHoverText = FText::Format(LOCTEXT("MovePinHere", "Move '{PinName}' Here"), Args);
+			CurrentIconBrush = FEditorStyle::GetBrush("Graph.ConnectorFeedback.OK");
+		}
+		else
+		{
+			CurrentHoverText = FText::Format(LOCTEXT("CannotMovePinHere", "Cannot Move '{PinName}' Here"), Args);
+			CurrentIconBrush = FEditorStyle::GetBrush("Graph.ConnectorFeedback.Error");
+		}
+	}
+
+	UK2Node_EditablePinBase* GetTargetNode() const
+	{
+		return TargetNode;
+	}
+
+	TWeakPtr<FUserPinInfo> GetParamItem() const
+	{
+		return ParamItemPtr;
+	}
+
+private:
+	UK2Node_EditablePinBase* TargetNode;
+	TWeakPtr<FUserPinInfo> ParamItemPtr;
+};
+
+/** Handler for customizing the drag-and-drop behavior for function entry/result pins, allowing parameters to be reordered */
+class FBlueprintGraphArgumentDragDropHandler : public IDetailDragDropHandler
+{
+public:
+	FBlueprintGraphArgumentDragDropHandler(
+		TWeakPtr<FBaseBlueprintGraphActionDetails> InGraphActionDetailsPtr,
+		UK2Node_EditablePinBase* InTargetNode,
+		TWeakPtr<FUserPinInfo> InParamItemPtr)
+		: GraphActionDetailsPtr(InGraphActionDetailsPtr)
+		, TargetNode(InTargetNode)
+		, ParamItemPtr(InParamItemPtr)
+	{
+	}
+
+	virtual TSharedPtr<FDragDropOperation> CreateDragDropOperation() const override
+	{
+		TSharedPtr<FBlueprintGraphArgumentDragDropOp> DragOp = MakeShared<FBlueprintGraphArgumentDragDropOp>(TargetNode, ParamItemPtr);
+		DragOp->Init();
+		return DragOp;
+	}
+
+	/** Compute new target index for use with AcceptDrop/CanAcceptDrop based on drop zone (above vs below) */
+	static int32 ComputeNewIndex(int32 OriginalIndex, int32 DropOntoIndex, EItemDropZone DropZone)
+	{
+		check(DropZone != EItemDropZone::OntoItem);
+
+		int32 NewIndex = DropOntoIndex;
+		if (DropZone == EItemDropZone::BelowItem)
+		{
+			// If the drop zone is below, then we actually move it to the next item's index
+			NewIndex++;
+		}
+		if (OriginalIndex < NewIndex)
+		{
+			// If the item is moved down the list, then all the other elements below it are shifted up one
+			NewIndex--;
+		}
+
+		return ensure(NewIndex >= 0) ? NewIndex : 0;
+	}
+
+	virtual bool AcceptDrop(const FDragDropEvent& DragDropEvent, EItemDropZone DropZone) const override
+	{
+		const TSharedPtr<FBlueprintGraphArgumentDragDropOp> DragOp = DragDropEvent.GetOperationAs<FBlueprintGraphArgumentDragDropOp>();
+		if (!DragOp.IsValid() || DragOp->GetTargetNode() != TargetNode)
+		{
+			return false;
+		}
+
+		if (!ensure(ParamItemPtr.IsValid()) || !ensure(DragOp->GetParamItem().IsValid()))
+		{
+			return false;
+		}
+
+		// Check that the original and new indices are valid, and that they aren't the same (we're actually moving something)
+		const int32 OriginalParamIndex = DragOp->GetTargetNode()->UserDefinedPins.Find(DragOp->GetParamItem().Pin());
+		const int32 OntoParamIndex = TargetNode->UserDefinedPins.Find(ParamItemPtr.Pin());
+		const int32 NewParamIndex = ComputeNewIndex(OriginalParamIndex, OntoParamIndex, DropZone);
+		if (OriginalParamIndex == INDEX_NONE || OriginalParamIndex == NewParamIndex || NewParamIndex < 0 || NewParamIndex >= TargetNode->UserDefinedPins.Num())
+		{
+			return false;
+		}
+
+		const FScopedTransaction Transaction(LOCTEXT("K2_MovePin", "Move Pin"));
+		TArray<UK2Node_EditablePinBase*> TargetNodes = GatherAllResultNodes(TargetNode);
+		for (UK2Node_EditablePinBase* Node : TargetNodes)
+		{
+			Node->Modify();
+
+			TSharedPtr<FUserPinInfo> ParamToMove = Node->UserDefinedPins[OriginalParamIndex];
+			Node->UserDefinedPins.RemoveAt(OriginalParamIndex);
+			Node->UserDefinedPins.Insert(ParamToMove, NewParamIndex);
+
+			TSharedPtr<FBaseBlueprintGraphActionDetails> GraphActionDetails = GraphActionDetailsPtr.Pin();
+			if (GraphActionDetails.IsValid())
+			{
+				GraphActionDetails->OnParamsChanged(Node, true);
+			}
+		}
+
+		return true;
+	}
+
+	virtual TOptional<EItemDropZone> CanAcceptDrop(const FDragDropEvent& DragDropEvent, EItemDropZone DropZone) const override
+	{
+		const TSharedPtr<FBlueprintGraphArgumentDragDropOp> DragOp = DragDropEvent.GetOperationAs<FBlueprintGraphArgumentDragDropOp>();
+		if (!DragOp.IsValid() || DragOp->GetTargetNode() != TargetNode)
+		{
+			return TOptional<EItemDropZone>();
+		}
+
+		// We're reordering, so there's no logical interpretation for dropping directly onto another parameter.
+		// Just change it to a drop-above in this case.
+		const EItemDropZone OverrideZone = (DropZone == EItemDropZone::BelowItem) ? EItemDropZone::BelowItem : EItemDropZone::AboveItem;
+
+		// Check that the original and new indices are valid, and that they aren't the same (we're actually moving something)
+		const int32 OriginalParamIndex = DragOp->GetTargetNode()->UserDefinedPins.Find(DragOp->GetParamItem().Pin());
+		const int32 OntoParamIndex = TargetNode->UserDefinedPins.Find(ParamItemPtr.Pin());
+		const int32 NewParamIndex = ComputeNewIndex(OriginalParamIndex, OntoParamIndex, OverrideZone);
+		if (OriginalParamIndex == INDEX_NONE || OriginalParamIndex == NewParamIndex || NewParamIndex < 0 || NewParamIndex >= TargetNode->UserDefinedPins.Num())
+		{
+			return TOptional<EItemDropZone>();
+		}
+
+		DragOp->SetValidTarget(true);
+		return OverrideZone;
+	}
+
+private:
+	/** The parent graph action details customization */
+	TWeakPtr<FBaseBlueprintGraphActionDetails> GraphActionDetailsPtr;
+
+	/** The target node that the argument pin is on */
+	UK2Node_EditablePinBase* TargetNode;
+
+	/** The argument pin that this drag handler reflects */
+	TWeakPtr<FUserPinInfo> ParamItemPtr;
+};
+
 void FBlueprintGraphArgumentGroupLayout::SetOnRebuildChildren( FSimpleDelegate InOnRegenerateChildren )
 {
 	GraphActionDetailsPtr.Pin()->SetRefreshDelegate(InOnRegenerateChildren, TargetNode == GraphActionDetailsPtr.Pin()->GetFunctionEntryNode().Get());
@@ -2780,11 +2959,25 @@ void FBlueprintGraphArgumentGroupLayout::GenerateChildContent( IDetailChildrenBu
 			bool bIsInputNode = TargetNode == GraphActionDetailsPtr.Pin()->GetFunctionEntryNode().Get();
 			for (int32 i = 0; i < Pins.Num(); ++i)
 			{
+				// If possible, use stable guids for the argument names since the path names are used to store
+				// expansion state. Using guids means that the expansion state travels with the row when
+				// reordering arguments.
+				// Fall back to the old style of using the pin index for the name if we can't find the pin.
+				FString ArgumentName;
+				if (UEdGraphPin* Pin = TargetNode->FindPin(Pins[i]->PinName, Pins[i]->DesiredPinDirection))
+				{
+					ArgumentName = Pin->PinId.ToString();
+				}
+				else
+				{
+					ArgumentName = bIsInputNode ? FString::Printf(TEXT("%s.InputArgument%i"), i) : FString::Printf(TEXT("OutputArgument%i"), i);
+				}
+
 				TSharedRef<class FBlueprintGraphArgumentLayout> BlueprintArgumentLayout = MakeShareable(new FBlueprintGraphArgumentLayout(
 					TWeakPtr<FUserPinInfo>(Pins[i]),
 					TargetNode.Get(),
 					GraphActionDetailsPtr,
-					FName(*(bIsInputNode ? FString::Printf(TEXT("InputArgument%i"), i) : FString::Printf(TEXT("OutputArgument%i"), i))),
+					FName(*ArgumentName),
 					bIsInputNode));
 				ChildrenBuilder.AddCustomBuilder(BlueprintArgumentLayout);
 				WasContentAdded = true;
@@ -2880,8 +3073,8 @@ void FBlueprintGraphArgumentLayout::GenerateHeaderRowContent( FDetailWidgetRow&
 		SNew(SHorizontalBox)
 		+SHorizontalBox::Slot()
 		.VAlign(VAlign_Center)
-		.Padding(0.0f, 0.0f, 4.0f, 0.0f)
-		.AutoWidth()
+		.Padding(0)
+		.FillWidth(1.0f)
 		[
 			SNew(SPinTypeSelector, FGetPinTypeTree::CreateUObject(K2Schema, &UEdGraphSchema_K2::GetVariableTypeTree))
 				.TargetPinType(this, &FBlueprintGraphArgumentLayout::OnGetPinInfo)
@@ -2895,37 +3088,6 @@ void FBlueprintGraphArgumentLayout::GenerateHeaderRowContent( FDetailWidgetRow&
 				.CustomFilter(CustomPinTypeFilter)
 		]
 		+ SHorizontalBox::Slot()
-		.AutoWidth()
-		[
-			SNew(SButton)
-			.ButtonStyle(FAppStyle::Get(), TEXT("SimpleButton"))
-			.ContentPadding(0)
-			.IsEnabled(!IsPinEditingReadOnly())
-			.OnClicked(this, &FBlueprintGraphArgumentLayout::OnArgMoveUp)
-			.ToolTipText(LOCTEXT("FunctionArgDetailsArgMoveUpTooltip", "Move this parameter up in the list."))
-			[
-				SNew(SImage)
-				.Image(FEditorStyle::GetBrush("Icons.ChevronUp"))
-				.ColorAndOpacity(FSlateColor::UseForeground())
-			]
-		]
-		+ SHorizontalBox::Slot()
-		.AutoWidth()
-		.Padding(2, 0)
-		[
-			SNew(SButton)
-			.ButtonStyle(FAppStyle::Get(), TEXT("SimpleButton"))
-			.ContentPadding(0)
-			.IsEnabled(!IsPinEditingReadOnly())
-			.OnClicked(this, &FBlueprintGraphArgumentLayout::OnArgMoveDown)
-			.ToolTipText(LOCTEXT("FunctionArgDetailsArgMoveDownTooltip", "Move this parameter down in the list."))
-			[
-				SNew(SImage)
-				.Image(FEditorStyle::GetBrush("Icons.ChevronDown"))
-				.ColorAndOpacity(FSlateColor::UseForeground())
-			]
-		]
-		+ SHorizontalBox::Slot()
 		.HAlign(HAlign_Right)
 		.VAlign(VAlign_Center)
 		.Padding(10, 0, 0, 0)
@@ -2934,7 +3096,8 @@ void FBlueprintGraphArgumentLayout::GenerateHeaderRowContent( FDetailWidgetRow&
 			PropertyCustomizationHelpers::MakeClearButton(FSimpleDelegate::CreateSP(this, &FBlueprintGraphArgumentLayout::OnRemoveClicked), LOCTEXT("FunctionArgDetailsClearTooltip", "Remove this parameter."), !IsPinEditingReadOnly())
 		]
 
-	];
+	]
+	.DragDropHandler(MakeShared<FBlueprintGraphArgumentDragDropHandler>(GraphActionDetailsPtr, TargetNode, ParamItemPtr));
 }
 
 void FBlueprintGraphArgumentLayout::GenerateChildContent( IDetailChildrenBuilder& ChildrenBuilder )
@@ -3019,24 +3182,6 @@ void FBlueprintGraphArgumentLayout::GenerateChildContent( IDetailChildrenBuilder
 	
 }
 
-namespace 
-{
-	static TArray<UK2Node_EditablePinBase*> GatherAllResultNodes(UK2Node_EditablePinBase* TargetNode)
-	{
-		if (UK2Node_FunctionResult* ResultNode = Cast<UK2Node_FunctionResult>(TargetNode))
-		{
-			return (TArray<UK2Node_EditablePinBase*>)ResultNode->GetAllResultNodes();
-		}
-		TArray<UK2Node_EditablePinBase*> Result;
-		if (TargetNode)
-		{
-			Result.Add(TargetNode);
-		}
-		return Result;
-	}
-
-} // namespace
-
 void FBlueprintGraphArgumentLayout::OnRemoveClicked()
 {
 	TSharedPtr<FUserPinInfo> ParamItem = ParamItemPtr.Pin();
@@ -3059,52 +3204,6 @@ void FBlueprintGraphArgumentLayout::OnRemoveClicked()
 	}
 }
 
-FReply FBlueprintGraphArgumentLayout::OnArgMoveUp()
-{
-	const int32 ThisParamIndex = TargetNode->UserDefinedPins.Find( ParamItemPtr.Pin() );
-	const int32 NewParamIndex = ThisParamIndex-1;
-	if (ThisParamIndex != INDEX_NONE && NewParamIndex >= 0)
-	{
-		const FScopedTransaction Transaction( LOCTEXT("K2_MovePinUp", "Move Pin Up") );
-		TArray<UK2Node_EditablePinBase*> TargetNodes = GatherAllResultNodes(TargetNode);
-		for (UK2Node_EditablePinBase* Node : TargetNodes)
-		{
-			Node->Modify();
-			Node->UserDefinedPins.Swap(ThisParamIndex, NewParamIndex);
-
-			TSharedPtr<FBaseBlueprintGraphActionDetails> GraphActionDetails = GraphActionDetailsPtr.Pin();
-			if (GraphActionDetails.IsValid())
-			{
-				GraphActionDetails->OnParamsChanged(Node, true);
-			}
-		}
-	}
-	return FReply::Handled();
-}
-
-FReply FBlueprintGraphArgumentLayout::OnArgMoveDown()
-{
-	const int32 ThisParamIndex = TargetNode->UserDefinedPins.Find( ParamItemPtr.Pin() );
-	const int32 NewParamIndex = ThisParamIndex+1;
-	if (ThisParamIndex != INDEX_NONE && NewParamIndex < TargetNode->UserDefinedPins.Num())
-	{
-		const FScopedTransaction Transaction( LOCTEXT("K2_MovePinDown", "Move Pin Down") );
-		TArray<UK2Node_EditablePinBase*> TargetNodes = GatherAllResultNodes(TargetNode);
-		for (UK2Node_EditablePinBase* Node : TargetNodes)
-		{
-			Node->Modify();
-			Node->UserDefinedPins.Swap(ThisParamIndex, NewParamIndex);
-			
-			TSharedPtr<FBaseBlueprintGraphActionDetails> GraphActionDetails = GraphActionDetailsPtr.Pin();
-			if (GraphActionDetails.IsValid())
-			{
-				GraphActionDetails->OnParamsChanged(Node, true);
-			}
-		}
-	}
-	return FReply::Handled();
-}
-
 bool FBlueprintGraphArgumentLayout::ShouldPinBeReadOnly(bool bIsEditingPinType/* = false*/) const
 {
 	if (TargetNode && ParamItemPtr.IsValid())
diff --git a/Engine/Source/Editor/Kismet/Private/BlueprintDetailsCustomization.h b/Engine/Source/Editor/Kismet/Private/BlueprintDetailsCustomization.h
--- a/Engine/Source/Editor/Kismet/Private/BlueprintDetailsCustomization.h
+++ b/Engine/Source/Editor/Kismet/Private/BlueprintDetailsCustomization.h
@@ -473,8 +473,6 @@ private:
 
 	/** Callbacks for all the functionality for modifying arguments */
 	void OnRemoveClicked();
-	FReply OnArgMoveUp();
-	FReply OnArgMoveDown();
 
 	FText OnGetArgNameText() const;
 	FText OnGetArgToolTipText() const;
