Mark macros being pasted into Macro Libraries as public

diff --git a/Engine/Source/Editor/Kismet/Private/BPGraphClipboardData.cpp b/Engine/Source/Editor/Kismet/Private/BPGraphClipboardData.cpp
--- a/Engine/Source/Editor/Kismet/Private/BPGraphClipboardData.cpp
+++ b/Engine/Source/Editor/Kismet/Private/BPGraphClipboardData.cpp
@@ -120,6 +120,12 @@ UEdGraph* FBPGraphClipboardData::CreateAndPopulateGraph(UBlueprint* InBlueprint,
 						break;
 					}
 				}
+
+				// Mark the macro as public if it will be called from external objects
+				if (InBlueprint->BlueprintType == BPTYPE_MacroLibrary)
+				{
+					Graph->SetFlags(RF_Public);
+				}
 			}
 
 			FBlueprintEditorUtils::MarkBlueprintAsStructurallyModified(InBlueprint);
