Show GrabHand cursor for drag-and-drop handle in details rows

This is intended to make the drag behavior more noticeable and a bit more consistent with some other places in the editor that allow drag-and-drop, e.g., the Create menu
Also update the mouse cursor in various drag ops to use GrabHandClosed while dragging

diff --git a/Engine/Source/Editor/Kismet/Private/BlueprintDetailsCustomization.cpp b/Engine/Source/Editor/Kismet/Private/BlueprintDetailsCustomization.cpp
--- a/Engine/Source/Editor/Kismet/Private/BlueprintDetailsCustomization.cpp
+++ b/Engine/Source/Editor/Kismet/Private/BlueprintDetailsCustomization.cpp
@@ -2781,6 +2781,7 @@ public:
 		: TargetNode(InTargetNode)
 		, ParamItemPtr(InParamItemPtr)
 	{
+		MouseCursor = EMouseCursor::GrabHandClosed;
 	}
 
 	void Init()
diff --git a/Engine/Source/Editor/Kismet/Private/UserDefinedEnumEditor.cpp b/Engine/Source/Editor/Kismet/Private/UserDefinedEnumEditor.cpp
--- a/Engine/Source/Editor/Kismet/Private/UserDefinedEnumEditor.cpp
+++ b/Engine/Source/Editor/Kismet/Private/UserDefinedEnumEditor.cpp
@@ -240,6 +240,7 @@ public:
 		check(InEnumeratorIndex >= 0 && InEnumeratorIndex < InTargetEnum->NumEnums());
 
 		EnumDisplayText = InTargetEnum->GetDisplayNameTextByIndex(InEnumeratorIndex);
+		MouseCursor = EMouseCursor::GrabHandClosed;
 	}
 
 	void Init()
diff --git a/Engine/Source/Editor/Kismet/Private/UserDefinedStructureEditor.cpp b/Engine/Source/Editor/Kismet/Private/UserDefinedStructureEditor.cpp
--- a/Engine/Source/Editor/Kismet/Private/UserDefinedStructureEditor.cpp
+++ b/Engine/Source/Editor/Kismet/Private/UserDefinedStructureEditor.cpp
@@ -616,6 +616,7 @@ public:
 		: StructureDetails(InStructureDetails)
 		, FieldGuid(InFieldGuid)
 	{
+		MouseCursor = EMouseCursor::GrabHandClosed;
 		if (TSharedPtr<FUserDefinedStructureDetails> StructureDetailsSP = InStructureDetails.Pin())
 		{
 			VariableFriendlyName = FStructureEditorUtils::GetVariableFriendlyName(StructureDetailsSP->GetUserDefinedStruct(), FieldGuid);
diff --git a/Engine/Source/Editor/PropertyEditor/Private/PropertyEditorHelpers.cpp b/Engine/Source/Editor/PropertyEditor/Private/PropertyEditorHelpers.cpp
--- a/Engine/Source/Editor/PropertyEditor/Private/PropertyEditorHelpers.cpp
+++ b/Engine/Source/Editor/PropertyEditor/Private/PropertyEditorHelpers.cpp
@@ -832,6 +832,7 @@ namespace PropertyEditorHelpers
 				]
 			]
 		.ParentRow(InParentRow)
+		.Cursor(EMouseCursor::GrabHand)
 		.IsEnabled(InEnabledAttr)
 		.Visibility_Lambda([InParentRow]() { return InParentRow->IsHovered() ? EVisibility::Visible : EVisibility::Hidden; });
 		return Handle;
diff --git a/Engine/Source/Editor/PropertyEditor/Private/SDetailSingleItemRow.cpp b/Engine/Source/Editor/PropertyEditor/Private/SDetailSingleItemRow.cpp
--- a/Engine/Source/Editor/PropertyEditor/Private/SDetailSingleItemRow.cpp
+++ b/Engine/Source/Editor/PropertyEditor/Private/SDetailSingleItemRow.cpp
@@ -1236,6 +1236,7 @@ FArrayRowDragDropOp::FArrayRowDragDropOp(TSharedPtr<SDetailSingleItemRow> InRow)
 {
 	check(InRow.IsValid());
 	Row = InRow;
+	MouseCursor = EMouseCursor::GrabHandClosed;
 }
 
 void FArrayRowDragDropOp::Init()
