Add missing include to fix non-unity build

diff --git a/Engine/Source/Editor/Kismet/Private/UserDefinedStructureEditor.cpp b/Engine/Source/Editor/Kismet/Private/UserDefinedStructureEditor.cpp
--- a/Engine/Source/Editor/Kismet/Private/UserDefinedStructureEditor.cpp
+++ b/Engine/Source/Editor/Kismet/Private/UserDefinedStructureEditor.cpp
@@ -28,6 +28,8 @@
 #include "DragAndDrop/DecoratedDragDropOp.h"
 #include "SEditorHeaderButton.h"
 #include "Styling/ToolBarStyle.h"
+#include "Framework/MultiBox/MultiBoxBuilder.h"
+#include "Framework/MultiBox/MultiBoxExtender.h"
 
 #include "PropertyCustomizationHelpers.h"
 #include "SPinTypeSelector.h"
