Make GetTransientOuterForRename more generally available and use to fix crash when changing class if that class needs to be within

diff --git a/Engine/Source/Editor/Kismet/Private/BlueprintCompilationManager.cpp b/Engine/Source/Editor/Kismet/Private/BlueprintCompilationManager.cpp
--- a/Engine/Source/Editor/Kismet/Private/BlueprintCompilationManager.cpp
+++ b/Engine/Source/Editor/Kismet/Private/BlueprintCompilationManager.cpp
@@ -139,7 +139,6 @@ struct FBlueprintCompilationManagerImpl : public FGCObject
 	static void ReinstanceBatch(TArray<FReinstancingJob>& Reinstancers, TMap< UClass*, UClass* >& InOutOldToNewClassMap, FUObjectSerializeContext* InLoadContext);
 	static UClass* FastGenerateSkeletonClass(UBlueprint* BP, FKismetCompilerContext& CompilerContext, bool bIsSkeletonOnly, TArray<FSkeletonFixupData>& OutSkeletonFixupData);
 	static bool IsQueuedForCompilation(UBlueprint* BP);
-	static UObject* GetOuterForRename(UClass* ForClass);
 
 	// Declaration of archive to fix up bytecode references of blueprints that are actively compiled:
 	class FFixupBytecodeReferences : public FArchiveUObject
@@ -2290,7 +2289,7 @@ void FBlueprintCompilationManagerImpl::ReinstanceBatch(TArray<FReinstancingJob>&
 				UObject* OriginalOuter = Archetype->GetOuter();
 				EObjectFlags OriginalFlags = Archetype->GetFlags();
 
-				UObject* Destination = GetOuterForRename(Archetype->GetClass());
+				UObject* Destination = GetTransientOuterForRename(Archetype->GetClass());
 				Archetype->Rename(
 					nullptr,
 					// destination - this is the important part of this call. Moving the object 
@@ -2943,20 +2942,6 @@ bool FBlueprintCompilationManagerImpl::IsQueuedForCompilation(UBlueprint* BP)
 	return BP->bQueuedForCompilation;
 }
 
-UObject* FBlueprintCompilationManagerImpl::GetOuterForRename(UClass* ForClass)
-{
-	// if someone has tautologically placed themself within their own hierarchy then we'll
-	// just assume they're ok with eventually being outered to a upackage, similar UPackage
-	// is a UObject, so if someone demands that they be outered to 'a uobject' we'll 
-	// just leave them directly parented to the transient package:
-	if(ForClass->ClassWithin && ForClass->ClassWithin != ForClass && ForClass->ClassWithin != UObject::StaticClass())
-	{
-		FScopedAllowAbstractClassAllocation AllowAbstract;
-		return NewObject<UObject>( GetOuterForRename(ForClass->ClassWithin), ForClass->ClassWithin, NAME_None, RF_Transient );
-	}
-	return GetTransientPackage();
-}
-
 // FFixupBytecodeReferences Implementation:
 FBlueprintCompilationManagerImpl::FFixupBytecodeReferences::FFixupBytecodeReferences(UObject* InObject)
 {
diff --git a/Engine/Source/Editor/PropertyEditor/Private/UserInterface/PropertyEditor/SPropertyEditorEditInline.cpp b/Engine/Source/Editor/PropertyEditor/Private/UserInterface/PropertyEditor/SPropertyEditorEditInline.cpp
--- a/Engine/Source/Editor/PropertyEditor/Private/UserInterface/PropertyEditor/SPropertyEditorEditInline.cpp
+++ b/Engine/Source/Editor/PropertyEditor/Private/UserInterface/PropertyEditor/SPropertyEditorEditInline.cpp
@@ -324,7 +324,7 @@ void SPropertyEditorEditInline::OnClassPicked(UClass* InClass)
 				ConstructorHelpers::StripObjectClass(PrevPerObjectValues[Index]);
 				if (UObject* SubObject = StaticFindObject(UObject::StaticClass(), ANY_PACKAGE, *PrevPerObjectValues[Index]))
 				{
-					SubObject->Rename(*MakeUniqueObjectName(GetTransientPackage(), SubObject->GetClass()).ToString(), GetTransientPackage(), REN_DontCreateRedirectors);
+					SubObject->Rename(nullptr, GetTransientOuterForRename(SubObject->GetClass()), REN_DontCreateRedirectors);
 				}
 			}
 		}
diff --git a/Engine/Source/Runtime/CoreUObject/Private/UObject/Obj.cpp b/Engine/Source/Runtime/CoreUObject/Private/UObject/Obj.cpp
--- a/Engine/Source/Runtime/CoreUObject/Private/UObject/Obj.cpp
+++ b/Engine/Source/Runtime/CoreUObject/Private/UObject/Obj.cpp
@@ -4642,14 +4642,25 @@ void StaticExit()
 	Misc.
 -----------------------------------------------------------------------------*/
 
-//
-// Return the static transient package.
-//
 UPackage* GetTransientPackage()
 {
 	return GObjTransientPkg;
 }
 
+UObject* GetTransientOuterForRename(UClass* ForClass)
+{
+	// if someone has tautologically placed themself within their own hierarchy then we'll
+	// just assume they're ok with eventually being outered to a upackage, similar UPackage
+	// is a UObject, so if someone demands that they be outered to 'a uobject' we'll 
+	// just leave them directly parented to the transient package:
+	if (ForClass->ClassWithin && ForClass->ClassWithin != ForClass && ForClass->ClassWithin != UObject::StaticClass())
+	{
+		FScopedAllowAbstractClassAllocation AllowAbstract;
+		return NewObject<UObject>(GetTransientOuterForRename(ForClass->ClassWithin), ForClass->ClassWithin, NAME_None, RF_Transient);
+	}
+	return GetTransientPackage();
+}
+
 //keep this global to ensure that an actual write is prepared
 volatile const UObject** GUObjectAbortNullPointer = nullptr;
 
diff --git a/Engine/Source/Runtime/CoreUObject/Public/UObject/UObjectGlobals.h b/Engine/Source/Runtime/CoreUObject/Public/UObject/UObjectGlobals.h
--- a/Engine/Source/Runtime/CoreUObject/Public/UObject/UObjectGlobals.h
+++ b/Engine/Source/Runtime/CoreUObject/Public/UObject/UObjectGlobals.h
@@ -192,6 +192,9 @@ COREUOBJECT_API TArray<const TCHAR*> ParsePropertyFlags(EPropertyFlags Flags);
 /** Returns the transient top-level package, which is useful for temporarily storing objects that should never be saved */
 COREUOBJECT_API UPackage* GetTransientPackage();
 
+/** Returns an object in the transient package which respects the rules of Within */
+COREUOBJECT_API UObject* GetTransientOuterForRename(UClass* ForClass);
+
 /**
  * Gets INI file name from object's reference if it contains one. 
  *
