Register Kismet.DebuggingViewToolBar Toolbar before creating the Blueprint Debugger Window

diff --git a/Engine/Source/Editor/Kismet/Private/BlueprintDebugger.cpp b/Engine/Source/Editor/Kismet/Private/BlueprintDebugger.cpp
--- a/Engine/Source/Editor/Kismet/Private/BlueprintDebugger.cpp
+++ b/Engine/Source/Editor/Kismet/Private/BlueprintDebugger.cpp
@@ -14,6 +14,8 @@
 #include "WorkspaceMenuStructure.h"
 #include "WorkspaceMenuStructureModule.h"
 #include "BlueprintEditorTabs.h"
+#include "ToolMenus.h"
+#include "Kismet2/DebuggerCommands.h"
 
 #define LOCTEXT_NAMESPACE "BlueprintDebugger"
 
@@ -119,6 +121,9 @@ TSharedRef<SDockTab> FBlueprintDebuggerImpl::CreateBluprintDebuggerTab(const FSp
 		ensure(BlueprintDebuggerLayout.IsValid());
 	}
 
+	// Register Toolbar
+	SKismetDebuggingView::TryRegisterDebugToolbar();
+
 	const FName ExecutionFlowTabName = FName(TEXT("ExecutionFlowApp"));
 	const FName CallStackTabName = CallStackViewer::GetTabName();
 
diff --git a/Engine/Source/Editor/Kismet/Private/Debugging/SKismetDebuggingView.cpp b/Engine/Source/Editor/Kismet/Private/Debugging/SKismetDebuggingView.cpp
--- a/Engine/Source/Editor/Kismet/Private/Debugging/SKismetDebuggingView.cpp
+++ b/Engine/Source/Editor/Kismet/Private/Debugging/SKismetDebuggingView.cpp
@@ -1966,6 +1966,17 @@ FText SKismetDebuggingView::GetTabLabel() const
 		NSLOCTEXT("BlueprintExecutionFlow", "TabTitle", "Data Flow");
 }
 
+void SKismetDebuggingView::TryRegisterDebugToolbar()
+{
+	static const FName ToolbarName = "Kismet.DebuggingViewToolBar";
+	if (!UToolMenus::Get()->IsMenuRegistered(ToolbarName))
+	{
+		UToolMenu* ToolBar = UToolMenus::Get()->RegisterMenu(ToolbarName, NAME_None, EMultiBoxType::SlimHorizontalToolBar);
+		FToolMenuSection& Section = ToolBar->AddSection("Debug");
+		FPlayWorldCommands::BuildToolbar(Section);
+	}
+}
+
 FText SKismetDebuggingView::GetTopText() const
 {
 	return LOCTEXT("ShowDebugForActors", "Showing debug info for instances of the blueprint:");
@@ -2041,16 +2052,7 @@ void SKismetDebuggingView::Construct(const FArguments& InArgs)
 
 	// Build the debug toolbar
 	static const FName ToolbarName = "Kismet.DebuggingViewToolBar";
-	if (!UToolMenus::Get()->IsMenuRegistered(ToolbarName))
-	{
-		UToolMenu* ToolBar = UToolMenus::Get()->RegisterMenu(ToolbarName, NAME_None, EMultiBoxType::SlimHorizontalToolBar);
-
-		{
-			FToolMenuSection& Section = ToolBar->AddSection("Debug");
-			FPlayWorldCommands::BuildToolbar(Section);
-		}
-	}
-
+	TryRegisterDebugToolbar();
 	FToolMenuContext MenuContext(FPlayWorldCommands::GlobalPlayWorldActions);
 	TSharedRef<SWidget> ToolbarWidget = UToolMenus::Get()->GenerateWidget(ToolbarName, MenuContext);
 	
diff --git a/Engine/Source/Editor/Kismet/Private/Debugging/SKismetDebuggingView.h b/Engine/Source/Editor/Kismet/Private/Debugging/SKismetDebuggingView.h
--- a/Engine/Source/Editor/Kismet/Private/Debugging/SKismetDebuggingView.h
+++ b/Engine/Source/Editor/Kismet/Private/Debugging/SKismetDebuggingView.h
@@ -162,6 +162,9 @@ public:
 	static TWeakObjectPtr<const UObject> CurrentActiveObject;
 	
 	FText GetTabLabel() const;
+
+	/** Registers the Kismet.DebuggingViewToolBar if it hasn't already been registered */
+	static void TryRegisterDebugToolbar();
 protected:
 	FText GetTopText() const;
 	bool CanDisableAllBreakpoints() const;
