Add menu bar to Blueprint Debugger for opening closed tabs
(Visual style is consistent with Visual Logger)

diff --git a/Engine/Source/Editor/Kismet/Private/BlueprintDebugger.cpp b/Engine/Source/Editor/Kismet/Private/BlueprintDebugger.cpp
--- a/Engine/Source/Editor/Kismet/Private/BlueprintDebugger.cpp
+++ b/Engine/Source/Editor/Kismet/Private/BlueprintDebugger.cpp
@@ -295,12 +295,31 @@ TSharedRef<SDockTab> FBlueprintDebuggerImpl::CreateBluprintDebuggerTab(const FSp
 		)
 	);
 
+	FMenuBarBuilder MenuBarBuilder(CommandList);
+	MenuBarBuilder.AddPullDownMenu(
+		LOCTEXT("WindowMenuLabel", "Window"),
+		FText::GetEmpty(),
+		FNewMenuDelegate::CreateLambda([](FMenuBuilder& Builder) {
+			Builder.AddMenuEntry(FBlueprintDebuggerCommands::Get().ShowCallStackViewer);
+			Builder.AddMenuEntry(FBlueprintDebuggerCommands::Get().ShowExecutionTrace);
+			})
+	);
+
 	NomadTab->SetContent(
-		SNew(SBorder)
-		.BorderImage(FEditorStyle::GetBrush("ToolPanel.DarkGroupBorder"))
-		.Padding(FMargin(0.f, 2.f))
+		SNew(SVerticalBox)
+		+ SVerticalBox::Slot()
+		.AutoHeight()
+		[
+			MenuBarBuilder.MakeWidget()
+		]
+		+SVerticalBox::Slot()
 		[
-			TabContents
+			SNew(SBorder)
+			.BorderImage(FEditorStyle::GetBrush("ToolPanel.DarkGroupBorder"))
+			.Padding(FMargin(0.f, 2.f))
+			[
+				TabContents
+			]
 		]
 	);
 
