Allow Blueprint Debugger search field to search for actual variable names as well as display names

diff --git a/Engine/Source/Editor/Kismet/Private/Debugging/SKismetDebuggingView.cpp b/Engine/Source/Editor/Kismet/Private/Debugging/SKismetDebuggingView.cpp
--- a/Engine/Source/Editor/Kismet/Private/Debugging/SKismetDebuggingView.cpp
+++ b/Engine/Source/Editor/Kismet/Private/Debugging/SKismetDebuggingView.cpp
@@ -60,6 +60,11 @@ namespace KismetDebugViewConstants
 
 uint16 FDebugLineItem::ActiveTypeBitset = TNumericLimits<uint16>::Max(); // set all to active by default
 
+FText FDebugLineItem::GetName() const
+{
+	return FText::GetEmpty();
+}
+
 FText FDebugLineItem::GetDisplayName() const
 {
 	return FText::GetEmpty();
@@ -141,11 +146,14 @@ void FDebugLineItem::MakeMenu(FMenuBuilder& MenuBuilder)
 	);
 }
 
-void FDebugLineItem::UpdateSearchFlags(bool bIsRootNode)
+void FDebugLineItem::UpdateSearchFlags(bool bIsRootNode, bool bIsContainerElement)
 {
 	const FString SearchString = SearchBox->GetText().ToString();
 	
-	bVisible = GetDisplayName().ToString().Contains(SearchString) || GetDescription().ToString().Contains(SearchString);
+	// Container elements share their parent's property name, so we shouldn't search them by name
+	bVisible = (!bIsContainerElement && GetName().ToString().Contains(SearchString)) ||
+		GetDisplayName().ToString().Contains(SearchString) || 
+		GetDescription().ToString().Contains(SearchString);
 
 	// for root nodes, bParentsMatchSearch always matches bVisible
 	if(bVisible || bIsRootNode)
@@ -276,12 +284,13 @@ public:
 	* O( number of recursive children )
 	*/
 	bool SearchRecursive(TArray<FLineItemWithChildren*>& Parents,
-		TSharedPtr<STreeView<FDebugTreeItemPtr>> DebugTreeView)
+		TSharedPtr<STreeView<FDebugTreeItemPtr>> DebugTreeView,
+		bool bIsContainerElement = false)
 	{
 		TSharedPtr<ITableRow> Row = DebugTreeView->WidgetFromItem(SharedThis(this));
 		bVisible = false;
 
-		UpdateSearchFlags();
+		UpdateSearchFlags(/*bIsRootNode =*/ false, bIsContainerElement);
 
 		bool bChildMatch = false;
 		Parents.Push(this);
@@ -308,7 +317,7 @@ public:
                 }
 
 				// if any children need to expand, so should this
-				if(Child->SearchRecursive(Parents, DebugTreeView))
+				if(Child->SearchRecursive(Parents, DebugTreeView, IsContainer()))
 				{
 					bVisible = true;
 					bChildMatch = true;
@@ -323,7 +332,7 @@ public:
 			}
 			else
 			{
-				ChildRef->UpdateSearchFlags();
+				ChildRef->UpdateSearchFlags(/*bIsRootNode =*/ false, IsContainer());
 				
 				// if any children need to expand, so should this
 				if(ChildRef->IsVisible())
@@ -408,6 +417,12 @@ protected:
 	// This frames children
 	TSet<FDebugTreeItemPtr, FDebugTreeItemKeyFuncs> ChildrenMirrors;
 
+	/** @returns whether this item represents a container property */
+	virtual bool IsContainer() const 
+	{
+		return false;
+	}
+
 	virtual void GatherChildren(TArray<FDebugTreeItemPtr>& OutChildren, bool bRespectSearch) {}
 	
 	/**
@@ -627,6 +642,11 @@ public:
 	{
 		return HashCombine(GetTypeHash(Data.Property), GetTypeHash(Data.DisplayName.ToString()));
 	}
+
+	virtual FText GetName() const override
+	{
+		return Data.Name;
+	}
 	
 	virtual FText GetDescription() const override
 	{
@@ -755,6 +775,12 @@ public:
             EnsureChildIsAdded(OutChildren, FWatchChildLineItem(*ChildData, SearchBox), bRespectSearch);
         }
 	}
+
+protected:
+	virtual bool IsContainer() const override
+	{
+		return Data.Property->IsA<FSetProperty>() || Data.Property->IsA<FArrayProperty>() || Data.Property->IsA<FMapProperty>();
+	}
 };
 
 //////////////////////////////////////////////////////////////////////////\
diff --git a/Engine/Source/Editor/Kismet/Private/Debugging/SKismetDebuggingView.h b/Engine/Source/Editor/Kismet/Private/Debugging/SKismetDebuggingView.h
--- a/Engine/Source/Editor/Kismet/Private/Debugging/SKismetDebuggingView.h
+++ b/Engine/Source/Editor/Kismet/Private/Debugging/SKismetDebuggingView.h
@@ -91,7 +91,7 @@ public:
 	static void OnDebugLineTypeActiveChanged(ECheckBoxState CheckState, EDebugLineType Type);
 
 	// updates bVisible and bParentsMatchSearch based on this node alone
-	void UpdateSearchFlags(bool bIsRootNode = false);
+	void UpdateSearchFlags(bool bIsRootNode = false, bool bIsContainerElement = false);
 
 	bool IsVisible();
 	bool DoParentsMatchSearch();
@@ -117,6 +117,9 @@ protected:
 	// called after Compare returns true
 	[[maybe_unused]] virtual void UpdateData(const FDebugLineItem& NewerData) {}
 
+	// @return The actual underlying name of the item, for search compatibility
+	virtual FText GetName() const;
+
 	// @return The text to display in the name column, unless GenerateNameWidget is overridden
 	virtual FText GetDisplayName() const;
 
diff --git a/Engine/Source/Editor/UnrealEd/Private/Kismet2/KismetDebugUtilities.cpp b/Engine/Source/Editor/UnrealEd/Private/Kismet2/KismetDebugUtilities.cpp
--- a/Engine/Source/Editor/UnrealEd/Private/Kismet2/KismetDebugUtilities.cpp
+++ b/Engine/Source/Editor/UnrealEd/Private/Kismet2/KismetDebugUtilities.cpp
@@ -1797,6 +1797,7 @@ void FKismetDebugUtilities::GetDebugInfoInternal(TSharedPtr<FPropertyInstanceInf
 }
 
 FPropertyInstanceInfo::FPropertyInstanceInfo(FPropertyInstance PropertyInstance) :
+	Name(FText::FromString(PropertyInstance.Property->GetName())),
 	DisplayName(PropertyInstance.Property->GetDisplayNameText()),
 	Type(UEdGraphSchema_K2::TypeToText(PropertyInstance.Property)),
 	Property(PropertyInstance.Property)
diff --git a/Engine/Source/Editor/UnrealEd/Public/Kismet2/KismetDebugUtilities.h b/Engine/Source/Editor/UnrealEd/Public/Kismet2/KismetDebugUtilities.h
--- a/Engine/Source/Editor/UnrealEd/Public/Kismet2/KismetDebugUtilities.h
+++ b/Engine/Source/Editor/UnrealEd/Public/Kismet2/KismetDebugUtilities.h
@@ -78,6 +78,7 @@ struct FPropertyInstanceInfo
 	void PopulateChildren(FPropertyInstance PropertyInstance,
 		TMap<FPropertyInstance, TSharedPtr<FPropertyInstanceInfo>> &VisitedNodes);
 	
+	FText Name;
 	FText DisplayName;
 	FText Value;
 	FText Type;
