Allow Debugging of Transient BP instances

diff --git a/Engine/Source/Editor/Kismet/Private/Debugging/SKismetDebuggingView.cpp b/Engine/Source/Editor/Kismet/Private/Debugging/SKismetDebuggingView.cpp
--- a/Engine/Source/Editor/Kismet/Private/Debugging/SKismetDebuggingView.cpp
+++ b/Engine/Source/Editor/Kismet/Private/Debugging/SKismetDebuggingView.cpp
@@ -2269,7 +2269,7 @@ void SKismetDebuggingView::Tick( const FGeometry& AllottedGeometry, const double
 	if(bIsDebugging && BlueprintToWatchPtr.IsValid())
 	{
 		UClass* GeneratedClass = Cast<UClass>(BlueprintToWatchPtr->GeneratedClass);
-		for(FThreadSafeObjectIterator Iter(GeneratedClass); Iter; ++Iter)
+		for(FThreadSafeObjectIterator Iter(GeneratedClass, /*bOnlyGCedObjects =*/ false, /*AdditionalExclusionFlags =*/ RF_ArchetypeObject | RF_ClassDefaultObject); Iter; ++Iter)
         {
             UObject* Instance = *Iter;
 			if(!Instance)
@@ -2277,12 +2277,6 @@ void SKismetDebuggingView::Tick( const FGeometry& AllottedGeometry, const double
 				continue;
 			}
 
-			// only include non temporary, non archetype objects
-            if(Instance->HasAnyFlags(RF_ArchetypeObject | RF_Transient))
-            {
-                continue;
-            }
-			
 			// only include actors in current world
             if(AActor* Actor = Cast<AActor>(Instance))
             {
