Fixing blueprint variable categories not being set correctly in details panel

diff --git a/Engine/Source/Editor/Kismet/Private/BlueprintDetailsCustomization.cpp b/Engine/Source/Editor/Kismet/Private/BlueprintDetailsCustomization.cpp
--- a/Engine/Source/Editor/Kismet/Private/BlueprintDetailsCustomization.cpp
+++ b/Engine/Source/Editor/Kismet/Private/BlueprintDetailsCustomization.cpp
@@ -558,13 +558,12 @@ void FBlueprintVarActionDetails::CustomizeDetails( IDetailLayoutBuilder& DetailL
 					.OnTextCommitted(this, &FBlueprintVarActionDetails::OnCategoryTextCommitted, CachedVariableName )
 					.OnVerifyTextChanged_Lambda([&](const FText& InNewText, FText& OutErrorMessage) -> bool
 					{
-						const FText NewText = FEditorCategoryUtils::GetCategoryDisplayString(InNewText);
-						if (NewText.IsEmpty())
+						if (InNewText.IsEmpty())
 						{
 							OutErrorMessage = LOCTEXT("CategoryEmpty", "Cannot add a category with an empty string.");
 							return false;
 						}
-						if (NewText.EqualTo(FText::FromString(GetBlueprintObj()->GetName())))
+						if (InNewText.EqualTo(FText::FromString(GetBlueprintObj()->GetName())))
 						{
 							OutErrorMessage = LOCTEXT("CategoryEqualsBlueprintName", "Cannot add a category with the same name as the blueprint.");
 							return false;
@@ -1475,14 +1474,16 @@ void FBlueprintVarActionDetails::PopulateCategories(SMyBlueprint* MyBlueprint, T
 		FText Category = FBlueprintEditorUtils::GetBlueprintVariableCategory(Blueprint, VisibleVariables[i], nullptr);
 		if (!Category.IsEmpty() && !Category.EqualTo(FText::FromString(Blueprint->GetName())))
 		{
+			const FText DisplayCategory = FEditorCategoryUtils::GetCategoryDisplayString(Category);
+			
 			bool bNewCategory = true;
 			for (int32 j = 0; j < CategorySource.Num() && bNewCategory; ++j)
 			{
-				bNewCategory &= !CategorySource[j].Get()->EqualTo(Category);
+				bNewCategory &= !CategorySource[j].Get()->EqualTo(DisplayCategory);
 			}
 			if (bNewCategory)
 			{
-				CategorySource.Add(MakeShareable(new FText(Category)));
+				CategorySource.Add(MakeShareable(new FText(DisplayCategory)));
 			}
 		}
 	}
@@ -1496,15 +1497,17 @@ void FBlueprintVarActionDetails::PopulateCategories(SMyBlueprint* MyBlueprint, T
 
 			if(!FunctionCategory.IsEmpty())
 			{
+				const FText DisplayCategory = FEditorCategoryUtils::GetCategoryDisplayString(FunctionCategory);
+				
 				bool bNewCategory = true;
 				for (int32 j = 0; j < CategorySource.Num() && bNewCategory; ++j)
 				{
-					bNewCategory &= !CategorySource[j].Get()->EqualTo(FunctionCategory);
+					bNewCategory &= !CategorySource[j].Get()->EqualTo(DisplayCategory);
 				}
 
 				if(bNewCategory)
 				{
-					CategorySource.Add(MakeShareable(new FText(FunctionCategory)));
+					CategorySource.Add(MakeShareable(new FText(DisplayCategory)));
 				}
 			}
 		}
@@ -1514,14 +1517,16 @@ void FBlueprintVarActionDetails::PopulateCategories(SMyBlueprint* MyBlueprint, T
 		{
 			for (FBPVariableDescription& Variable : FunctionEntryNode->LocalVariables)
 			{
+				const FText DisplayCategory = FEditorCategoryUtils::GetCategoryDisplayString(Variable.Category);
+				
 				bool bNewCategory = true;
 				for (int32 j = 0; j < CategorySource.Num() && bNewCategory; ++j)
 				{
-					bNewCategory &= !CategorySource[j].Get()->EqualTo(Variable.Category);
+					bNewCategory &= !CategorySource[j].Get()->EqualTo(DisplayCategory);
 				}
 				if (bNewCategory)
 				{
-					CategorySource.Add(MakeShareable(new FText(Variable.Category)));
+					CategorySource.Add(MakeShareable(new FText(DisplayCategory)));
 				}
 			}
 		}
@@ -1534,14 +1539,16 @@ void FBlueprintVarActionDetails::PopulateCategories(SMyBlueprint* MyBlueprint, T
 		{
 			if (!TypedEntryNode->MetaData.Category.IsEmpty())
 			{
+				const FText DisplayCategory = FEditorCategoryUtils::GetCategoryDisplayString(TypedEntryNode->MetaData.Category);
+				
 				bool bNewCategory = true;
 				for (int32 j = 0; j < CategorySource.Num() && bNewCategory; ++j)
 				{
-					bNewCategory &= !CategorySource[j].Get()->EqualTo(TypedEntryNode->MetaData.Category);
+					bNewCategory &= !CategorySource[j].Get()->EqualTo(DisplayCategory);
 				}
 				if (bNewCategory)
 				{
-					CategorySource.Add(MakeShareable(new FText(TypedEntryNode->MetaData.Category)));
+					CategorySource.Add(MakeShareable(new FText(DisplayCategory)));
 				}
 			}
 		}
@@ -1559,15 +1566,17 @@ void FBlueprintVarActionDetails::PopulateCategories(SMyBlueprint* MyBlueprint, T
 
 			if (!FunctionCategory.IsEmpty())
 			{
+				const FText DisplayCategory = FEditorCategoryUtils::GetCategoryDisplayString(FunctionCategory);
+			
 				bool bNewCategory = true;
 				for (int32 j = 0; j < CategorySource.Num() && bNewCategory; ++j)
 				{
-					bNewCategory &= !CategorySource[j].Get()->EqualTo(FunctionCategory);
+					bNewCategory &= !CategorySource[j].Get()->EqualTo(DisplayCategory);
 				}
 
 				if (bNewCategory)
 				{
-					CategorySource.Add(MakeShareable(new FText(FunctionCategory)));
+					CategorySource.Add(MakeShareable(new FText(DisplayCategory)));
 				}
 			}
 		}
@@ -1680,13 +1689,10 @@ void FBlueprintVarActionDetails::OnCategoryTextCommitted(const FText& NewText, E
 {
 	if (InTextCommit == ETextCommit::OnEnter || InTextCommit == ETextCommit::OnUserMovedFocus)
 	{
-		// Sanitize category name
-		FText CategoryName = FEditorCategoryUtils::GetCategoryDisplayString(NewText);
-
-		FBlueprintEditorUtils::SetBlueprintVariableCategory(GetBlueprintObj(), VarName, GetLocalVariableScope(CachedVariableProperty.Get()), CategoryName);
+		FBlueprintEditorUtils::SetBlueprintVariableCategory(GetBlueprintObj(), VarName, GetLocalVariableScope(CachedVariableProperty.Get()), NewText);
 		check(MyBlueprint.IsValid());
 		PopulateCategories(MyBlueprint.Pin().Get(), CategorySource);
-		MyBlueprint.Pin()->ExpandCategory(CategoryName);
+		MyBlueprint.Pin()->ExpandCategory(NewText);
 	}
 }
 
@@ -6187,13 +6193,12 @@ void FBlueprintComponentDetails::CustomizeDetails(IDetailLayoutBuilder& DetailLa
 					.OnTextCommitted(this, &FBlueprintComponentDetails::OnVariableCategoryTextCommitted, CachedNodePtr->GetVariableName())
 					.OnVerifyTextChanged_Lambda([&](const FText& InNewText, FText& OutErrorMessage) -> bool
 					{
-						const FText NewText = FEditorCategoryUtils::GetCategoryDisplayString(InNewText);
-						if (NewText.IsEmpty())
+						if (InNewText.IsEmpty())
 						{
 							OutErrorMessage = LOCTEXT("CategoryEmpty", "Cannot add a category with an empty string.");
 							return false;
 						}
-						if (NewText.EqualTo(FText::FromString(GetBlueprintObj()->GetName())))
+						if (InNewText.EqualTo(FText::FromString(GetBlueprintObj()->GetName())))
 						{
 							OutErrorMessage = LOCTEXT("CategoryEqualsBlueprintName", "Cannot add a category with the same name as the blueprint.");
 							return false;
@@ -6508,14 +6513,16 @@ void FBlueprintComponentDetails::PopulateVariableCategories()
 		FText Category = FBlueprintEditorUtils::GetBlueprintVariableCategory(BlueprintObj, VariableName, nullptr);
 		if (!Category.IsEmpty() && !Category.EqualTo(FText::FromString(BlueprintObj->GetName())))
 		{
+			const FText DisplayCategory = FEditorCategoryUtils::GetCategoryDisplayString(Category);
+			
 			bool bNewCategory = true;
 			for (int32 j = 0; j < VariableCategorySource.Num() && bNewCategory; ++j)
 			{
-				bNewCategory &= !VariableCategorySource[j].Get()->EqualTo(Category);
+				bNewCategory &= !VariableCategorySource[j].Get()->EqualTo(DisplayCategory);
 			}
 			if (bNewCategory)
 			{
-				VariableCategorySource.Add(MakeShareable(new FText(Category)));
+				VariableCategorySource.Add(MakeShareable(new FText(DisplayCategory)));
 			}
 		}
 	}
