Control Rig: Refresh local variable details panel when local variable type changes in MyBlueprint

diff --git a/Engine/Source/Editor/Kismet/Private/SBlueprintPalette.cpp b/Engine/Source/Editor/Kismet/Private/SBlueprintPalette.cpp
--- a/Engine/Source/Editor/Kismet/Private/SBlueprintPalette.cpp
+++ b/Engine/Source/Editor/Kismet/Private/SBlueprintPalette.cpp
@@ -714,12 +714,14 @@ public:
 /*******************************************************************************
 * SPinTypeSelectorHelper
 *******************************************************************************/
+DECLARE_DELEGATE_OneParam(FOnPinTypeChanged, const FEdGraphPinType&)
 
 class SPinTypeSelectorHelper : public SCompoundWidget
 {
 public:
 	SLATE_BEGIN_ARGS( SPinTypeSelectorHelper ) {}
 		SLATE_ATTRIBUTE(bool, ReadOnly)
+		SLATE_EVENT(FOnPinTypeChanged, OnTypeChanged)
 	SLATE_END_ARGS()
 
 	/**
@@ -754,6 +756,8 @@ private:
 
 	void ConstructInternal(const FArguments& InArgs)
 	{
+		OnTypeChanged = InArgs._OnTypeChanged;
+		
 		TSharedPtr<IPinTypeSelectorFilter> CustomPinTypeFilter;
 		if (BlueprintEditorPtr.IsValid())
 		{
@@ -830,6 +834,11 @@ private:
 				ActionPtr.Pin()->ChangeVariableType(InNewPinType);
 			}
 		}
+
+		if (OnTypeChanged.IsBound())
+		{
+			OnTypeChanged.Execute(InNewPinType);
+		}
 	}
 
 private:
@@ -844,6 +853,9 @@ private:
 
 	/** Variable Property to change the type of */
 	TWeakFieldPtr<FProperty> VariableProperty;
+
+	/** Event when type has changed */
+	FOnPinTypeChanged OnTypeChanged;
 };
 
 /*******************************************************************************
@@ -1171,7 +1183,8 @@ void SBlueprintPaletteItem::Construct(const FArguments& InArgs, FCreateWidgetFor
 		const UEdGraphSchema_K2* Schema = GetDefault<UEdGraphSchema_K2>();
 		IconWidget = SNew(SPinTypeSelectorHelper, StaticCastSharedPtr<FEdGraphSchemaAction_BlueprintVariableBase>(GraphAction), Blueprint, BlueprintEditorPtr)
 			.IsEnabled(bIsEditingEnabled)
-			.ReadOnly_Lambda([this]() {return !IsHovered(); });
+			.ReadOnly_Lambda([this]() {return !IsHovered(); })
+			.OnTypeChanged_Lambda([this](const FEdGraphPinType& NewPinType) { BlueprintEditorPtr.Pin()->GetMyBlueprintWidget()->Refresh(); });
 	}
 	
 	// Determine the access level of this action if it is a function graph or for interface events
