Blueprints: During reinstancing, enable delta serialization if we are reparenting an object to enable transient properties to be copied across to the new CDO. Adds new Blueprint Compiler options and internal reinstancer flags to allow users to enable delta serialization.

diff --git a/Engine/Source/Editor/Kismet/Private/BlueprintCompilationManager.cpp b/Engine/Source/Editor/Kismet/Private/BlueprintCompilationManager.cpp
--- a/Engine/Source/Editor/Kismet/Private/BlueprintCompilationManager.cpp
+++ b/Engine/Source/Editor/Kismet/Private/BlueprintCompilationManager.cpp
@@ -422,6 +422,7 @@ struct FCompilerData
 		InternalOptions.bSaveIntermediateProducts = (UserOptions & EBlueprintCompileOptions::SaveIntermediateProducts) != EBlueprintCompileOptions::None;
 		InternalOptions.bSkipDefaultObjectValidation = (UserOptions & EBlueprintCompileOptions::SkipDefaultObjectValidation) != EBlueprintCompileOptions::None;
 		InternalOptions.bSkipFiBSearchMetaUpdate = (UserOptions & EBlueprintCompileOptions::SkipFiBSearchMetaUpdate) != EBlueprintCompileOptions::None;
+		InternalOptions.bUseDeltaSerializationDuringReinstancing = (UserOptions & EBlueprintCompileOptions::UseDeltaSerializationDuringReinstancing) != EBlueprintCompileOptions::None;
 		InternalOptions.CompileType = bBytecodeOnly ? EKismetCompileType::BytecodeOnly : EKismetCompileType::Full;
 
 		if(!bBytecodeOnly && CPPResults.CppOptions)
@@ -451,6 +452,7 @@ struct FCompilerData
 	bool ShouldSkipIfDependenciesAreUnchanged() const { return InternalOptions.CompileType == EKismetCompileType::BytecodeOnly || JobType == ECompilationManagerJobType::RelinkOnly; }
 	bool ShouldValidateClassDefaultObject() const { return JobType == ECompilationManagerJobType::Normal && !InternalOptions.bSkipDefaultObjectValidation; }
 	bool ShouldUpdateBlueprintSearchMetadata() const { return JobType == ECompilationManagerJobType::Normal && !InternalOptions.bSkipFiBSearchMetaUpdate; }
+	bool UseDeltaSerializationDuringReinstancing() const { return InternalOptions.bUseDeltaSerializationDuringReinstancing; }
 
 	UBlueprint* BP;
 	FCompilerResultsLog* ActiveResultsLog;
@@ -1210,6 +1212,11 @@ void FBlueprintCompilationManagerImpl::FlushCompilationQueueImpl(bool bSuppressB
 					CompileReinstancerFlags |= EBlueprintCompileReinstancerFlags::AvoidCDODuplication;
 				}
 
+				if (CompilerData.UseDeltaSerializationDuringReinstancing())
+				{
+					CompileReinstancerFlags |= EBlueprintCompileReinstancerFlags::UseDeltaSerialization;
+				}
+
 				CompilerData.Reinstancer = TSharedPtr<FBlueprintCompileReinstancer>(
 					new FBlueprintCompileReinstancer(
 						BP->GeneratedClass,
@@ -2078,8 +2085,9 @@ void FBlueprintCompilationManagerImpl::ReinstanceBatch(TArray<FReinstancingJob>&
 			OldCDO = ReinstancingJob.OldToNew.Key->ClassDefaultObject;
 			if (OldCDO && ReinstancingJob.Reinstancer.IsValid())
 			{
+				const bool bUseDeltaSerialization = ReinstancingJob.Reinstancer.IsValid() ? ReinstancingJob.Reinstancer->bUseDeltaSerializationToCopyProperties : false;
 				UObject* NewCDO = ReinstancingJob.OldToNew.Value->GetDefaultObject(true);
-				FBlueprintCompileReinstancer::CopyPropertiesForUnrelatedObjects(OldCDO, NewCDO, true);
+				FBlueprintCompileReinstancer::CopyPropertiesForUnrelatedObjects(OldCDO, NewCDO, true, bUseDeltaSerialization);
 
 				if (ReinstancingJob.Compiler.IsValid())
 				{
@@ -2276,8 +2284,10 @@ void FBlueprintCompilationManagerImpl::ReinstanceBatch(TArray<FReinstancingJob>&
 				{
 					// The new object hierarchy has been created, all of the old instances are in the transient package and new
 					// ones have taken their place. Referenc members will mostly be pointing at *old* instances, and will get fixed
-					// up below:
-					FBlueprintCompileReinstancer::CopyPropertiesForUnrelatedObjects(OldInstance, *NewInstance, false);
+					// up below:			
+					const bool bUseDeltaSerialization = ReinstancingJob.Reinstancer.IsValid() ? ReinstancingJob.Reinstancer->bUseDeltaSerializationToCopyProperties : false;
+
+					FBlueprintCompileReinstancer::CopyPropertiesForUnrelatedObjects(OldInstance, *NewInstance, false, bUseDeltaSerialization);
 				}
 			}
 		}
diff --git a/Engine/Source/Editor/Kismet/Private/BlueprintEditor.cpp b/Engine/Source/Editor/Kismet/Private/BlueprintEditor.cpp
--- a/Engine/Source/Editor/Kismet/Private/BlueprintEditor.cpp
+++ b/Engine/Source/Editor/Kismet/Private/BlueprintEditor.cpp
@@ -1460,6 +1460,7 @@ void FBlueprintEditor::OnChangeBreadCrumbGraph(UEdGraph* InGraph)
 
 FBlueprintEditor::FBlueprintEditor()
 	: bSaveIntermediateBuildProducts(false)
+	, bIsReparentingBlueprint(false)
 	, bPendingDeferredClose(false)
 	, bRequestedSavingOpenDocumentState(false)
 	, bBlueprintModifiedOnOpen (false)
@@ -2943,6 +2944,9 @@ void FBlueprintEditor::ReparentBlueprint_NewParentChosen(UClass* ChosenClass)
 
 		if ( bReparent )
 		{
+			// Notify that we are currently reparenting this blueprint so that we get the proper compilation flags
+			TGuardValue<bool> GuardValue(bIsReparentingBlueprint, true);
+
 			const FScopedTransaction Transaction( LOCTEXT("ReparentBlueprint", "Reparent Blueprint") );
 			UE_LOG(LogBlueprint, Warning, TEXT("Reparenting blueprint %s from %s to %s..."), *BlueprintObj->GetFullName(), BlueprintObj->ParentClass ? *BlueprintObj->ParentClass->GetName() : TEXT("[None]"), *ChosenClass->GetName());
 			
@@ -3510,6 +3514,12 @@ void FBlueprintEditor::Compile()
 		{
 			CompileOptions |= EBlueprintCompileOptions::SaveIntermediateProducts;
 		}
+
+		if (bIsReparentingBlueprint)
+		{
+			CompileOptions |= EBlueprintCompileOptions::UseDeltaSerializationDuringReinstancing;
+		}
+
 		FKismetEditorUtilities::CompileBlueprint(BlueprintObj, CompileOptions, &LogResults);
 
 		LogResults.EndEvent();
diff --git a/Engine/Source/Editor/Kismet/Public/BlueprintEditor.h b/Engine/Source/Editor/Kismet/Public/BlueprintEditor.h
--- a/Engine/Source/Editor/Kismet/Public/BlueprintEditor.h
+++ b/Engine/Source/Editor/Kismet/Public/BlueprintEditor.h
@@ -1219,16 +1219,19 @@ public://@TODO
 
 protected:
 
-	// Should intermediate build products be saved when recompiling?
+	/** Should intermediate build products be saved when recompiling? */
 	bool bSaveIntermediateBuildProducts;
 
+	/** True if the current blueprint is in the process of being reparented */
+	bool bIsReparentingBlueprint;
+
 	/** Flags if this blueprint editor should close on its next tick. */
 	bool bPendingDeferredClose;
 
 	/** Currently focused graph editor */
 	TWeakPtr<class SGraphEditor> FocusedGraphEdPtr;
 	
-	// Factory that spawns graph editors; used to look up all tabs spawned by it.
+	/** Factory that spawns graph editors; used to look up all tabs spawned by it. */
 	TWeakPtr<FDocumentTabFactory> GraphEditorTabFactoryPtr;
 
 	/** User-defined enumerators to keep loaded */
diff --git a/Engine/Source/Editor/UnrealEd/Private/Kismet2/KismetReinstanceUtilities.cpp b/Engine/Source/Editor/UnrealEd/Private/Kismet2/KismetReinstanceUtilities.cpp
--- a/Engine/Source/Editor/UnrealEd/Private/Kismet2/KismetReinstanceUtilities.cpp
+++ b/Engine/Source/Editor/UnrealEd/Private/Kismet2/KismetReinstanceUtilities.cpp
@@ -369,6 +369,7 @@ FBlueprintCompileReinstancer::FBlueprintCompileReinstancer(UClass* InClassToRein
 			ReinstClassType = RCT_Native;
 		}
 		bAllowResaveAtTheEndIfRequested = bAutoInferSaveOnCompile && !bIsBytecodeOnly && (ReinstClassType != RCT_BpSkeleton);
+		bUseDeltaSerializationToCopyProperties = !!(Flags & EBlueprintCompileReinstancerFlags::UseDeltaSerialization);
 
 		SaveClassFieldMapping(InClassToReinstance);
 
@@ -2717,14 +2718,15 @@ void FBlueprintCompileReinstancer::ReparentChild(UClass* ChildClass)
 	ChildClass->StaticLink(true);
 }
 
-void FBlueprintCompileReinstancer::CopyPropertiesForUnrelatedObjects(UObject* OldObject, UObject* NewObject, bool bClearExternalReferences)
+void FBlueprintCompileReinstancer::CopyPropertiesForUnrelatedObjects(UObject* OldObject, UObject* NewObject, bool bClearExternalReferences, bool bForceDeltaSerialization /* = false */)
 {
 	InstancedPropertyUtils::FInstancedPropertyMap InstancedPropertyMap;
 	InstancedPropertyUtils::FArchiveInstancedSubObjCollector  InstancedSubObjCollector(OldObject, InstancedPropertyMap);
 
 	UEngine::FCopyPropertiesForUnrelatedObjectsParams Params;
-	Params.bAggressiveDefaultSubobjectReplacement = false;//true;
-	Params.bDoDelta = !OldObject->HasAnyFlags(RF_ClassDefaultObject);
+	Params.bAggressiveDefaultSubobjectReplacement = false;
+	// During a blueprint reparent, delta serialization must be enabled to correctly copy all properties
+	Params.bDoDelta = bForceDeltaSerialization || !OldObject->HasAnyFlags(RF_ClassDefaultObject);
 	Params.bCopyDeprecatedProperties = true;
 	Params.bSkipCompilerGeneratedDefaults = true;
 	Params.bClearReferences = bClearExternalReferences;
diff --git a/Engine/Source/Editor/UnrealEd/Public/Kismet2/KismetEditorUtilities.h b/Engine/Source/Editor/UnrealEd/Public/Kismet2/KismetEditorUtilities.h
--- a/Engine/Source/Editor/UnrealEd/Public/Kismet2/KismetEditorUtilities.h
+++ b/Engine/Source/Editor/UnrealEd/Public/Kismet2/KismetEditorUtilities.h
@@ -54,6 +54,8 @@ enum class EBlueprintCompileOptions
 	SkipDefaultObjectValidation = 0x100,
 	/** Skips Find-in-Blueprint search data update - in some cases (e.g. creating new assets) this is being deferred until after compilation */
 	SkipFiBSearchMetaUpdate = 0x200,
+	/** Allow the delta serialization during FBlueprintCompileReinstancer::CopyPropertiesForUnrelatedObjects */
+	UseDeltaSerializationDuringReinstancing = 0x400,
 };
 
 ENUM_CLASS_FLAGS(EBlueprintCompileOptions)
diff --git a/Engine/Source/Editor/UnrealEd/Public/Kismet2/KismetReinstanceUtilities.h b/Engine/Source/Editor/UnrealEd/Public/Kismet2/KismetReinstanceUtilities.h
--- a/Engine/Source/Editor/UnrealEd/Public/Kismet2/KismetReinstanceUtilities.h
+++ b/Engine/Source/Editor/UnrealEd/Public/Kismet2/KismetReinstanceUtilities.h
@@ -19,7 +19,8 @@ enum class EBlueprintCompileReinstancerFlags
 
 	BytecodeOnly			= 0x1,
 	AutoInferSaveOnCompile	= 0x2,
-	AvoidCDODuplication		= 0x4
+	AvoidCDODuplication		= 0x4,
+	UseDeltaSerialization	= 0x8,
 };
 
 ENUM_CLASS_FLAGS(EBlueprintCompileReinstancerFlags)
@@ -131,6 +132,9 @@ protected:
 	/** TRUE if this reinstancer should resave compiled Blueprints if the user has requested it */
 	bool bAllowResaveAtTheEndIfRequested;
 
+	/** TRUE if delta serialization should be forced during FBlueprintCompileReinstancer::CopyPropertiesForUnrelatedObjects */
+	bool bUseDeltaSerializationToCopyProperties;
+
 public:
 	// FSerializableObject interface
 	virtual void AddReferencedObjects(FReferenceCollector& Collector) override;
@@ -243,7 +247,16 @@ protected:
 	/** Determine whether reinstancing actors should preserve the root component of the new actor */
 	virtual bool ShouldPreserveRootComponentOfReinstancedActor() const { return true; }
 
-	static void CopyPropertiesForUnrelatedObjects(UObject* OldObject, UObject* NewObject, bool bClearExternalReferences);
+	/**
+	* Attempts to copy as many properties as possible from the old object to the new. 
+	* Use during BP compilation to copy properties from the old CDO to the new one.
+	* 
+	* @param OldObject						The old object to copy properties from
+	* @param NewObject						The new Object to copy properties to
+	* @param bClearExternalReferences		If true then attempt to replace references to old classes and instances on this object with the corresponding new ones
+	* @param bForceDeltaSerialization		If true the delta serialization will be used when copying
+	*/
+	static void CopyPropertiesForUnrelatedObjects(UObject* OldObject, UObject* NewObject, bool bClearExternalReferences, bool bForceDeltaSerialization = false);
 
 private:
 	/** Handles the work of ReplaceInstancesOfClass, handling both normal replacement of instances and batch */
diff --git a/Engine/Source/Runtime/Engine/Classes/Engine/Blueprint.h b/Engine/Source/Runtime/Engine/Classes/Engine/Blueprint.h
--- a/Engine/Source/Runtime/Engine/Classes/Engine/Blueprint.h
+++ b/Engine/Source/Runtime/Engine/Classes/Engine/Blueprint.h
@@ -165,6 +165,9 @@ public:
 	/** Whether or not to update Find-in-Blueprint search metadata */
 	bool bSkipFiBSearchMetaUpdate;
 
+	/** Whether or not to use Delta Serialization when copying unrelated objects */
+	bool bUseDeltaSerializationDuringReinstancing;
+
 	TSharedPtr<FString> OutHeaderSourceCode;
 	TSharedPtr<FString> OutCppSourceCode;
 	FCompilerNativizationOptions NativizationOptions;
