Control Rig: Local variables cannot be copied in Control Rig

diff --git a/Engine/Plugins/Animation/ControlRig/Source/ControlRigEditor/Private/Editor/ControlRigEditor.cpp b/Engine/Plugins/Animation/ControlRig/Source/ControlRigEditor/Private/Editor/ControlRigEditor.cpp
--- a/Engine/Plugins/Animation/ControlRig/Source/ControlRigEditor/Private/Editor/ControlRigEditor.cpp
+++ b/Engine/Plugins/Animation/ControlRig/Source/ControlRigEditor/Private/Editor/ControlRigEditor.cpp
@@ -2087,6 +2087,22 @@ void FControlRigEditor::OnAddNewLocalVariable()
 	}
 }
 
+void FControlRigEditor::OnPasteNewLocalVariable(const FBPVariableDescription& VariableDescription)
+{
+	if (!CanAddNewLocalVariable())
+	{
+		return;
+	}
+
+	FRigVMGraphVariableDescription TypeVar;
+	TypeVar.ChangeType(VariableDescription.VarType);
+	FRigVMGraphVariableDescription NewVar = GetFocusedController()->AddLocalVariable(VariableDescription.VarName, TypeVar.CPPType, TypeVar.CPPTypeObject, VariableDescription.DefaultValue, true, true);
+	if(NewVar.Name.IsNone())
+	{
+		LogSimpleMessage( LOCTEXT("PasteLocalVariable_Error", "Pasting new local variable failed.") );
+	}
+}
+
 void FControlRigEditor::DeleteSelectedNodes()
 {
 	DECLARE_SCOPE_HIERARCHICAL_COUNTER_FUNC()
diff --git a/Engine/Plugins/Animation/ControlRig/Source/ControlRigEditor/Private/Editor/ControlRigEditor.h b/Engine/Plugins/Animation/ControlRig/Source/ControlRigEditor/Private/Editor/ControlRigEditor.h
--- a/Engine/Plugins/Animation/ControlRig/Source/ControlRigEditor/Private/Editor/ControlRigEditor.h
+++ b/Engine/Plugins/Animation/ControlRig/Source/ControlRigEditor/Private/Editor/ControlRigEditor.h
@@ -118,6 +118,7 @@ public:
 	virtual bool TransactionObjectAffectsBlueprint(UObject* InTransactedObject) override;
 	virtual bool CanAddNewLocalVariable() const override;
 	virtual void OnAddNewLocalVariable() override;
+	virtual void OnPasteNewLocalVariable(const FBPVariableDescription& VariableDescription) override;
 	
 	virtual void DeleteSelectedNodes() override;
 	virtual bool CanDeleteNodes() const override;
diff --git a/Engine/Source/Editor/Kismet/Private/BlueprintEditor.cpp b/Engine/Source/Editor/Kismet/Private/BlueprintEditor.cpp
--- a/Engine/Source/Editor/Kismet/Private/BlueprintEditor.cpp
+++ b/Engine/Source/Editor/Kismet/Private/BlueprintEditor.cpp
@@ -8853,6 +8853,25 @@ void FBlueprintEditor::OnAddNewLocalVariable()
 	}
 }
 
+void FBlueprintEditor::OnPasteNewLocalVariable(const FBPVariableDescription& VariableDescription)
+{
+	if (!CanAddNewLocalVariable())
+	{
+		return;
+	}
+
+	// Find the top level graph to place the local variables into
+	UEdGraph* TargetGraph = FBlueprintEditorUtils::GetTopLevelGraph(FocusedGraphEdPtr.Pin()->GetCurrentGraph());
+	check(TargetGraph->GetSchema()->GetGraphType(TargetGraph) == GT_Function);
+
+	bool bSuccess = MyBlueprintWidget.IsValid() && FBlueprintEditorUtils::AddLocalVariable(GetBlueprintObj(), TargetGraph, VariableDescription.VarName, VariableDescription.VarType, VariableDescription.DefaultValue);
+
+	if(!bSuccess)
+	{
+		LogSimpleMessage( LOCTEXT("PasteLocalVariable_Error", "Pasting new local variable failed.") );
+	}
+}
+
 void FBlueprintEditor::OnAddNewDelegate()
 {
 	if (!AddNewDelegateIsVisible())
diff --git a/Engine/Source/Editor/Kismet/Private/SMyBlueprint.cpp b/Engine/Source/Editor/Kismet/Private/SMyBlueprint.cpp
--- a/Engine/Source/Editor/Kismet/Private/SMyBlueprint.cpp
+++ b/Engine/Source/Editor/Kismet/Private/SMyBlueprint.cpp
@@ -3439,6 +3439,21 @@ void SMyBlueprint::OnCopy()
 			OutputString = VAR_PREFIX + OutputString;
 		}
 	}
+	else if (FEdGraphSchemaAction_BlueprintVariableBase* BPVariable = SelectionAsBlueprintVariable())
+	{
+		const UEdGraphSchema* Schema = GetFocusedGraph()->GetSchema();
+		TArray<FBPVariableDescription> LocalVariables;
+		Schema->GetLocalVariables(GetFocusedGraph(), LocalVariables);
+		for (const FBPVariableDescription& VariableDescription : LocalVariables)
+		{
+			if (VariableDescription.VarName == BPVariable->GetVariableName())
+			{
+				FBPVariableDescription::StaticStruct()->ExportText(OutputString, &VariableDescription, &VariableDescription, nullptr, 0, nullptr, false);
+				OutputString = VAR_PREFIX + OutputString;
+				break;
+			}
+		}
+	}
 	else if (FEdGraphSchemaAction_K2Graph* GraphAction = SelectionAsGraph())
 	{
 		if (!Blueprint->ExportGraphToText(GraphAction->EdGraph, OutputString))
@@ -3461,11 +3476,26 @@ bool SMyBlueprint::CanCopy() const
 	{
 		return FBlueprintEditorUtils::FindNewVariableIndex(Blueprint, VarAction->GetVariableName()) != INDEX_NONE;
 	}
-	if (FEdGraphSchemaAction_K2LocalVar* LocalVarAction = SelectionAsLocalVar())
+	else if (FEdGraphSchemaAction_K2LocalVar* LocalVarAction = SelectionAsLocalVar())
 	{
 		return FBlueprintEditorUtils::FindLocalVariable(Blueprint, Cast<UStruct>(LocalVarAction->GetVariableScope()), LocalVarAction->GetVariableName()) != nullptr;
 	}
-	if (FEdGraphSchemaAction_K2Graph* GraphAction = SelectionAsGraph())
+	else if (FEdGraphSchemaAction_BlueprintVariableBase* BPVariable = SelectionAsBlueprintVariable())
+	{
+		TArray<FBPVariableDescription> LocalVariables;
+		const UEdGraphSchema* Schema = GetFocusedGraph()->GetSchema();
+		Schema->GetLocalVariables(GetFocusedGraph(), LocalVariables);
+
+		for (const FBPVariableDescription& VariableDescription : LocalVariables)
+		{
+			if (VariableDescription.VarName == BPVariable->GetVariableName())
+			{
+				return true;
+			}
+		}
+		return false;		
+	}
+	else if (FEdGraphSchemaAction_K2Graph* GraphAction = SelectionAsGraph())
 	{
 		if (GraphAction->GraphType == EEdGraphSchemaAction_K2Graph::Function ||
 			GraphAction->GraphType == EEdGraphSchemaAction_K2Graph::Macro)
@@ -3559,37 +3589,42 @@ void SMyBlueprint::OnPasteLocalVariable()
 			TArray<UK2Node_FunctionEntry*> FunctionEntry;
 			FocusedGraph->GetNodesOfClass<UK2Node_FunctionEntry>(FunctionEntry);
 
-			if (FunctionEntry.Num() == 1)
+			FString ClipboardText;
+			FPlatformApplicationMisc::ClipboardPaste(ClipboardText);
+			if (!ensure(ClipboardText.StartsWith(VAR_PREFIX, ESearchCase::CaseSensitive)))
 			{
-				FString ClipboardText;
-				FPlatformApplicationMisc::ClipboardPaste(ClipboardText);
-				if (!ensure(ClipboardText.StartsWith(VAR_PREFIX, ESearchCase::CaseSensitive)))
-				{
-					return;
-				}
+				return;
+			}
 
-				FBPVariableDescription Description;
-				FStringOutputDevice Errors;
-				const TCHAR* Import = ClipboardText.GetCharArray().GetData() + FCString::Strlen(VAR_PREFIX);
-				FBPVariableDescription::StaticStruct()->ImportText(Import, &Description, nullptr, 0, &Errors, FBPVariableDescription::StaticStruct()->GetName());
-				if (Errors.IsEmpty())
+			FBPVariableDescription Description;
+			FStringOutputDevice Errors;
+			const TCHAR* Import = ClipboardText.GetCharArray().GetData() + FCString::Strlen(VAR_PREFIX);
+			FBPVariableDescription::StaticStruct()->ImportText(Import, &Description, nullptr, 0, &Errors, FBPVariableDescription::StaticStruct()->GetName());
+			if (Errors.IsEmpty())
+			{
+				FBPVariableDescription NewVar = FBlueprintEditorUtils::DuplicateVariableDescription(Blueprint, Description);
+				if (NewVar.VarGuid.IsValid())
 				{
-					FBPVariableDescription NewVar = FBlueprintEditorUtils::DuplicateVariableDescription(Blueprint, Description);
-					if (NewVar.VarGuid.IsValid())
-					{
-						FScopedTransaction Transaction(FText::Format(LOCTEXT("PasteLocalVariable", "Paste Local Variable: {0}"), FText::FromName(NewVar.VarName)));
+					FScopedTransaction Transaction(FText::Format(LOCTEXT("PasteLocalVariable", "Paste Local Variable: {0}"), FText::FromName(NewVar.VarName)));
+
+					NewVar.Category = GetPasteCategory();
 
-						NewVar.Category = GetPasteCategory();
 
+					if (FunctionEntry.Num() == 1)
+					{
 						FunctionEntry[0]->Modify();
 						FunctionEntry[0]->LocalVariables.Add(NewVar);
+					}
+					else
+					{
+						BlueprintEditorPtr.Pin()->OnPasteNewLocalVariable(NewVar);
+					}
 
-						// Potentially adjust variable names for any child blueprints
-						FBlueprintEditorUtils::ValidateBlueprintChildVariables(Blueprint, NewVar.VarName);
-						FBlueprintEditorUtils::MarkBlueprintAsStructurallyModified(Blueprint);
+					// Potentially adjust variable names for any child blueprints
+					FBlueprintEditorUtils::ValidateBlueprintChildVariables(Blueprint, NewVar.VarName);
+					FBlueprintEditorUtils::MarkBlueprintAsStructurallyModified(Blueprint);
 
-						SelectItemByName(NewVar.VarName);
-					}
+					SelectItemByName(NewVar.VarName);
 				}
 			}
 		}
diff --git a/Engine/Source/Editor/Kismet/Public/BlueprintEditor.h b/Engine/Source/Editor/Kismet/Public/BlueprintEditor.h
--- a/Engine/Source/Editor/Kismet/Public/BlueprintEditor.h
+++ b/Engine/Source/Editor/Kismet/Public/BlueprintEditor.h
@@ -552,6 +552,9 @@ public:
 	/** Adds a new local variable to the focused function graph */
 	virtual void OnAddNewLocalVariable();
 
+	/** Pastes a new local variable to the focused function graph */
+	virtual void OnPasteNewLocalVariable(const FBPVariableDescription& VariableDescription);
+
 	// Type of new document/graph being created by a menu item
 	enum ECreatedDocumentType
 	{
