Fix non-unity compile error

diff --git a/Engine/Source/Editor/Kismet/Private/BlueprintDebugger.h b/Engine/Source/Editor/Kismet/Private/BlueprintDebugger.h
--- a/Engine/Source/Editor/Kismet/Private/BlueprintDebugger.h
+++ b/Engine/Source/Editor/Kismet/Private/BlueprintDebugger.h
@@ -4,6 +4,8 @@
 
 #include "Templates/UniquePtr.h"
 
+class UBlueprint;
+
 struct FBlueprintDebugger
 {
 	// Initializes the global state of the debugger (commands, tab spawners, etc):
