Blueprint Debugger: Add more World information to the FParentLineItem widget and only allow objects from a PIE world

diff --git a/Engine/Source/Editor/Kismet/Private/Debugging/SKismetDebugTreeView.cpp b/Engine/Source/Editor/Kismet/Private/Debugging/SKismetDebugTreeView.cpp
--- a/Engine/Source/Editor/Kismet/Private/Debugging/SKismetDebugTreeView.cpp
+++ b/Engine/Source/Editor/Kismet/Private/Debugging/SKismetDebugTreeView.cpp
@@ -2484,6 +2484,18 @@ public:
 		ObjectRef = Object;
 	}
 
+	virtual TSharedRef<SWidget> GenerateNameWidget(TSharedPtr<FString> InSearchString) override
+	{
+		return SNew(PropertyInfoViewStyle::STextHighlightOverlay)
+			.FullText(this, &FParentLineItem::GetDisplayName)
+			.HighlightText(this, &FParentLineItem::GetHighlightText, InSearchString)
+			[
+				SNew(STextBlock)
+					.ToolTipText(this, &FParentLineItem::GetTooltipText)
+					.Text(this, &FParentLineItem::GetDisplayName)
+			];
+	}
+
 	virtual UObject* GetParentObject() override
 	{
 		return ObjectRef.Get();
@@ -2646,6 +2658,51 @@ private:
 			FKismetDebugUtilities::ClearPinWatches(Blueprint);
 		}
 	}
+
+	FText GetTooltipText() const
+	{
+		if (const UObject* Object = ObjectRef.Get())
+		{
+			if (UWorld* World = Object->GetTypedOuter<UWorld>())
+			{
+				ENetMode NetMode = World->GetNetMode();
+
+				FText WorldName;
+
+				switch (NetMode)
+				{
+				case NM_Standalone:
+					WorldName = LOCTEXT("DebugWorldStandalone", "Standalone");
+					break;
+
+				case NM_ListenServer:
+					WorldName = LOCTEXT("DebugWorldListenServer", "Listen Server");
+					break;
+
+				case NM_DedicatedServer:
+					WorldName = LOCTEXT("DebugWorldDedicatedServer", "Dedicated Server");
+					break;
+
+				case NM_Client:
+					if (FWorldContext* PieContext = GEngine->GetWorldContextFromWorld(World))
+					{
+						WorldName = FText::Format(LOCTEXT("DebugWorldClient", "Client {0}"), PieContext->PIEInstance - 1);
+						break;
+					}
+					[[fallthrough]];
+				default:
+					WorldName = LOCTEXT("UnknownWorldMode", "[Unknown]");
+				};
+
+				return FText::FormatNamed(LOCTEXT("ParentLineTooltip", "{ObjectFullPath}\nWorld: {WorldFullPath}\nWorld Type: {WorldType}"),
+					TEXT("ObjectFullPath"), FText::FromString(Object->GetPathName()),
+					TEXT("WorldFullPath"), FText::FromString(World->GetPathName()),
+					TEXT("WorldType"), WorldName);
+			}
+		}
+
+		return GetDisplayName();
+	}
 };
 
 //////////////////////////////////////////////////////////////////////////
diff --git a/Engine/Source/Editor/Kismet/Private/Debugging/SKismetDebuggingView.cpp b/Engine/Source/Editor/Kismet/Private/Debugging/SKismetDebuggingView.cpp
--- a/Engine/Source/Editor/Kismet/Private/Debugging/SKismetDebuggingView.cpp
+++ b/Engine/Source/Editor/Kismet/Private/Debugging/SKismetDebuggingView.cpp
@@ -411,10 +411,10 @@ void SKismetDebuggingView::Tick(const FGeometry& AllottedGeometry, const double
 				continue;
 			}
 
-			// only include actors in current world
-			if (AActor* Actor = Cast<AActor>(Instance))
+			// only include instances of objects in a PIE world
+			if (UWorld* World = Instance->GetTypedOuter<UWorld>())
 			{
-				if (!GEditor->PlayWorld->ContainsActor(Actor))
+				if (!World || World->WorldType != EWorldType::PIE)
 				{
 					continue;
 				}
