Auto-import namespace on non-imported variable type selection in the Blueprint editor details view.

This includes changing types on class variables as well as function inputs/outputs and local variables.

diff --git a/Engine/Source/Editor/Kismet/Private/BlueprintDetailsCustomization.cpp b/Engine/Source/Editor/Kismet/Private/BlueprintDetailsCustomization.cpp
--- a/Engine/Source/Editor/Kismet/Private/BlueprintDetailsCustomization.cpp
+++ b/Engine/Source/Editor/Kismet/Private/BlueprintDetailsCustomization.cpp
@@ -64,6 +64,7 @@
 #include "Widgets/Layout/SWidgetSwitcher.h"
 #include "BlueprintNamespaceHelper.h"
 #include "BlueprintNamespaceRegistry.h"
+#include "BlueprintNamespaceUtilities.h"
 #include "Widgets/Input/SSuggestionTextBox.h"
 #include "DragAndDrop/DecoratedDragDropOp.h"
 
@@ -1411,6 +1412,22 @@ void FBlueprintVarActionDetails::OnVarTypeChanged(const FEdGraphPinType& NewPinT
 				{
 					FBlueprintEditorUtils::ChangeMemberVariableType(GetBlueprintObj(), VarName, NewPinType);
 				}
+
+				if (GetDefault<UBlueprintEditorSettings>()->bEnableNamespaceImportingFeatures)
+				{
+					TSharedPtr<FBlueprintEditor> BlueprintEditor = MyBlueprint.Pin()->GetBlueprintEditor().Pin();
+					
+					// If the underlying type object's namespace is not imported, auto-import it now into the current editor context.
+					const UObject* PinSubCategoryObject = NewPinType.PinSubCategoryObject.Get();
+					if (PinSubCategoryObject && BlueprintEditor.IsValid() && BlueprintEditor->IsNonImportedObject(PinSubCategoryObject))
+					{
+						FString Namespace = FBlueprintNamespaceUtilities::GetObjectNamespace(PinSubCategoryObject);
+						if (!Namespace.IsEmpty())
+						{
+							BlueprintEditor->ImportNamespace(Namespace);
+						}
+					}
+				}
 			}
 		}
 	}
@@ -3349,6 +3366,26 @@ void FBlueprintGraphArgumentLayout::PinInfoChanged(const FEdGraphPinType& PinTyp
 
 							// Reset default value, it probably doesn't match
 							(*UDPinPtr)->PinDefaultValue.Reset();
+
+							if (GetDefault<UBlueprintEditorSettings>()->bEnableNamespaceImportingFeatures)
+							{
+								TSharedPtr<FBlueprintEditor> BlueprintEditor;
+								if(MyBPPinned.IsValid())
+								{
+									BlueprintEditor = MyBPPinned->GetBlueprintEditor().Pin();
+								}
+
+								// If the underlying type object's namespace is not imported, auto-import it now into the current editor context.
+								const UObject* PinSubCategoryObject = PinType.PinSubCategoryObject.Get();
+								if (PinSubCategoryObject && BlueprintEditor.IsValid() && BlueprintEditor->IsNonImportedObject(PinSubCategoryObject))
+								{
+									FString Namespace = FBlueprintNamespaceUtilities::GetObjectNamespace(PinSubCategoryObject);
+									if (!Namespace.IsEmpty())
+									{
+										BlueprintEditor->ImportNamespace(Namespace);
+									}
+								}
+							}
 						}
 						GraphActionDetailsPinned->OnParamsChanged(Node);
 					}
