Control Rig: Fix inconsistency when renaming control rig functions/variables/local variables with numbers

diff --git a/Engine/Plugins/Animation/ControlRig/Source/ControlRigDeveloper/Private/Graph/ControlRigGraphSchema.cpp b/Engine/Plugins/Animation/ControlRig/Source/ControlRigDeveloper/Private/Graph/ControlRigGraphSchema.cpp
--- a/Engine/Plugins/Animation/ControlRig/Source/ControlRigDeveloper/Private/Graph/ControlRigGraphSchema.cpp
+++ b/Engine/Plugins/Animation/ControlRig/Source/ControlRigDeveloper/Private/Graph/ControlRigGraphSchema.cpp
@@ -90,6 +90,45 @@ EValidatorResult FControlRigLocalVariableNameValidator::IsValid(const FName& Nam
 	return IsValid(Name.ToString(), bOriginal);
 }
 
+FControlRigNameValidator::FControlRigNameValidator(const UBlueprint* Blueprint, const UStruct* ValidationScope, FName InExistingName)
+: FStringSetNameValidator(InExistingName.ToString())
+{
+	if (Blueprint)
+	{
+		TSet<FName> NamesTemp;
+		FBlueprintEditorUtils::GetClassVariableList(Blueprint, NamesTemp, true);
+		FBlueprintEditorUtils::GetFunctionNameList(Blueprint, NamesTemp);
+		FBlueprintEditorUtils::GetAllGraphNames(Blueprint, NamesTemp);
+		FBlueprintEditorUtils::GetSCSVariableNameList(Blueprint, NamesTemp);
+		FBlueprintEditorUtils::GetImplementingBlueprintsFunctionNameList(Blueprint, NamesTemp);
+
+		for (FName & Name : NamesTemp)
+		{
+			Names.Add(Name.ToString());
+		}
+	}
+}
+
+EValidatorResult FControlRigNameValidator::IsValid(const FString& Name, bool bOriginal)
+{
+	EValidatorResult Result = FStringSetNameValidator::IsValid(Name, bOriginal);
+	if (Result == EValidatorResult::Ok)
+	{
+		if (URigVMController::GetSanitizedName(Name, false, true) == Name)
+		{
+			return Result;
+		}
+
+		return EValidatorResult::ContainsInvalidCharacters;
+	}
+	return Result;
+}
+
+EValidatorResult FControlRigNameValidator::IsValid(const FName& Name, bool bOriginal)
+{
+	return IsValid(Name.ToString(), bOriginal);
+}
+
 FEdGraphPinType FControlRigGraphSchemaAction_LocalVar::GetPinType() const
 {
 	if (UControlRigGraph* Graph = Cast<UControlRigGraph>(GetVariableScope()))
@@ -755,6 +794,22 @@ void UControlRigGraphSchema::InsertAdditionalActions(TArray<UBlueprint*> InBluep
 	}
 }
 
+TSharedPtr<INameValidatorInterface> UControlRigGraphSchema::GetNameValidator(const UBlueprint* BlueprintObj, const FName& OriginalName, const UStruct* ValidationScope, const FName& ActionTypeId) const
+{
+	if (ActionTypeId == FControlRigGraphSchemaAction_LocalVar::StaticGetTypeId())
+	{
+		if (const UControlRigGraph* ControlRigGraph = Cast<UControlRigGraph>(ValidationScope))
+		{
+			if (const URigVMGraph* Graph = ControlRigGraph->GetModel())
+			{
+				return MakeShareable(new FControlRigLocalVariableNameValidator(BlueprintObj, Graph, OriginalName));
+			}
+		}
+	}
+
+	return MakeShareable(new FControlRigNameValidator(BlueprintObj, ValidationScope, OriginalName));		
+}
+
 bool UControlRigGraphSchema::SupportsPinType(UScriptStruct* ScriptStruct) const
 {
 	if(!ScriptStruct)
diff --git a/Engine/Plugins/Animation/ControlRig/Source/ControlRigDeveloper/Public/Graph/ControlRigGraphSchema.h b/Engine/Plugins/Animation/ControlRig/Source/ControlRigDeveloper/Public/Graph/ControlRigGraphSchema.h
--- a/Engine/Plugins/Animation/ControlRig/Source/ControlRigDeveloper/Public/Graph/ControlRigGraphSchema.h
+++ b/Engine/Plugins/Animation/ControlRig/Source/ControlRigDeveloper/Public/Graph/ControlRigGraphSchema.h
@@ -58,6 +58,20 @@ public:
 	// End FNameValidatorInterface
 };
 
+/////////////////////////////////////////////////////
+// FControlRigNameValidator
+class CONTROLRIGDEVELOPER_API FControlRigNameValidator : public FStringSetNameValidator
+{
+
+public:
+	FControlRigNameValidator(const class UBlueprint* Blueprint, const UStruct* ValidationScope, FName InExistingName = NAME_None);
+
+	// Begin FNameValidatorInterface
+	virtual EValidatorResult IsValid(const FString& Name, bool bOriginal) override;
+	virtual EValidatorResult IsValid(const FName& Name, bool bOriginal) override;
+	// End FNameValidatorInterface
+};
+
 USTRUCT()
 struct CONTROLRIGDEVELOPER_API FControlRigGraphSchemaAction_LocalVar : public FEdGraphSchemaAction_BlueprintVariableBase
 {
@@ -219,6 +233,7 @@ public:
 	virtual bool SupportsDropPinOnNode(UEdGraphNode* InTargetNode, const FEdGraphPinType& InSourcePinType, EEdGraphPinDirection InSourcePinDirection, FText& OutErrorMessage) const override;
 	virtual void SetPinBeingDroppedOnNode(UEdGraphPin* InSourcePin) const override { PinBeingDropped = InSourcePin; }
 	virtual void InsertAdditionalActions(TArray<UBlueprint*> InBlueprints, TArray<UEdGraph*> InGraphs, TArray<UEdGraphPin*> InPins, FGraphActionListBuilderBase& OutAllActions) const override;
+	virtual TSharedPtr<INameValidatorInterface> GetNameValidator(const UBlueprint* BlueprintObj, const FName& OriginalName, const UStruct* ValidationScope, const FName& ActionTypeId) const override;
 
 	/** Returns true if the schema supports the script type */
 	bool SupportsPinType(UScriptStruct* ScriptStruct) const;
diff --git a/Engine/Source/Editor/Kismet/Private/BlueprintDetailsCustomization.cpp b/Engine/Source/Editor/Kismet/Private/BlueprintDetailsCustomization.cpp
--- a/Engine/Source/Editor/Kismet/Private/BlueprintDetailsCustomization.cpp
+++ b/Engine/Source/Editor/Kismet/Private/BlueprintDetailsCustomization.cpp
@@ -270,6 +270,27 @@ void FBlueprintVarActionDetails::CustomizeDetails( IDetailLayoutBuilder& DetailL
 		BlueprintEditor.Pin()->OnRefresh().AddSP(this, &FBlueprintVarActionDetails::OnPostEditorRefresh);
 	}
 
+	UBlueprint* BlueprintPtr = GetBlueprintObj();
+	
+	// Get an appropiate name validator
+	TSharedPtr<INameValidatorInterface> NameValidator = nullptr;
+	{
+		const UEdGraphSchema* Schema = nullptr;		
+		if (BlueprintPtr)
+		{
+			TArray<UEdGraph*> Graphs;
+			BlueprintPtr->GetAllGraphs(Graphs);
+			if (Graphs.Num() > 0)
+			{
+				Schema = Graphs[0]->GetSchema();
+			}
+		}			
+	
+		if (Schema)
+		{
+			NameValidator = Schema->GetNameValidator(BlueprintPtr, GetVariableName(), nullptr, FEdGraphSchemaAction_K2Var::StaticGetTypeId());	
+		}
+	}
 
 	FProperty* VariableProperty = CachedVariableProperty.Get();
 
@@ -302,6 +323,27 @@ void FBlueprintVarActionDetails::CustomizeDetails( IDetailLayoutBuilder& DetailL
 		.ToolTip(VarNameTooltip)
 		.OnTextChanged(this, &FBlueprintVarActionDetails::OnVarNameChanged)
 		.OnTextCommitted(this, &FBlueprintVarActionDetails::OnVarNameCommitted)
+		.OnVerifyTextChanged_Lambda([this, NameValidator](const FText& InNewText, FText& OutErrorMessage) -> bool
+		{	
+			if (NameValidator.IsValid())
+			{
+				EValidatorResult ValidatorResult = NameValidator->IsValid(InNewText.ToString());
+				switch (ValidatorResult)
+				{
+					case EValidatorResult::Ok:
+					case EValidatorResult::ExistingName:
+						// These are fine, don't need to surface to the user, the rename can 'proceed' even if the name is the existing one
+						return true;
+						break;
+					default:
+						OutErrorMessage = INameValidatorInterface::GetErrorText(InNewText.ToString(), ValidatorResult);
+						return false;
+						break;
+				}
+			}
+			
+			return true;
+		})
 		.IsReadOnly(this, &FBlueprintVarActionDetails::GetVariableNameChangeEnabled)
 		.Font(IDetailLayoutBuilder::GetDetailFont())
 	];
diff --git a/Engine/Source/Editor/Kismet/Private/SBlueprintPalette.cpp b/Engine/Source/Editor/Kismet/Private/SBlueprintPalette.cpp
--- a/Engine/Source/Editor/Kismet/Private/SBlueprintPalette.cpp
+++ b/Engine/Source/Editor/Kismet/Private/SBlueprintPalette.cpp
@@ -1485,11 +1485,25 @@ bool SBlueprintPaletteItem::OnNameTextVerifyChanged(const FText& InNewText, FTex
 
 	UStruct* ValidationScope = nullptr;
 
+	const UEdGraphSchema* Schema = nullptr;
+
 	// Check if certain action names are unchanged.
 	if (ActionPtr.Pin()->GetTypeId() == FEdGraphSchemaAction_K2Var::StaticGetTypeId())
 	{
 		FEdGraphSchemaAction_K2Var* VarAction = (FEdGraphSchemaAction_K2Var*)ActionPtr.Pin().Get();
 		OriginalName = (VarAction->GetVariableName());
+
+		UClass* VarClass = VarAction->GetVariableClass();
+		if (VarClass)
+		{
+			UBlueprint* BlueprintObj = UBlueprint::GetBlueprintFromClass(VarClass);
+			TArray<UEdGraph*> Graphs;
+			BlueprintObj->GetAllGraphs(Graphs);
+			if (Graphs.Num() > 0)
+			{
+				Schema = Graphs[0]->GetSchema();
+			}
+		}
 	}
 	else if (ActionPtr.Pin()->GetTypeId() == FEdGraphSchemaAction_K2LocalVar::StaticGetTypeId())
 	{
@@ -1497,6 +1511,18 @@ bool SBlueprintPaletteItem::OnNameTextVerifyChanged(const FText& InNewText, FTex
 		OriginalName = (LocalVarAction->GetVariableName());
 		
 		ValidationScope = CastChecked<UStruct>(LocalVarAction->GetVariableScope());
+
+		UClass* VarClass = LocalVarAction->GetVariableClass();
+		if (VarClass)
+		{
+			UBlueprint* BlueprintObj = UBlueprint::GetBlueprintFromClass(VarClass);
+			TArray<UEdGraph*> Graphs;
+			BlueprintObj->GetAllGraphs(Graphs);
+			if (Graphs.Num() > 0)
+			{
+				Schema = Graphs[0]->GetSchema();
+			}
+		}
 	}
 	else
 	{
@@ -1516,6 +1542,7 @@ bool SBlueprintPaletteItem::OnNameTextVerifyChanged(const FText& InNewText, FTex
 		if (Graph)
 		{
 			OriginalName = Graph->GetFName();
+			Schema = Graph->GetSchema();
 		}
 	}
 
@@ -1541,20 +1568,26 @@ bool SBlueprintPaletteItem::OnNameTextVerifyChanged(const FText& InNewText, FTex
 	}
 	else
 	{
-		TSharedPtr<INameValidatorInterface> NameValidator = MakeShareable(new FKismetNameValidator(BlueprintObj, OriginalName, ValidationScope));
-
-		EValidatorResult ValidatorResult = NameValidator->IsValid(TextAsString);
-		switch (ValidatorResult)
+		TSharedPtr<INameValidatorInterface> NameValidator = nullptr;
+		if (Schema)
 		{
-		case EValidatorResult::Ok:
-		case EValidatorResult::ExistingName:
-			// These are fine, don't need to surface to the user, the rename can 'proceed' even if the name is the existing one
-			break;
-		default:
-			OutErrorMessage = INameValidatorInterface::GetErrorText(TextAsString, ValidatorResult);
-			break;
+			NameValidator = Schema->GetNameValidator(BlueprintObj, OriginalName, ValidationScope, ActionPtr.Pin()->GetTypeId());	
+		}
+		
+		if (NameValidator.IsValid())
+		{
+			EValidatorResult ValidatorResult = NameValidator->IsValid(TextAsString);
+			switch (ValidatorResult)
+			{
+			case EValidatorResult::Ok:
+			case EValidatorResult::ExistingName:
+				// These are fine, don't need to surface to the user, the rename can 'proceed' even if the name is the existing one
+				break;
+			default:
+				OutErrorMessage = INameValidatorInterface::GetErrorText(TextAsString, ValidatorResult);
+				break;
+			}
 		}
-	
 	}
 	
 	return OutErrorMessage.IsEmpty();
diff --git a/Engine/Source/Runtime/Engine/Classes/EdGraph/EdGraphSchema.h b/Engine/Source/Runtime/Engine/Classes/EdGraph/EdGraphSchema.h
--- a/Engine/Source/Runtime/Engine/Classes/EdGraph/EdGraphSchema.h
+++ b/Engine/Source/Runtime/Engine/Classes/EdGraph/EdGraphSchema.h
@@ -10,6 +10,9 @@
 #include "AssetData.h"
 #include "UObject/ObjectKey.h"
 #include "Input/Reply.h"
+#if WITH_EDITOR
+#include "Kismet2/Kismet2NameValidators.h"
+#endif
 #include "EdGraphSchema.generated.h"
 
 class FSlateRect;
@@ -1223,4 +1226,20 @@ class ENGINE_API UEdGraphSchema : public UObject
 	 * @param OutAllActions Resulting compatible actions
 	 */
 	virtual void InsertAdditionalActions(TArray<UBlueprint*> InBlueprints, TArray<UEdGraph*> InGraphs, TArray<UEdGraphPin*> InPins, FGraphActionListBuilderBase& OutAllActions) const {}
+
+#if WITH_EDITOR
+	/**
+	 * Returns a name validator appropiate for the schema and object that is being named
+	 * @param InBlueprintObj The blueprint where the object being named lives.
+	 * @param InOriginalName The original name of the object.
+	 * @param InValidationScope The scope where the named object lives.
+	 * @param InActionTypeId The type of object that is being named.
+	 * @param NameValidator The name validator to use when naming this object.
+	 */
+	virtual TSharedPtr<INameValidatorInterface> GetNameValidator(const UBlueprint* InBlueprintObj, const FName& InOriginalName, const UStruct* InValidationScope, const FName& InActionTypeId) const
+	{
+		return MakeShareable(new FKismetNameValidator(InBlueprintObj, InOriginalName, InValidationScope));
+	}
+#endif
+	
 };
