Fix high memory usage in Blueprint Debugger

diff --git a/Engine/Source/Editor/Kismet/Private/Debugging/SKismetDebugTreeView.cpp b/Engine/Source/Editor/Kismet/Private/Debugging/SKismetDebugTreeView.cpp
--- a/Engine/Source/Editor/Kismet/Private/Debugging/SKismetDebugTreeView.cpp
+++ b/Engine/Source/Editor/Kismet/Private/Debugging/SKismetDebugTreeView.cpp
@@ -1238,6 +1238,8 @@ protected:
 	TWeakObjectPtr< UObject > ParentObjectRef;
 	const FEdGraphPinReference ObjectRef;
 	const TArray<FName> PathToProperty;
+	mutable TSharedPtr<FPropertyInstanceInfo> CachedPropertyInfo = nullptr;
+	mutable FText CachedErrorString;
 public:
 	FWatchLineItem(const UEdGraphPin* PinToWatch, UObject* ParentObject)
 		: FLineItemWithChildren(DLT_Watch)
@@ -1331,32 +1333,12 @@ protected:
 
 	virtual void GatherChildren(TArray<FDebugTreeItemPtr>& OutChildren, const FString& InSearchString, bool bRespectSearch) override
 	{
-		if (UEdGraphPin* PinToWatch = ObjectRef.Get())
+		TSharedPtr<FPropertyInstanceInfo> ThisDebugInfo = GetPropertyInfo();
+		if (ThisDebugInfo.IsValid())
 		{
-			// Try to determine the blueprint that generated the watch
-			UBlueprint* ParentBlueprint = GetBlueprintForObject(ParentObjectRef.Get());
-
-			// Find a valid property mapping and display the current value
-			UObject* ParentObject = ParentObjectRef.Get();
-			if ((ParentBlueprint != ParentObject) && (ParentBlueprint != nullptr))
+			for (const TSharedPtr<FPropertyInstanceInfo>& ChildData : ThisDebugInfo->Children)
 			{
-				TSharedPtr<FPropertyInstanceInfo> DebugInfo;
-				const FKismetDebugUtilities::EWatchTextResult WatchStatus = FKismetDebugUtilities::GetDebugInfo(DebugInfo, ParentBlueprint, ParentObject, PinToWatch);
-
-				if (WatchStatus == FKismetDebugUtilities::EWTR_Valid)
-				{
-					check(DebugInfo);
-
-					TSharedPtr<FPropertyInstanceInfo> ThisDebugInfo = GetWatchedPropertyDebugInfo(DebugInfo);
-					
-					if (ThisDebugInfo.IsValid())
-					{
-						for (const TSharedPtr<FPropertyInstanceInfo>& ChildData : ThisDebugInfo->Children)
-						{
-							EnsureChildIsAdded(OutChildren, FWatchChildLineItem(ChildData.ToSharedRef(), AsShared()), InSearchString, bRespectSearch);
-						}
-					}
-				}
+				EnsureChildIsAdded(OutChildren, FWatchChildLineItem(ChildData.ToSharedRef(), AsShared()), InSearchString, bRespectSearch);
 			}
 		}
 	}
@@ -1566,36 +1548,13 @@ FText FWatchLineItem::GetDescription() const
 		UObject* ParentObject = ParentObjectRef.Get();
 		if ((ParentBlueprint != ParentObject) && (ParentBlueprint != nullptr))
 		{
-			TSharedPtr<FPropertyInstanceInfo> DebugInfo;
-			const FKismetDebugUtilities::EWatchTextResult WatchStatus = FKismetDebugUtilities::GetDebugInfo(DebugInfo, ParentBlueprint, ParentObject, PinToWatch);
-
-			switch (WatchStatus)
-			{
-			case FKismetDebugUtilities::EWTR_Valid:
+			TSharedPtr<FPropertyInstanceInfo> DebugInfo = GetPropertyInfo();
+			if (DebugInfo.IsValid())
 			{
-				check(DebugInfo);
-				TSharedPtr<FPropertyInstanceInfo> ThisDebugInfo = GetWatchedPropertyDebugInfo(DebugInfo);
-				if (ThisDebugInfo.IsValid())
-				{
-					const FString ValStr = ThisDebugInfo->Value.ToString();
-					return FText::FromString(ValStr.Replace(TEXT("\n"), TEXT(" ")));
-				}
-				else
-				{
-					return LOCTEXT("UnknownProperty", "No debug data");
-				}
-			}
-
-			case FKismetDebugUtilities::EWTR_NotInScope:
-				return LOCTEXT("NotInScope", "Not in scope");
-
-			case FKismetDebugUtilities::EWTR_NoProperty:
-				return LOCTEXT("UnknownProperty", "No debug data");
-
-			default:
-			case FKismetDebugUtilities::EWTR_NoDebugObject:
-				return LOCTEXT("NoDebugObject", "No debug object");
+				const FString ValStr = DebugInfo->Value.ToString();
+				return FText::FromString(ValStr.Replace(TEXT("\n"), TEXT(" ")));
 			}
+			return CachedErrorString;
 		}
 	}
 
@@ -1715,38 +1674,20 @@ TSharedRef<SWidget> FWatchLineItem::GetNameIcon()
 
 const FSlateBrush* FWatchLineItem::GetPinIcon() const
 {
-	if (UEdGraphPin* ObjectToFocus = ObjectRef.Get())
+	TSharedPtr<FPropertyInstanceInfo> ThisDebugInfo = GetPropertyInfo();
+	if (ThisDebugInfo.IsValid())
 	{
-		// Try to determine the blueprint that generated the watch
-		UBlueprint* ParentBlueprint = GetBlueprintForObject(ParentObjectRef.Get());
-
-		// Find a valid property mapping and display the current value
-		UObject* ParentObject = ParentObjectRef.Get();
-		if ((ParentBlueprint != ParentObject) && (ParentBlueprint != nullptr))
-		{
-			TSharedPtr<FPropertyInstanceInfo> DebugInfo;
-			const FKismetDebugUtilities::EWatchTextResult WatchStatus = FKismetDebugUtilities::GetDebugInfo(DebugInfo, ParentBlueprint, ParentObject, ObjectToFocus);
-
-			if (WatchStatus == FKismetDebugUtilities::EWTR_Valid)
-			{
-				check(DebugInfo);
-				TSharedPtr<FPropertyInstanceInfo> ThisDebugInfo = GetWatchedPropertyDebugInfo(DebugInfo);
-				if (ThisDebugInfo.IsValid())
-				{
-					FSlateColor BaseColor;
-					FSlateColor SecondaryColor;
-					FSlateBrush const* SecondaryIcon;
-					const FSlateBrush* Icon = FBlueprintEditor::GetVarIconAndColorFromProperty(
-						ThisDebugInfo->Property.Get(),
-						BaseColor,
-						SecondaryIcon,
-						SecondaryColor
-					);
+		FSlateColor BaseColor;
+		FSlateColor SecondaryColor;
+		FSlateBrush const* SecondaryIcon;
+		const FSlateBrush* Icon = FBlueprintEditor::GetVarIconAndColorFromProperty(
+			ThisDebugInfo->Property.Get(),
+			BaseColor,
+			SecondaryIcon,
+			SecondaryColor
+		);
 
-					return Icon;
-				}
-			}
-		}
+		return Icon;
 	}
 
 	return FEditorStyle::GetBrush(TEXT("NoBrush"));
@@ -1754,38 +1695,20 @@ const FSlateBrush* FWatchLineItem::GetPinIcon() const
 
 const FSlateBrush* FWatchLineItem::GetSecondaryPinIcon() const
 {
-	if (UEdGraphPin* ObjectToFocus = ObjectRef.Get())
+	TSharedPtr<FPropertyInstanceInfo> ThisDebugInfo = GetPropertyInfo();
+	if (ThisDebugInfo.IsValid())
 	{
-		// Try to determine the blueprint that generated the watch
-		UBlueprint* ParentBlueprint = GetBlueprintForObject(ParentObjectRef.Get());
-
-		// Find a valid property mapping and display the current value
-		UObject* ParentObject = ParentObjectRef.Get();
-		if ((ParentBlueprint != ParentObject) && (ParentBlueprint != nullptr))
-		{
-			TSharedPtr<FPropertyInstanceInfo> DebugInfo;
-			const FKismetDebugUtilities::EWatchTextResult WatchStatus = FKismetDebugUtilities::GetDebugInfo(DebugInfo, ParentBlueprint, ParentObject, ObjectToFocus);
-
-			if (WatchStatus == FKismetDebugUtilities::EWTR_Valid)
-			{
-				check(DebugInfo);
-				TSharedPtr<FPropertyInstanceInfo> ThisDebugInfo = GetWatchedPropertyDebugInfo(DebugInfo);
-				if (ThisDebugInfo.IsValid())
-				{
-					FSlateColor BaseColor;
-					FSlateColor SecondaryColor;
-					FSlateBrush const* SecondaryIcon;
-					const FSlateBrush* Icon = FBlueprintEditor::GetVarIconAndColorFromProperty(
-						ThisDebugInfo->Property.Get(),
-						BaseColor,
-						SecondaryIcon,
-						SecondaryColor
-					);
+		FSlateColor BaseColor;
+		FSlateColor SecondaryColor;
+		FSlateBrush const* SecondaryIcon;
+		const FSlateBrush* Icon = FBlueprintEditor::GetVarIconAndColorFromProperty(
+			ThisDebugInfo->Property.Get(),
+			BaseColor,
+			SecondaryIcon,
+			SecondaryColor
+		);
 
-					return SecondaryIcon;
-				}
-			}
-		}
+		return SecondaryIcon;
 	}
 
 	return FEditorStyle::GetBrush(TEXT("NoBrush"));
@@ -1797,29 +1720,14 @@ FSlateColor FWatchLineItem::GetPinIconColor() const
 
 	if (UEdGraphPin* ObjectToFocus = ObjectRef.Get())
 	{
-		// Try to determine the blueprint that generated the watch
-		UBlueprint* ParentBlueprint = GetBlueprintForObject(ParentObjectRef.Get());
-
-		// Find a valid property mapping and display the current value
-		UObject* ParentObject = ParentObjectRef.Get();
-		if ((ParentBlueprint != ParentObject) && (ParentBlueprint != nullptr))
+		TSharedPtr<FPropertyInstanceInfo> ThisDebugInfo = GetPropertyInfo();
+		if (ThisDebugInfo.IsValid())
 		{
-			TSharedPtr<FPropertyInstanceInfo> DebugInfo;
-			const FKismetDebugUtilities::EWatchTextResult WatchStatus = FKismetDebugUtilities::GetDebugInfo(DebugInfo, ParentBlueprint, ParentObject, ObjectToFocus);
-
-			if (WatchStatus == FKismetDebugUtilities::EWTR_Valid)
+			if (const UEdGraphSchema* Schema = ObjectToFocus->GetSchema())
 			{
-				check(DebugInfo);
-				TSharedPtr<FPropertyInstanceInfo> ThisDebugInfo = GetWatchedPropertyDebugInfo(DebugInfo);
-				if (ThisDebugInfo.IsValid())
-				{
-					if (const UEdGraphSchema* Schema = ObjectToFocus->GetSchema())
-					{
-						PinIconColor = Schema->GetPinTypeColor(ObjectToFocus->PinType);
-					}
-				}
+				PinIconColor = Schema->GetPinTypeColor(ObjectToFocus->PinType);
 			}
-		}
+		}	
 
 		if (FKismetDebugUtilities::IsPinBeingWatched(FBlueprintEditorUtils::FindBlueprintForNode(ObjectToFocus->GetOwningNode()), ObjectToFocus, PathToProperty))
 		{
@@ -1829,7 +1737,6 @@ FSlateColor FWatchLineItem::GetPinIconColor() const
 		}
 	}
 
-
 	return PinIconColor;
 }
 
@@ -1839,27 +1746,12 @@ FSlateColor FWatchLineItem::GetSecondaryPinIconColor() const
 
 	if (UEdGraphPin* ObjectToFocus = ObjectRef.Get())
 	{
-		// Try to determine the blueprint that generated the watch
-		UBlueprint* ParentBlueprint = GetBlueprintForObject(ParentObjectRef.Get());
-
-		// Find a valid property mapping and display the current value
-		UObject* ParentObject = ParentObjectRef.Get();
-		if ((ParentBlueprint != ParentObject) && (ParentBlueprint != nullptr))
+		TSharedPtr<FPropertyInstanceInfo> ThisDebugInfo = GetPropertyInfo();
+		if (ThisDebugInfo.IsValid())
 		{
-			TSharedPtr<FPropertyInstanceInfo> DebugInfo;
-			const FKismetDebugUtilities::EWatchTextResult WatchStatus = FKismetDebugUtilities::GetDebugInfo(DebugInfo, ParentBlueprint, ParentObject, ObjectToFocus);
-
-			if (WatchStatus == FKismetDebugUtilities::EWTR_Valid)
+			if (const UEdGraphSchema* Schema = ObjectToFocus->GetSchema())
 			{
-				check(DebugInfo);
-				TSharedPtr<FPropertyInstanceInfo> ThisDebugInfo = GetWatchedPropertyDebugInfo(DebugInfo);
-				if (ThisDebugInfo.IsValid())
-				{
-					if (const UEdGraphSchema* Schema = ObjectToFocus->GetSchema())
-					{
-						PinIconColor = Schema->GetSecondaryPinTypeColor(ObjectToFocus->PinType);
-					}
-				}
+				PinIconColor = Schema->GetSecondaryPinTypeColor(ObjectToFocus->PinType);
 			}
 		}
 
@@ -1889,28 +1781,10 @@ EVisibility FWatchLineItem::GetWatchIconVisibility() const
 
 FText FWatchLineItem::GetTypename() const
 {
-	if (UEdGraphPin* ObjectToFocus = ObjectRef.Get())
+	TSharedPtr<FPropertyInstanceInfo> ThisDebugInfo = GetPropertyInfo();
+	if (ThisDebugInfo.IsValid())
 	{
-		// Try to determine the blueprint that generated the watch
-		UBlueprint* ParentBlueprint = GetBlueprintForObject(ParentObjectRef.Get());
-
-		// Find a valid property mapping and display the current value
-		UObject* ParentObject = ParentObjectRef.Get();
-		if ((ParentBlueprint != ParentObject) && (ParentBlueprint != nullptr))
-		{
-			TSharedPtr<FPropertyInstanceInfo> DebugInfo;
-			const FKismetDebugUtilities::EWatchTextResult WatchStatus = FKismetDebugUtilities::GetDebugInfo(DebugInfo, ParentBlueprint, ParentObject, ObjectToFocus);
-
-			if (WatchStatus == FKismetDebugUtilities::EWTR_Valid)
-			{
-				check(DebugInfo);
-				TSharedPtr<FPropertyInstanceInfo> ThisDebugInfo = GetWatchedPropertyDebugInfo(DebugInfo);
-				if (ThisDebugInfo.IsValid())
-				{
-					return UEdGraphSchema_K2::TypeToText(ThisDebugInfo->Property.Get());
-				}
-			}
-		}
+		return UEdGraphSchema_K2::TypeToText(ThisDebugInfo->Property.Get());
 	}
 
 	return FText::GetEmpty();
@@ -1956,6 +1830,11 @@ TSharedPtr<FPropertyInstanceInfo> FWatchLineItem::GetWatchedPropertyDebugInfo(co
 
 TSharedPtr<FPropertyInstanceInfo> FWatchLineItem::GetPropertyInfo() const
 {
+	if (CachedPropertyInfo.IsValid())
+	{
+		return CachedPropertyInfo;
+	}
+
 	if (UEdGraphPin* ObjectToFocus = ObjectRef.Get())
 	{
 		// Try to determine the blueprint that generated the watch
@@ -1967,11 +1846,24 @@ TSharedPtr<FPropertyInstanceInfo> FWatchLineItem::GetPropertyInfo() const
 		{
 			TSharedPtr<FPropertyInstanceInfo> DebugInfo;
 			const FKismetDebugUtilities::EWatchTextResult WatchStatus = FKismetDebugUtilities::GetDebugInfo(DebugInfo, ParentBlueprint, ParentObject, ObjectToFocus);
-
-			if (WatchStatus == FKismetDebugUtilities::EWTR_Valid)
+			switch (WatchStatus)
+			{
+			case FKismetDebugUtilities::EWTR_Valid:
 			{
-				check(DebugInfo);
-				return GetWatchedPropertyDebugInfo(DebugInfo);
+					check(DebugInfo);
+					CachedPropertyInfo = GetWatchedPropertyDebugInfo(DebugInfo);
+					return CachedPropertyInfo;
+			}
+
+			case FKismetDebugUtilities::EWTR_NotInScope:
+				CachedErrorString = LOCTEXT("NotInScope", "Not in scope");
+
+			case FKismetDebugUtilities::EWTR_NoProperty:
+				CachedErrorString = LOCTEXT("UnknownProperty", "No debug data");
+
+			default:
+			case FKismetDebugUtilities::EWTR_NoDebugObject:
+				CachedErrorString = LOCTEXT("NoDebugObject", "No debug object");
 			}
 		}
 	}
