Auto-import associated namespace(s) after a value change on reference types in the Blueprint property editor.

diff --git a/Engine/Source/Editor/Kismet/Private/BlueprintDetailsCustomization.cpp b/Engine/Source/Editor/Kismet/Private/BlueprintDetailsCustomization.cpp
--- a/Engine/Source/Editor/Kismet/Private/BlueprintDetailsCustomization.cpp
+++ b/Engine/Source/Editor/Kismet/Private/BlueprintDetailsCustomization.cpp
@@ -926,6 +926,12 @@ void FBlueprintVarActionDetails::CustomizeDetails( IDetailLayoutBuilder& DetailL
 		}
 		else 
 		{
+			TSharedPtr<IDetailsView> DetailsView;
+			if (BlueprintEditor.IsValid())
+			{
+				DetailsView = BlueprintEditor.Pin()->GetInspector()->GetPropertyView();
+			}
+
 			if(IsALocalVariable(VariableProperty))
 			{
 				UFunction* StructScope = VariableProperty->GetOwner<UFunction>();
@@ -939,8 +945,6 @@ void FBlueprintVarActionDetails::CustomizeDetails( IDetailLayoutBuilder& DetailL
 
 				// There should always be an entry node in the function graph
 				check(EntryNodes.Num() > 0);
-
-				const FStructProperty* PotentialUDSProperty = CastField<const FStructProperty>(VariableProperty);
 				
 				UK2Node_FunctionEntry* FuncEntry = EntryNodes[0];
 
@@ -959,15 +963,10 @@ void FBlueprintVarActionDetails::CustomizeDetails( IDetailLayoutBuilder& DetailL
 					}
 				}
 
-				if(BlueprintEditor.IsValid())
+				if (DetailsView.IsValid())
 				{
-					TSharedPtr< IDetailsView > DetailsView  = BlueprintEditor.Pin()->GetInspector()->GetPropertyView();
-
-					if(DetailsView.IsValid())
-					{
-						TWeakObjectPtr<UK2Node_EditablePinBase> EntryNode = FuncEntry;
-						DetailsView->OnFinishedChangingProperties().AddSP(this, &FBlueprintVarActionDetails::OnFinishedChangingProperties, StructData, EntryNode);
-					}
+					TWeakObjectPtr<UK2Node_EditablePinBase> EntryNode = FuncEntry;
+					DetailsView->OnFinishedChangingProperties().AddSP(this, &FBlueprintVarActionDetails::OnFinishedChangingLocalVariable, StructData, EntryNode);
 				}
 
 				IDetailPropertyRow* Row = DefaultValueCategory.AddExternalStructureProperty(StructData, VariableProperty->GetFName());
@@ -1001,6 +1000,11 @@ void FBlueprintVarActionDetails::CustomizeDetails( IDetailLayoutBuilder& DetailL
 					{
 						Row->IsEnabled(IsVariableInheritedByBlueprint());
 					}
+
+					if (DetailsView.IsValid())
+					{
+						DetailsView->OnFinishedChangingProperties().AddSP(this, &FBlueprintVarActionDetails::OnFinishedChangingVariable, TargetBlueprintDefaultObject);
+					}
 				}
 			}
 		}
@@ -2724,9 +2728,19 @@ EVisibility FBlueprintVarActionDetails::IsTooltipEditVisible() const
 	return EVisibility::Collapsed;
 }
 
-void FBlueprintVarActionDetails::OnFinishedChangingProperties(const FPropertyChangedEvent& InPropertyChangedEvent, TSharedPtr<FStructOnScope> InStructData, TWeakObjectPtr<UK2Node_EditablePinBase> InEntryNode)
+void FBlueprintVarActionDetails::OnFinishedChangingVariable(const FPropertyChangedEvent& InPropertyChangedEvent, UObject* InModifiedObjectInstance)
 {
-	if( !InPropertyChangedEvent.MemberProperty ||
+	if (!InModifiedObjectInstance)
+	{
+		return;
+	}
+
+	ImportNamespacesForPropertyValue(InModifiedObjectInstance->GetClass(), InPropertyChangedEvent.MemberProperty, InModifiedObjectInstance);
+}
+
+void FBlueprintVarActionDetails::OnFinishedChangingLocalVariable(const FPropertyChangedEvent& InPropertyChangedEvent, TSharedPtr<FStructOnScope> InStructData, TWeakObjectPtr<UK2Node_EditablePinBase> InEntryNode)
+{
+	if (!InPropertyChangedEvent.MemberProperty ||
 		!InPropertyChangedEvent.MemberProperty->GetOwnerStruct() ||
 		!InPropertyChangedEvent.MemberProperty->GetOwnerStruct()->IsA<UFunction>())
 	{
@@ -2749,7 +2763,7 @@ void FBlueprintVarActionDetails::OnFinishedChangingProperties(const FPropertyCha
 
 		bDefaultValueSet = FBlueprintEditorUtils::PropertyValueToString(DirectProperty, InStructData->GetStructMemory(), DefaultValueString, FuncEntry);
 
-		if(bDefaultValueSet)
+		if (bDefaultValueSet)
 		{
 			// Search out the correct local variable in the Function Entry Node and set the default value
 			for (FBPVariableDescription& LocalVar : FuncEntry->LocalVariables)
@@ -2767,6 +2781,35 @@ void FBlueprintVarActionDetails::OnFinishedChangingProperties(const FPropertyCha
 				}
 			}
 		}
+
+		ImportNamespacesForPropertyValue(InStructData->GetStruct(), DirectProperty, InStructData->GetStructMemory());
+	}
+}
+
+void FBlueprintVarActionDetails::ImportNamespacesForPropertyValue(const UStruct* InStruct, const FProperty* InProperty, const void* InContainer)
+{
+	if (!GetDefault<UBlueprintEditorSettings>()->bEnableNamespaceImportingFeatures)
+	{
+		return;
+	}
+
+	// Auto-import any namespace(s) associated with the property's value into the current editor context.
+	TSet<FString> AssociatedNamespaces;
+	FBlueprintNamespaceUtilities::GetPropertyValueNamespaces(InStruct, InProperty, InContainer, AssociatedNamespaces);
+	if (AssociatedNamespaces.Num() > 0)
+	{
+		TSharedPtr<SMyBlueprint> MyBlueprintPtr = MyBlueprint.Pin();
+		if (MyBlueprintPtr.IsValid())
+		{
+			TSharedPtr<FBlueprintEditor> BlueprintEditor = MyBlueprintPtr->GetBlueprintEditor().Pin();
+			if (BlueprintEditor.IsValid())
+			{
+				for (const FString& AssociatedNamespace : AssociatedNamespaces)
+				{
+					BlueprintEditor->ImportNamespace(AssociatedNamespace);
+				}
+			}
+		}
 	}
 }
 
diff --git a/Engine/Source/Editor/Kismet/Private/BlueprintDetailsCustomization.h b/Engine/Source/Editor/Kismet/Private/BlueprintDetailsCustomization.h
--- a/Engine/Source/Editor/Kismet/Private/BlueprintDetailsCustomization.h
+++ b/Engine/Source/Editor/Kismet/Private/BlueprintDetailsCustomization.h
@@ -145,6 +145,14 @@ private:
 	void OnVarTypeChanged(const FEdGraphPinType& NewPinType);
 	EVisibility IsTooltipEditVisible() const;
 
+	/**
+	 * Callback when changing a variable property
+	 *
+	 * @param InPropertyChangedEvent	Information on the property changed
+	 * @param InModifiedObjectInstance	The object instance for which the value was changed
+	 */
+	void OnFinishedChangingVariable(const FPropertyChangedEvent& InPropertyChangedEvent, UObject* InModifiedObjectInstance);
+
 	/**
 	 * Callback when changing a local variable property
 	 *
@@ -152,7 +160,16 @@ private:
 	 * @param InStructData				The struct data where the value of the properties are stored
 	 * @param InEntryNode				Entry node where the default values of local variables are stored
 	 */
-	void OnFinishedChangingProperties(const FPropertyChangedEvent& InPropertyChangedEvent, TSharedPtr<FStructOnScope> InStructData, TWeakObjectPtr<UK2Node_EditablePinBase> InEntryNode);
+	void OnFinishedChangingLocalVariable(const FPropertyChangedEvent& InPropertyChangedEvent, TSharedPtr<FStructOnScope> InStructData, TWeakObjectPtr<UK2Node_EditablePinBase> InEntryNode);
+
+	/**
+	 * Auto-import any namespaces associated with a variable's value into the current editor context
+	 *
+	 * @param InStruct					A reference to the container's type
+	 * @param InProperty				A reference to the variable property; owner is expected to match the container's type
+	 * @param InContainer				A pointer to the data container (e.g. struct or object) where the property's value is stored
+	 */
+	void ImportNamespacesForPropertyValue(const UStruct* InStruct, const FProperty* InProperty, const void* InContainer);
 
 	/** Callback to decide if the category drop down menu should be enabled */
 	bool GetVariableCategoryChangeEnabled() const;
diff --git a/Engine/Source/Editor/Kismet/Private/BlueprintNamespaceHelper.cpp b/Engine/Source/Editor/Kismet/Private/BlueprintNamespaceHelper.cpp
--- a/Engine/Source/Editor/Kismet/Private/BlueprintNamespaceHelper.cpp
+++ b/Engine/Source/Editor/Kismet/Private/BlueprintNamespaceHelper.cpp
@@ -8,7 +8,6 @@
 #include "Toolkits/ToolkitManager.h"
 #include "EdGraph/EdGraphPin.h"
 #include "EdGraphSchema_K2.h"
-#include "AssetRegistryModule.h"
 #include "SPinTypeSelector.h"
 #include "Widgets/Input/SCheckBox.h"
 #include "Widgets/Text/STextBlock.h"
@@ -261,20 +260,11 @@ bool FBlueprintNamespaceHelper::IsImportedObject(const UObject* InObject) const
 
 bool FBlueprintNamespaceHelper::IsImportedObject(const FSoftObjectPath& InObjectPath) const
 {
-	if (const UObject* Object = InObjectPath.ResolveObject())
-	{
-		return IsImportedObject(Object);
-	}
-	
-	FString ObjectPathAsString = InObjectPath.ToString();
-	const FAssetRegistryModule& AssetRegistryModule = FModuleManager::LoadModuleChecked<FAssetRegistryModule>(TEXT("AssetRegistry"));
-	FAssetData AssetData = AssetRegistryModule.Get().GetAssetByObjectPath(*ObjectPathAsString);
-	if (!AssetData.IsValid() && ObjectPathAsString.RemoveFromEnd(TEXT("_C")))
-	{
-		AssetData = AssetRegistryModule.Get().GetAssetByObjectPath(*ObjectPathAsString);
-	}
+	// Determine the object's namespace identifier.
+	FString Namespace = FBlueprintNamespaceUtilities::GetObjectNamespace(InObjectPath);
 
-	return IsImportedAsset(AssetData);
+	// Return whether or not the namespace was added, explicitly or otherwise.
+	return IsIncludedInNamespaceList(Namespace);
 }
 
 bool FBlueprintNamespaceHelper::IsImportedAsset(const FAssetData& InAssetData) const
diff --git a/Engine/Source/Editor/Kismet/Private/BlueprintNamespaceUtilities.cpp b/Engine/Source/Editor/Kismet/Private/BlueprintNamespaceUtilities.cpp
--- a/Engine/Source/Editor/Kismet/Private/BlueprintNamespaceUtilities.cpp
+++ b/Engine/Source/Editor/Kismet/Private/BlueprintNamespaceUtilities.cpp
@@ -2,6 +2,7 @@
 
 #include "BlueprintNamespaceUtilities.h"
 #include "Engine/Blueprint.h"
+#include "AssetRegistryModule.h"
 #include "AssetRegistry/AssetData.h"
 #include "EdGraphSchema_K2.h"
 #include "Misc/AssertionMacros.h"
@@ -123,6 +124,100 @@ FString FBlueprintNamespaceUtilities::GetObjectNamespace(const UObject* InObject
 	return Namespace;
 }
 
+FString FBlueprintNamespaceUtilities::GetObjectNamespace(const FSoftObjectPath& InObjectPath)
+{
+	if (const UObject* Object = InObjectPath.ResolveObject())
+	{
+		return GetObjectNamespace(Object);
+	}
+
+	FString ObjectPathAsString = InObjectPath.ToString();
+	const FAssetRegistryModule& AssetRegistryModule = FModuleManager::LoadModuleChecked<FAssetRegistryModule>(TEXT("AssetRegistry"));
+	FAssetData AssetData = AssetRegistryModule.Get().GetAssetByObjectPath(*ObjectPathAsString);
+	if (!AssetData.IsValid() && ObjectPathAsString.RemoveFromEnd(TEXT("_C")))
+	{
+		AssetData = AssetRegistryModule.Get().GetAssetByObjectPath(*ObjectPathAsString);
+	}
+
+	return GetAssetNamespace(AssetData);
+}
+
+void FBlueprintNamespaceUtilities::GetPropertyValueNamespaces(const UStruct* InStruct, const FProperty* InProperty, const void* InContainer, TSet<FString>& OutNamespaces)
+{
+	if (!InStruct || !InProperty || !InContainer)
+	{
+		return;
+	}
+
+	const UStruct* PropertyOwner = InProperty->GetOwnerStruct();
+	if (!PropertyOwner)
+	{
+		return;
+	}
+
+	if (!ensureMsgf(InStruct == PropertyOwner, TEXT("Property %s is a member of struct %s which does not match the given struct %s"), *InProperty->GetName(), *PropertyOwner->GetName(), *InStruct->GetName()))
+	{
+		return;
+	}
+
+	for (int32 ArrayIdx = 0; ArrayIdx < InProperty->ArrayDim; ++ArrayIdx)
+	{
+		const uint8* ValuePtr = InProperty->ContainerPtrToValuePtr<uint8>(InContainer, ArrayIdx);
+
+		if (const FStructProperty* StructProperty = CastField<FStructProperty>(InProperty))
+		{
+			for (TFieldIterator<FProperty> It(StructProperty->Struct); It; ++It)
+			{
+				GetPropertyValueNamespaces(StructProperty->Struct, *It, (*It)->ContainerPtrToValuePtr<uint8>(ValuePtr), OutNamespaces);
+			}
+		}
+		else if (const FArrayProperty* ArrayProperty = CastField<FArrayProperty>(InProperty))
+		{
+			FScriptArrayHelper ArrayHelper(ArrayProperty, ValuePtr);
+			for (int32 ValueIdx = 0; ValueIdx < ArrayHelper.Num(); ++ValueIdx)
+			{
+				GetPropertyValueNamespaces(InStruct, ArrayProperty->Inner, ArrayHelper.GetRawPtr(ValueIdx), OutNamespaces);
+			}
+		}
+		else if (const FSetProperty* SetProperty = CastField<FSetProperty>(InProperty))
+		{
+			FScriptSetHelper SetHelper(SetProperty, ValuePtr);
+			for (int32 ValueIdx = 0; ValueIdx < SetHelper.Num(); ++ValueIdx)
+			{
+				GetPropertyValueNamespaces(InStruct, SetProperty->ElementProp, SetHelper.GetElementPtr(ValueIdx), OutNamespaces);
+			}
+		}
+		else if (const FMapProperty* MapProperty = CastField<FMapProperty>(InProperty))
+		{
+			FScriptMapHelper MapHelper(MapProperty, ValuePtr);
+			for (int32 ValueIdx = 0; ValueIdx < MapHelper.Num(); ++ValueIdx)
+			{
+				const uint8* MapValuePtr = MapHelper.GetPairPtr(ValueIdx);
+				GetPropertyValueNamespaces(InStruct, MapProperty->KeyProp, MapValuePtr, OutNamespaces);
+				GetPropertyValueNamespaces(InStruct, MapProperty->ValueProp, MapValuePtr, OutNamespaces);
+			}
+		}
+		else if (const FSoftObjectProperty* SoftObjectProperty = CastField<FSoftObjectProperty>(InProperty))
+		{
+			const FSoftObjectPath& ObjectPath = SoftObjectProperty->GetPropertyValue(ValuePtr).ToSoftObjectPath();
+			if (ObjectPath.IsValid())
+			{
+				FString Namespace = GetObjectNamespace(ObjectPath);
+				OutNamespaces.Add(Namespace);
+			}
+		}
+		else if (const FObjectPropertyBase* ObjectProperty = CastField<FObjectPropertyBase>(InProperty))
+		{
+			const UObject* ObjectValue = ObjectProperty->GetObjectPropertyValue(ValuePtr);
+			if (ObjectValue)
+			{
+				FString Namespace = GetObjectNamespace(ObjectValue);
+				OutNamespaces.Add(Namespace);
+			}
+		}
+	}
+}
+
 void FBlueprintNamespaceUtilities::SetDefaultBlueprintNamespaceType(EDefaultBlueprintNamespaceType InType)
 {
 	using namespace UE::Editor::Kismet::Private;
diff --git a/Engine/Source/Editor/Kismet/Private/SKismetInspector.cpp b/Engine/Source/Editor/Kismet/Private/SKismetInspector.cpp
--- a/Engine/Source/Editor/Kismet/Private/SKismetInspector.cpp
+++ b/Engine/Source/Editor/Kismet/Private/SKismetInspector.cpp
@@ -45,6 +45,8 @@
 #include "FormatTextDetails.h"
 #include "Kismet2/BlueprintEditorUtils.h"
 #include "ClassViewerFilter.h"
+#include "BlueprintEditorSettings.h"
+#include "BlueprintNamespaceUtilities.h"
 
 #define LOCTEXT_NAMESPACE "KismetInspector"
 
@@ -624,7 +626,8 @@ void SKismetInspector::UpdateFromObjects(const TArray<UObject*>& PropertyObjects
 	}
 
 	PropertyView->OnFinishedChangingProperties().Clear();
-	PropertyView->OnFinishedChangingProperties().Add( UserOnFinishedChangingProperties );
+	PropertyView->OnFinishedChangingProperties().Add(UserOnFinishedChangingProperties);
+	PropertyView->OnFinishedChangingProperties().AddSP(this, &SKismetInspector::OnFinishedChangingProperties);
 
 	// Proceed to update
 	SelectedObjects.Empty();
@@ -998,6 +1001,58 @@ void SKismetInspector::SetPublicViewCheckboxState( ECheckBoxState InIsChecked )
 	BlueprintEditorPtr.Pin()->StartEditingDefaults();
 }
 
+void SKismetInspector::OnFinishedChangingProperties(const FPropertyChangedEvent& InPropertyChangedEvent)
+{
+	ImportNamespacesForPropertyValue(InPropertyChangedEvent.MemberProperty);
+}
+
+void SKismetInspector::ImportNamespacesForPropertyValue(const FProperty* InProperty) const
+{
+	if (!GetDefault<UBlueprintEditorSettings>()->bEnableNamespaceImportingFeatures)
+	{
+		return;
+	}
+
+	if (!InProperty || SelectedObjects.Num() == 0)
+	{
+		return;
+	}
+
+	// Gather all namespace identifier strings associated with the property's value for each edited object.
+	TSet<FString> AssociatedNamespaces;
+	for (const TWeakObjectPtr<UObject>& SelectedObjectPtr : SelectedObjects)
+	{
+		if (const UObject* SelectedObject = SelectedObjectPtr.Get())
+		{
+			const UStruct* SelectedType = SelectedObject->GetClass();
+
+			// Ensure that the selected object type matches the property's owner.
+			// For example, a details customization may select an unrelated object
+			// and then customize each row with an external object reference. In
+			// those cases, the customization would need to handle this explicitly.
+			if (SelectedType != InProperty->GetOwnerStruct())
+			{
+				continue;
+			}
+
+			FBlueprintNamespaceUtilities::GetPropertyValueNamespaces(SelectedType, InProperty, SelectedObject, AssociatedNamespaces);
+		}
+	}
+
+	// Auto-import any namespace(s) associated with the property's value into the current editor context.
+	if (AssociatedNamespaces.Num() > 0)
+	{
+		TSharedPtr<FBlueprintEditor> BlueprintEditor = BlueprintEditorPtr.Pin();
+		if (BlueprintEditor.IsValid())
+		{
+			for (const FString& AssociatedNamespace : AssociatedNamespaces)
+			{
+				BlueprintEditor->ImportNamespace(AssociatedNamespace);
+			}
+		}
+	}
+}
+
 //////////////////////////////////////////////////////////////////////////
 
 #undef LOCTEXT_NAMESPACE
diff --git a/Engine/Source/Editor/Kismet/Public/BlueprintNamespaceUtilities.h b/Engine/Source/Editor/Kismet/Public/BlueprintNamespaceUtilities.h
--- a/Engine/Source/Editor/Kismet/Public/BlueprintNamespaceUtilities.h
+++ b/Engine/Source/Editor/Kismet/Public/BlueprintNamespaceUtilities.h
@@ -31,11 +31,29 @@ public:
 	/**
 	 * Analyzes the given object to determine its explicitly-assigned namespace identifier, or otherwise returns its default namespace.
 	 * 
-	 * @param InObject	A reference to the input object.
+	 * @param InObject		A reference to the input object.
 	 * @return The unique Blueprint namespace identifier associated with the given object, or an empty string if the object belongs to the global namespace (default).
 	 */
 	static FString GetObjectNamespace(const UObject* InObject);
 
+	/**
+	 * Analyzes the given object path to determine its explicitly-assigned namespace identifier, or otherwise returns its default namespace.
+	 *
+	 * @param InObjectPath	Path to the given object (may not be loaded yet).
+	 * @return The unique Blueprint namespace identifier associated with the given object (even if unloaded), or an empty string if the object belongs to the global namespace (default).
+	 */
+	static FString GetObjectNamespace(const FSoftObjectPath& InObjectPath);
+
+	/**
+	 * Analyzes a property value to determine explicitly-assigned namespace identifiers from any object references, or otherwise returns the default namespace for each occurrence (default).
+	 * 
+	 * @param InStruct		A reference to the container's struct/object type.
+	 * @param InProperty	The property for which we will analyze the value.
+	 * @param InContainer	The source address of the struct/object containing the property's value.
+	 * @param OutNamespaces	Zero or more unique namespace identifier(s) referenced by the property value. An entry with an empty string equates to the default global namespace.
+	 */
+	static void GetPropertyValueNamespaces(const UStruct* InStruct, const FProperty* InProperty, const void* InContainer, TSet<FString>& OutNamespaces);
+
 	/**
 	 * Sets the default Blueprint namespace type that objects/assets should use when not explicitly assigned.
 	 *
diff --git a/Engine/Source/Editor/Kismet/Public/SKismetInspector.h b/Engine/Source/Editor/Kismet/Public/SKismetInspector.h
--- a/Engine/Source/Editor/Kismet/Public/SKismetInspector.h
+++ b/Engine/Source/Editor/Kismet/Public/SKismetInspector.h
@@ -208,4 +208,10 @@ protected:
 	void SetPublicViewCheckboxState(ECheckBoxState InIsChecked);
 
 	bool IsAnyParentContainerSelected(const FPropertyAndParent& PropertyAndParent) const;
+
+	/** Callback invoked after a value change on the selected object(s) */
+	void OnFinishedChangingProperties(const FPropertyChangedEvent& InPropertyChangedEvent);
+
+	/** Auto-import any namespaces associated with a property's value into the current editor context */
+	void ImportNamespacesForPropertyValue(const FProperty* InProperty) const;
 };
