Fix Crash when right clicking on Message Line Items

diff --git a/Engine/Source/Editor/Kismet/Private/Debugging/SKismetDebugTreeView.cpp b/Engine/Source/Editor/Kismet/Private/Debugging/SKismetDebugTreeView.cpp
--- a/Engine/Source/Editor/Kismet/Private/Debugging/SKismetDebugTreeView.cpp
+++ b/Engine/Source/Editor/Kismet/Private/Debugging/SKismetDebugTreeView.cpp
@@ -319,29 +319,35 @@ TSharedRef<SWidget> FDebugLineItem::GenerateValueWidget(TSharedPtr<FString> InSe
 
 void FDebugLineItem::MakeMenu(FMenuBuilder& MenuBuilder, bool bInDebuggerTab)
 {
-	const FUIAction CopyName(
-		FExecuteAction::CreateRaw(this, &FDebugLineItem::CopyNameToClipboard),
-		FCanExecuteAction::CreateRaw(this, &FDebugLineItem::HasName)
-	);
+	if (HasName())
+	{
+		const FUIAction CopyName(
+			FExecuteAction::CreateSP(this, &FDebugLineItem::CopyNameToClipboard),
+			FCanExecuteAction::CreateSP(this, &FDebugLineItem::HasName)
+		);
 
-	MenuBuilder.AddMenuEntry(
-		LOCTEXT("CopyName", "Copy Name"),
-		LOCTEXT("CopyName_ToolTip", "Copy name to clipboard"),
-		FSlateIcon(),
-		CopyName
-	);
+		MenuBuilder.AddMenuEntry(
+			LOCTEXT("CopyName", "Copy Name"),
+			LOCTEXT("CopyName_ToolTip", "Copy name to clipboard"),
+			FSlateIcon(),
+			CopyName
+		);
+	}
 
-	const FUIAction CopyValue(
-		FExecuteAction::CreateRaw(this, &FDebugLineItem::CopyValueToClipboard),
-		FCanExecuteAction::CreateRaw(this, &FDebugLineItem::HasValue)
-	);
+	if (HasValue())
+	{
+		const FUIAction CopyValue(
+			FExecuteAction::CreateSP(this, &FDebugLineItem::CopyValueToClipboard),
+			FCanExecuteAction::CreateSP(this, &FDebugLineItem::HasValue)
+		);
 
-	MenuBuilder.AddMenuEntry(
-		LOCTEXT("CopyValue", "Copy Value"),
-		LOCTEXT("CopyValue_ToolTip", "Copy value to clipboard"),
-		FSlateIcon(),
-		CopyValue
-	);
+		MenuBuilder.AddMenuEntry(
+			LOCTEXT("CopyValue", "Copy Value"),
+			LOCTEXT("CopyValue_ToolTip", "Copy value to clipboard"),
+			FSlateIcon(),
+			CopyValue
+		);
+	}
 
 	ExtendContextMenu(MenuBuilder, bInDebuggerTab);
 }
