Add DebugInfo to differentiate watches of individual elements of containers

diff --git a/Engine/Source/Editor/Kismet/Private/Debugging/SKismetDebugTreeView.cpp b/Engine/Source/Editor/Kismet/Private/Debugging/SKismetDebugTreeView.cpp
--- a/Engine/Source/Editor/Kismet/Private/Debugging/SKismetDebugTreeView.cpp
+++ b/Engine/Source/Editor/Kismet/Private/Debugging/SKismetDebugTreeView.cpp
@@ -1243,9 +1243,10 @@ struct FWatchLineItem : public FLineItemWithChildren
 protected:
 	TWeakObjectPtr< UObject > ParentObjectRef;
 	const FEdGraphPinReference ObjectRef;
-	const TArray<FName> PathToProperty;
 	mutable TSharedPtr<FPropertyInstanceInfo> CachedPropertyInfo = nullptr;
 	mutable FText CachedErrorString;
+	const TArray<FName> PathToProperty;
+
 public:
 	FWatchLineItem(const UEdGraphPin* PinToWatch, UObject* ParentObject)
 		: FLineItemWithChildren(DLT_Watch)
@@ -1312,12 +1313,20 @@ public:
 		return ObjectRef;
 	}
 
+	const TArray<FName>& GetPathToProperty() const
+	{
+		return PathToProperty;
+	}
+
 	virtual FText GetDescription() const override;
 	virtual TSharedRef<SWidget> GenerateNameWidget(TSharedPtr<FString> InSearchString) override;
 	virtual TSharedRef<SWidget> GenerateValueWidget(TSharedPtr<FString> InSearchString) override;
 	virtual TSharedRef<SWidget> GetValueIcon() override;
 	virtual TSharedRef<SWidget> GetNameIcon() override;
 
+	/** Returns a SharedPtr to the debug info for the property being represented by this TreeItem */
+	TSharedPtr<FPropertyInstanceInfo> GetPropertyInfo() const;
+
 protected:
 	virtual FDebugLineItem* Duplicate() const override
 	{
@@ -1362,8 +1371,6 @@ protected:
 	// Finds the DebugInfo for the property pointed to by PathToProperty from the Pin's debuginfo
 	TSharedPtr<FPropertyInstanceInfo> GetWatchedPropertyDebugInfo(const TSharedPtr<FPropertyInstanceInfo>& InPinInfo) const;
 
-	TSharedPtr<FPropertyInstanceInfo> GetPropertyInfo() const;
-
 private:
 	void AddWatch() const
 	{
@@ -1813,12 +1820,28 @@ TSharedPtr<FPropertyInstanceInfo> FWatchLineItem::GetWatchedPropertyDebugInfo(co
 		if (ThisDebugInfo.IsValid())
 		{
 			const FString ChildNameStr = ChildName.ToString();
-			TSharedPtr<FPropertyInstanceInfo>* FoundChild = ThisDebugInfo->Children.FindByPredicate(
-				[&ChildNameStr](const TSharedPtr<FPropertyInstanceInfo>& Child)
-				{
-					return Child->Property->GetAuthoredName() == ChildNameStr;
-				}
-			);
+			FProperty* Property = ThisDebugInfo->Property.Get();
+
+			TSharedPtr<FPropertyInstanceInfo>* FoundChild;
+			
+			if (Property->IsA<FSetProperty>() || Property->IsA<FArrayProperty>() || Property->IsA<FMapProperty>())
+			{
+				FoundChild = ThisDebugInfo->Children.FindByPredicate(
+					[&ChildNameStr](const TSharedPtr<FPropertyInstanceInfo>& Child)
+					{
+						return Child->DisplayName.ToString() == ChildNameStr;
+					}
+				);
+			}
+			else
+			{
+				FoundChild = ThisDebugInfo->Children.FindByPredicate(
+					[&ChildNameStr](const TSharedPtr<FPropertyInstanceInfo>& Child)
+					{
+						return Child->Property->GetAuthoredName() == ChildNameStr;
+					}
+				);
+			}
 
 			if (FoundChild)
 			{
@@ -1954,7 +1977,22 @@ UEdGraphPin* FWatchChildLineItem::BuildPathToProperty(TArray<FName>& OutPathToPr
 {
 	UEdGraphPin* PinToWatch = nullptr;
 	OutPathToProperty.Reset();
-	OutPathToProperty.Emplace(Data->Property->GetAuthoredName());
+
+	auto AddPropertyToPath = [&OutPathToProperty](const TSharedRef<FPropertyInstanceInfo>& InPropertyInfo)
+	{
+		if (InPropertyInfo->bIsInContainer)
+		{
+			// display name of map elements is their key exported to text
+			OutPathToProperty.Emplace(InPropertyInfo->DisplayName.ToString());
+		}
+		else
+		{
+			OutPathToProperty.Emplace(InPropertyInfo->Property->GetAuthoredName());
+		}
+	};
+
+	AddPropertyToPath(Data);
+
 	FDebugTreeItemPtr Parent = ParentTreeItem.Pin();
 	while (Parent.IsValid())
 	{
@@ -1963,6 +2001,14 @@ UEdGraphPin* FWatchChildLineItem::BuildPathToProperty(TArray<FName>& OutPathToPr
 			// this is a FWatchLineItem
 			TSharedPtr<FWatchLineItem> ParentAsWatch = StaticCastSharedPtr<FWatchLineItem>(Parent);
 			PinToWatch = ParentAsWatch->GetPin().Get();
+
+			// Add the original path to the front of ours (we're still constructing it backwards) 
+			const TArray<FName>& ParentPath = ParentAsWatch->GetPathToProperty();
+			for (int32 PathIndex = ParentPath.Num() - 1; PathIndex >= 0; --PathIndex)
+			{
+				OutPathToProperty.Emplace(ParentPath[PathIndex]);
+			}
+
 			break;
 		}
 
@@ -1970,7 +2016,9 @@ UEdGraphPin* FWatchChildLineItem::BuildPathToProperty(TArray<FName>& OutPathToPr
 		{
 			// This is a FWatchChildLineItem
 			TSharedPtr<FWatchChildLineItem> ParentAsChild = StaticCastSharedPtr<FWatchChildLineItem>(Parent);
-			OutPathToProperty.Emplace(ParentAsChild->Data->Property->GetAuthoredName());
+
+			AddPropertyToPath(ParentAsChild->Data);
+
 			Parent = ParentAsChild->ParentTreeItem.Pin();
 		}
 		else
diff --git a/Engine/Source/Editor/UnrealEd/Private/Kismet2/KismetDebugUtilities.cpp b/Engine/Source/Editor/UnrealEd/Private/Kismet2/KismetDebugUtilities.cpp
--- a/Engine/Source/Editor/UnrealEd/Private/Kismet2/KismetDebugUtilities.cpp
+++ b/Engine/Source/Editor/UnrealEd/Private/Kismet2/KismetDebugUtilities.cpp
@@ -2111,6 +2111,7 @@ void FPropertyInstanceInfo::PopulateChildren(FPropertyInstance PropertyInstance,
 			
 			// overwrite the display name with the array index for the current element
 			ChildInfo->DisplayName = FText::Format(LOCTEXT("ArrayIndexName", "[{0}]"), FText::AsNumber(i));
+			ChildInfo->bIsInContainer = true;
 			Children.Add(ChildInfo);
 		}
 	}
@@ -2154,6 +2155,7 @@ void FPropertyInstanceInfo::PopulateChildren(FPropertyInstance PropertyInstance,
 				);
 				NameStr += TEXT("] ");
 				ChildInfo->DisplayName = FText::FromString(NameStr);
+				ChildInfo->bIsInContainer = true;
 				
 				Children.Add(ChildInfo);
 				--Count;
@@ -2177,6 +2179,7 @@ void FPropertyInstanceInfo::PopulateChildren(FPropertyInstance PropertyInstance,
 				const TSharedPtr<FPropertyInstanceInfo> ChildInfo = FindOrMake(ChildProperty, VisitedNodes);
 				
 				ChildInfo->DisplayName = FText::Format(LOCTEXT("SetIndexName", "[{0}]"), FText::AsNumber(Index));
+				ChildInfo->bIsInContainer = true;
 
 				Children.Add(ChildInfo);
 
diff --git a/Engine/Source/Editor/UnrealEd/Public/Kismet2/KismetDebugUtilities.h b/Engine/Source/Editor/UnrealEd/Public/Kismet2/KismetDebugUtilities.h
--- a/Engine/Source/Editor/UnrealEd/Public/Kismet2/KismetDebugUtilities.h
+++ b/Engine/Source/Editor/UnrealEd/Public/Kismet2/KismetDebugUtilities.h
@@ -85,6 +85,7 @@ struct FPropertyInstanceInfo
 	TWeakObjectPtr<UObject> Object = nullptr; // only filled if property is a UObject
 	TFieldPath<FProperty> Property;
 	TArray<TSharedPtr<FPropertyInstanceInfo>> Children;
+	bool bIsInContainer = false;
 };
 
 inline uint32 GetTypeHash(const FPropertyInstanceInfo::FPropertyInstance& PropertyInstance)
