Update styling for add/view event buttons

diff --git a/Engine/Source/Editor/DetailCustomizations/Private/ActorDetails.cpp b/Engine/Source/Editor/DetailCustomizations/Private/ActorDetails.cpp
--- a/Engine/Source/Editor/DetailCustomizations/Private/ActorDetails.cpp
+++ b/Engine/Source/Editor/DetailCustomizations/Private/ActorDetails.cpp
@@ -502,7 +502,7 @@ void FActorDetails::AddEventsCategory(IDetailLayoutBuilder& DetailBuilder)
 			FText EventText = Property->GetDisplayNameText();
 
 			EventsCategory.AddCustomRow(EventText)
-			.NameContent()
+			.WholeRowContent()
 			[
 				SNew(SHorizontalBox)
 				.ToolTipText(Property->GetToolTipText())
@@ -513,7 +513,7 @@ void FActorDetails::AddEventsCategory(IDetailLayoutBuilder& DetailBuilder)
 				.Padding(0.0f, 0.0f, 5.0f, 0.0f)
 				[
 					SNew(SImage)
-					.Image(FEditorStyle::GetBrush("GraphEditor.Event_16x"))
+					.Image(FAppStyle::Get().GetBrush("GraphEditor.Event_16x"))
 				]
 
 				+ SHorizontalBox::Slot()
@@ -523,21 +523,24 @@ void FActorDetails::AddEventsCategory(IDetailLayoutBuilder& DetailBuilder)
 					.Font(IDetailLayoutBuilder::GetDetailFont())
 					.Text(EventText)
 				]
-			]
-			// A green "Plus" button to add a binding. For dynamic delegates on the CDO, you can always
-			// make a new binding, so always display the "Plus"
-			.ValueContent()
-			.MinDesiredWidth(150.0f)
-			.MaxDesiredWidth(200.0f)
-			[
-				SNew(SButton)
-				.ButtonStyle(FEditorStyle::Get(), "FlatButton.Success")
-				.HAlign(HAlign_Center)
-				.OnClicked(this, &FActorDetails::HandleAddOrViewEventForVariable, Blueprint, Property)
-				.ForegroundColor(FSlateColor::UseForeground())
-				[			
-					SNew(SImage)
-					.Image(FEditorStyle::GetBrush("Plus"))			
+
+				+ SHorizontalBox::Slot()
+				.HAlign(HAlign_Left)
+				.VAlign(VAlign_Center)
+				.Padding(0)
+				[
+					// A "Plus" button to add a binding. For dynamic delegates on the CDO, you can always
+					// make a new binding, so always display the "Plus"
+					SNew(SButton)
+					.ButtonStyle(FAppStyle::Get(), "DetailsView.CalloutButton")
+					.ContentPadding(FMargin(3.0, 2.0))
+					.HAlign(HAlign_Center)
+					.OnClicked(this, &FActorDetails::HandleAddOrViewEventForVariable, Blueprint, Property)
+					[
+						SNew(SImage)
+						.ColorAndOpacity(FSlateColor::UseForeground())
+						.Image(FAppStyle::Get().GetBrush("Icons.Plus"))
+					]
 				]
 			];
 		}
@@ -625,10 +628,10 @@ void FActorDetails::AddActorCategory( IDetailLayoutBuilder& DetailBuilder, const
 		.NameContent()
 		[
 			SNew(SHyperlink)
-				.Style(FEditorStyle::Get(), "HoverOnlyHyperlink")
+				.Style(FAppStyle::Get(), "HoverOnlyHyperlink")
 				.OnNavigate(this, &FActorDetails::OnNarrowSelectionSetToSpecificLevel, WeakLevelPtr)
 				.Text(ActorCountDescription)
-				.TextStyle(FEditorStyle::Get(), "DetailsView.HyperlinkStyle")
+				.TextStyle(FAppStyle::Get(), "DetailsView.HyperlinkStyle")
 				.ToolTipText(Tooltip)
 		]
 		.ValueContent()
diff --git a/Engine/Source/Editor/Kismet/Private/BlueprintDetailsCustomization.cpp b/Engine/Source/Editor/Kismet/Private/BlueprintDetailsCustomization.cpp
--- a/Engine/Source/Editor/Kismet/Private/BlueprintDetailsCustomization.cpp
+++ b/Engine/Source/Editor/Kismet/Private/BlueprintDetailsCustomization.cpp
@@ -132,7 +132,7 @@ void FBlueprintDetails::AddEventsCategory(IDetailLayoutBuilder& DetailBuilder, F
 					IDetailCategoryBuilder& EventCategory = DetailBuilder.EditCategory(TEXT("Events"), LOCTEXT("Events", "Events"), ECategoryPriority::Uncommon);
 
 					EventCategory.AddCustomRow(EventText)
-					.NameContent()
+					.WholeRowContent()
 					[
 						SNew(SHorizontalBox)
 						.ToolTipText(Property->GetToolTipText())
@@ -143,7 +143,7 @@ void FBlueprintDetails::AddEventsCategory(IDetailLayoutBuilder& DetailBuilder, F
 						.Padding(0.0f, 0.0f, 5.0f, 0.0f)
 						[
 							SNew(SImage)
-							.Image(FEditorStyle::GetBrush("GraphEditor.Event_16x"))
+							.Image(FAppStyle::Get().GetBrush("GraphEditor.Event_16x"))
 						]
 
 						+ SHorizontalBox::Slot()
@@ -153,31 +153,33 @@ void FBlueprintDetails::AddEventsCategory(IDetailLayoutBuilder& DetailBuilder, F
 							.Font(IDetailLayoutBuilder::GetDetailFont())
 							.Text(EventText)
 						]
-					]
-					.ValueContent()
-					.MinDesiredWidth(150.0f)
-					.MaxDesiredWidth(200.0f)
-					[
-						SNew(SButton)
-						.ButtonStyle(FEditorStyle::Get(), "FlatButton.Success")
-						.HAlign(HAlign_Center)
-						.OnClicked(this, &FBlueprintVarActionDetails::HandleAddOrViewEventForVariable, EventName, PropertyName, MakeWeakObjectPtr(PropertyClass))
-						.ForegroundColor(FSlateColor::UseForeground())
-						[
-							SNew(SWidgetSwitcher)
-							.WidgetIndex(this, &FBlueprintVarActionDetails::HandleAddOrViewIndexForButton, EventName, PropertyName)
 
-							+ SWidgetSwitcher::Slot()
+						+ SHorizontalBox::Slot()
+						.HAlign(HAlign_Left)
+						.VAlign(VAlign_Center)
+						.Padding(0)
+						[
+							SNew(SButton)
+							.ContentPadding(FMargin(3.0, 2.0))
+							.ButtonStyle(FAppStyle::Get(), "DetailsView.CalloutButton")
+							.OnClicked(this, &FBlueprintVarActionDetails::HandleAddOrViewEventForVariable, EventName, PropertyName, MakeWeakObjectPtr(PropertyClass))
 							[
-								SNew(STextBlock)
-								.Font(FEditorStyle::GetFontStyle(TEXT("BoldFont")))
-								.Text(LOCTEXT("ViewEvent", "View"))
-							]
+								SNew(SWidgetSwitcher)
+								.WidgetIndex(this, &FBlueprintVarActionDetails::HandleAddOrViewIndexForButton, EventName, PropertyName)
 
-							+ SWidgetSwitcher::Slot()
-							[
-								SNew(SImage)
-								.Image(FEditorStyle::GetBrush("Plus"))
+								+ SWidgetSwitcher::Slot()
+								[
+									SNew(SImage)
+									.ColorAndOpacity(FSlateColor::UseForeground())
+									.Image(FAppStyle::Get().GetBrush("Icons.SelectInViewport"))
+								]
+
+								+ SWidgetSwitcher::Slot()
+								[
+									SNew(SImage)
+									.ColorAndOpacity(FSlateColor::UseForeground())
+									.Image(FAppStyle::Get().GetBrush("Icons.Plus"))
+								]
 							]
 						]
 					];
@@ -594,7 +596,7 @@ void FBlueprintVarActionDetails::CustomizeDetails( IDetailLayoutBuilder& DetailL
 		.ButtonContent()
 		[
 			SNew(SBorder)
-			.BorderImage( FEditorStyle::GetBrush("NoBorder") )
+			.BorderImage(FAppStyle::Get().GetBrush("NoBorder") )
 			.Padding(FMargin(0, 0, 5, 0))
 			[
 				SNew(SEditableTextBox)
@@ -2877,12 +2879,12 @@ public:
 		if (IsValidTarget)
 		{
 			CurrentHoverText = FText::Format(LOCTEXT("MovePinHere", "Move '{PinName}' Here"), Args);
-			CurrentIconBrush = FEditorStyle::GetBrush("Graph.ConnectorFeedback.OK");
+			CurrentIconBrush = FAppStyle::Get().GetBrush("Graph.ConnectorFeedback.OK");
 		}
 		else
 		{
 			CurrentHoverText = FText::Format(LOCTEXT("CannotMovePinHere", "Cannot Move '{PinName}' Here"), Args);
-			CurrentIconBrush = FEditorStyle::GetBrush("Graph.ConnectorFeedback.Error");
+			CurrentIconBrush = FAppStyle::Get().GetBrush("Graph.ConnectorFeedback.Error");
 		}
 	}
 
@@ -3565,7 +3567,7 @@ void FBlueprintGraphActionDetails::CustomizeDetails( IDetailLayoutBuilder& Detai
 						.ButtonContent()
 						[
 							SNew(SBorder)
-							.BorderImage( FEditorStyle::GetBrush("NoBorder") )
+							.BorderImage(FAppStyle::Get().GetBrush("NoBorder") )
 							.Padding(FMargin(0, 0, 5, 0))
 							[
 								SNew(SEditableTextBox)
@@ -4004,7 +4006,7 @@ void FBlueprintGraphActionDetails::CustomizeDetails( IDetailLayoutBuilder& Detai
 			[
 
 				SNew(SButton)
-				.ButtonStyle(FEditorStyle::Get(), "SimpleButton")
+				.ButtonStyle(FAppStyle::Get(), "SimpleButton")
 				.ContentPadding(FMargin(1, 0))
 				.OnClicked(this, &FBlueprintGraphActionDetails::OnAddNewInputClicked)
 				.Visibility(this, &FBlueprintGraphActionDetails::GetAddNewInputOutputVisibility)
@@ -4037,7 +4039,7 @@ void FBlueprintGraphActionDetails::CustomizeDetails( IDetailLayoutBuilder& Detai
 				.Padding(FMargin(0,0,2,0))
 				[
 					SNew(SButton)
-					.ButtonStyle(FEditorStyle::Get(), "SimpleButton")
+					.ButtonStyle(FAppStyle::Get(), "SimpleButton")
 					.ContentPadding(FMargin(1, 0))
 					.OnClicked(this, &FBlueprintGraphActionDetails::OnAddNewOutputClicked)
 					.Visibility(this, &FBlueprintGraphActionDetails::GetAddNewInputOutputVisibility)
@@ -4424,8 +4426,8 @@ void FBlueprintDelegateActionDetails::CustomizeDetails( IDetailLayoutBuilder& De
 		.AutoWidth()
 		[
 			SNew(SButton)
-			.ButtonStyle(FEditorStyle::Get(), "RoundButton")
-			.ForegroundColor(FEditorStyle::GetSlateColor("DefaultForeground"))
+			.ButtonStyle(FAppStyle::Get(), "RoundButton")
+			.ForegroundColor(FAppStyle::Get().GetSlateColor("DefaultForeground"))
 			.ContentPadding(FMargin(2, 0))
 			.OnClicked(this, &FBlueprintDelegateActionDetails::OnAddNewInputClicked)
 			.HAlign(HAlign_Right)
@@ -4439,7 +4441,7 @@ void FBlueprintDelegateActionDetails::CustomizeDetails( IDetailLayoutBuilder& De
 				.Padding(FMargin(0, 1))
 				[
 					SNew(SImage)
-					.Image(FEditorStyle::GetBrush("Plus"))
+					.Image(FAppStyle::Get().GetBrush("Plus"))
 				]
 				+SHorizontalBox::Slot()
 				.VAlign(VAlign_Center)
@@ -6261,7 +6263,7 @@ void FBlueprintComponentDetails::CustomizeDetails(IDetailLayoutBuilder& DetailLa
 			.ButtonContent()
 			[
 				SNew(SBorder)
-				.BorderImage(FEditorStyle::GetBrush("NoBorder"))
+				.BorderImage(FAppStyle::Get().GetBrush("NoBorder"))
 				.Padding(FMargin(0, 0, 5, 0))
 				[
 					SNew(SEditableTextBox)
@@ -7093,7 +7095,7 @@ void FBlueprintDocumentationDetails::CustomizeDetails(IDetailLayoutBuilder& Deta
 			.ButtonContent()
 			[
 				SNew(SBorder)
-				.BorderImage( FEditorStyle::GetBrush( "NoBorder" ))
+				.BorderImage(FAppStyle::Get().GetBrush( "NoBorder" ))
 				.Padding( FMargin( 0, 0, 5, 0 ))
 				[
 					SNew( STextBlock )
diff --git a/Engine/Source/Editor/UMGEditor/Private/Customizations/UMGDetailCustomizations.cpp b/Engine/Source/Editor/UMGEditor/Private/Customizations/UMGDetailCustomizations.cpp
--- a/Engine/Source/Editor/UMGEditor/Private/Customizations/UMGDetailCustomizations.cpp
+++ b/Engine/Source/Editor/UMGEditor/Private/Customizations/UMGDetailCustomizations.cpp
@@ -50,8 +50,8 @@ public:
 		ChildSlot
 		[
 			SNew(SButton)
-			.ButtonStyle(FEditorStyle::Get(), "FlatButton.Success")
-			.TextStyle(FEditorStyle::Get(), "NormalText")
+			.ButtonStyle(FAppStyle::Get(), "FlatButton.Success")
+			.TextStyle(FAppStyle::Get(), "NormalText")
 			.HAlign(HAlign_Center)
 			.ForegroundColor(FSlateColor::UseForeground())
 			.ToolTipText(Action->GetTooltipDescription())
@@ -488,11 +488,11 @@ TSharedRef<SWidget> FBlueprintWidgetCustomization::MakePropertyBindingWidget(TWe
 					{
 						if ( Binding.Kind == EBindingKind::Function )
 						{
-							return FEditorStyle::GetBrush(FunctionIcon);
+							return FAppStyle::Get().GetBrush(FunctionIcon);
 						}
 						else // Property
 						{
-							return FEditorStyle::GetBrush(PropertyIcon);
+							return FAppStyle::Get().GetBrush(PropertyIcon);
 						}
 					}
 				}
@@ -586,7 +586,7 @@ void FBlueprintWidgetCustomization::CreateEventCustomization( IDetailLayoutBuild
 			.Padding(0,0,5,0)
 			[
 				SNew(SImage)
-				.Image(FEditorStyle::GetBrush("GraphEditor.Event_16x"))
+				.Image(FAppStyle::Get().GetBrush("GraphEditor.Event_16x"))
 			]
 
 			+ SHorizontalBox::Slot()
@@ -695,7 +695,7 @@ void FBlueprintWidgetCustomization::CreateMulticastEventCustomization(IDetailLay
 	IDetailCategoryBuilder& EventCategory = DetailLayout.EditCategory(TEXT("Events"), LOCTEXT("Events", "Events"), ECategoryPriority::Uncommon);
 
 	EventCategory.AddCustomRow(EventText)
-		.NameContent()
+		.WholeRowContent()
 		[
 			SNew(SHorizontalBox)
 			.ToolTipText(DelegateProperty->GetToolTipText())
@@ -705,39 +705,42 @@ void FBlueprintWidgetCustomization::CreateMulticastEventCustomization(IDetailLay
 			.Padding(0, 0, 5, 0)
 			[
 				SNew(SImage)
-				.Image(FEditorStyle::GetBrush("GraphEditor.Event_16x"))
+				.Image(FAppStyle::Get().GetBrush("GraphEditor.Event_16x"))
 			]
 
 			+ SHorizontalBox::Slot()
 			.VAlign(VAlign_Center)
+			.HAlign(HAlign_Left)
 			[
 				SNew(STextBlock)
 				.Font(IDetailLayoutBuilder::GetDetailFont())
 				.Text(EventText)
 			]
-		]
-		.ValueContent()
-		.MinDesiredWidth(150)
-		.MaxDesiredWidth(200)
-		[
-			SNew(SButton)
-			.ButtonStyle(FEditorStyle::Get(), "FlatButton.Success")
-			.HAlign(HAlign_Center)
-			.OnClicked(this, &FBlueprintWidgetCustomization::HandleAddOrViewEventForVariable, EventName, PropertyName, MakeWeakObjectPtr(PropertyClass))
-			.ForegroundColor(FSlateColor::UseForeground())
+
+			+ SHorizontalBox::Slot()
+			.HAlign(HAlign_Left)
+			.VAlign(VAlign_Center)
+			.Padding(0)
 			[
-				SNew(SWidgetSwitcher)
-				.WidgetIndex(this, &FBlueprintWidgetCustomization::HandleAddOrViewIndexForButton, EventName, PropertyName)
-				+ SWidgetSwitcher::Slot()
-				[
-					SNew(STextBlock)
-					.Font(FEditorStyle::GetFontStyle(TEXT("BoldFont")))
-					.Text(LOCTEXT("ViewEvent", "View"))
-				]
-				+ SWidgetSwitcher::Slot()
+				SNew(SButton)
+				.ButtonStyle(FAppStyle::Get(), "DetailsView.CalloutButton")
+				.ContentPadding(FMargin(3.0, 2.0))
+				.OnClicked(this, &FBlueprintWidgetCustomization::HandleAddOrViewEventForVariable, EventName, PropertyName, MakeWeakObjectPtr(PropertyClass))
 				[
-					SNew(SImage)
-					.Image(FEditorStyle::GetBrush("Plus"))
+					SNew(SWidgetSwitcher)
+					.WidgetIndex(this, &FBlueprintWidgetCustomization::HandleAddOrViewIndexForButton, EventName, PropertyName)
+					+ SWidgetSwitcher::Slot()
+					[
+						SNew(SImage)
+						.ColorAndOpacity(FSlateColor::UseForeground())
+						.Image(FAppStyle::Get().GetBrush("Icons.SelectInViewport"))
+					]
+					+ SWidgetSwitcher::Slot()
+					[
+						SNew(SImage)
+						.ColorAndOpacity(FSlateColor::UseForeground())
+						.Image(FAppStyle::Get().GetBrush("Icons.Plus"))
+					]
 				]
 			]
 		];
