Find In Blueprints Reskin - Makes use of parent row styling

diff --git a/Engine/Source/Editor/Kismet/Private/FindInBlueprints.cpp b/Engine/Source/Editor/Kismet/Private/FindInBlueprints.cpp
--- a/Engine/Source/Editor/Kismet/Private/FindInBlueprints.cpp
+++ b/Engine/Source/Editor/Kismet/Private/FindInBlueprints.cpp
@@ -1175,16 +1175,12 @@ TSharedRef<ITableRow> SFindInBlueprints::OnGenerateRow( FSearchResult InItem, co
 	if (bIsACategoryWidget)
 	{
 		return SNew( STableRow< TSharedPtr<FFindInBlueprintsResult> >, OwnerTable )
+			.Style( &FAppStyle::Get().GetWidgetStyle<FTableRowStyle>("ShowParentsTableView.Row") )
+			.Padding(FMargin(2.f, 3.f, 2.f, 3.f))
 			[
-				SNew(SBorder)
-				.VAlign(VAlign_Center)
-				.BorderImage(FAppStyle::Get().GetBrush("Brushes.Header"))
-				.Padding(FMargin(2.0f))
-				[
-					SNew(STextBlock)
-					.Text(InItem.Get(), &FFindInBlueprintsResult::GetDisplayString)
-					.ToolTipText(LOCTEXT("BlueprintCatSearchToolTip", "Blueprint"))
-				]
+				SNew(STextBlock)
+				.Text(InItem.Get(), &FFindInBlueprintsResult::GetDisplayString)
+				.ToolTipText(LOCTEXT("BlueprintCatSearchToolTip", "Blueprint"))
 			];
 	}
 	else // Functions/Event/Pin widget
@@ -1206,6 +1202,7 @@ TSharedRef<ITableRow> SFindInBlueprints::OnGenerateRow( FSearchResult InItem, co
 		FText Tooltip = FText::Format(LOCTEXT("BlueprintResultSearchToolTip", "{Category} : {DisplayTitle}"), Args);
 
 		return SNew( STableRow< TSharedPtr<FFindInBlueprintsResult> >, OwnerTable )
+			.Style( &FAppStyle::Get().GetWidgetStyle<FTableRowStyle>("ShowParentsTableView.Row") )
 			[
 				SNew(SHorizontalBox)
 				+SHorizontalBox::Slot()
