Fix visual lag when summoning Pin Value Inspector Tooltip

diff --git a/Engine/Source/Editor/Kismet/Private/SPinValueInspector.cpp b/Engine/Source/Editor/Kismet/Private/SPinValueInspector.cpp
--- a/Engine/Source/Editor/Kismet/Private/SPinValueInspector.cpp
+++ b/Engine/Source/Editor/Kismet/Private/SPinValueInspector.cpp
@@ -257,6 +257,7 @@ TWeakPtr<FPinValueInspectorTooltip> FPinValueInspectorTooltip::SummonTooltip(FEd
 
 			Instance = MakeShared<FPinValueInspectorTooltip>();
 			TooltipWindow->ShowWindow();
+			TooltipWindow->SetAllowFastUpdate(true);
 			TooltipWidget->SetContentWidget(SAssignNew(Instance->ValueInspectorWidget, SPinValueInspector, InPinRef));
 			return Instance;
 		}
@@ -293,6 +294,7 @@ void FPinValueInspectorTooltip::DismissTooltip()
 	}
 	TooltipWidget->ResetContentWidget();
 	TooltipWindow->HideWindow();
+	TooltipWindow->SetAllowFastUpdate(false);
 	Instance.Reset();
 }
 
