Remove disabled View In Debugger option in the Debugger window and remove Start/Stop watching options from properties that don't have pin parents

diff --git a/Engine/Source/Editor/Kismet/Private/Debugging/SKismetDebugTreeView.cpp b/Engine/Source/Editor/Kismet/Private/Debugging/SKismetDebugTreeView.cpp
--- a/Engine/Source/Editor/Kismet/Private/Debugging/SKismetDebugTreeView.cpp
+++ b/Engine/Source/Editor/Kismet/Private/Debugging/SKismetDebugTreeView.cpp
@@ -844,31 +844,37 @@ public:
 	{
 		//Navigate to Class source
 
-		//Add Watch
-		FUIAction AddThisWatch(
-			FExecuteAction::CreateSP(this, &FWatchChildLineItem::AddWatch),
-			FCanExecuteAction::CreateSP(this, &FWatchChildLineItem::CanAddWatch)
-		);
+		// Only add watch options if this has a pin in it's parent chain
+		// (ok to discard the path, this only runs when a context menu is constructed)
+		TArray<FName> PathToProperty;
+		if (UEdGraphPin* PinParent = BuildPathToProperty(PathToProperty))
+		{
+			//Add Watch
+			FUIAction AddThisWatch(
+				FExecuteAction::CreateSP(this, &FWatchChildLineItem::AddWatch),
+				FCanExecuteAction::CreateSP(this, &FWatchChildLineItem::CanAddWatch)
+			);
 
-		MenuBuilder.AddMenuEntry(
-			LOCTEXT("AddPropertyWatch", "Start Watching"),
-			LOCTEXT("AddPropertyWatchTooltip", "Start Watching This Variable"),
-			FSlateIcon(),
-			AddThisWatch
-		);
+				MenuBuilder.AddMenuEntry(
+					LOCTEXT("AddPropertyWatch", "Start Watching"),
+					LOCTEXT("AddPropertyWatchTooltip", "Start Watching This Variable"),
+					FSlateIcon(),
+					AddThisWatch
+				);
 
-		//Add Watch
-		FUIAction ClearThisWatch(
-			FExecuteAction::CreateSP(this, &FWatchChildLineItem::ClearWatch),
-			FCanExecuteAction::CreateSP(this, &FWatchChildLineItem::CanClearWatch)
-		);
+			//Add Watch
+			FUIAction ClearThisWatch(
+				FExecuteAction::CreateSP(this, &FWatchChildLineItem::ClearWatch),
+				FCanExecuteAction::CreateSP(this, &FWatchChildLineItem::CanClearWatch)
+			);
 
-		MenuBuilder.AddMenuEntry(
-			LOCTEXT("ClearPropertyWatch", "Stop Watching"),
-			LOCTEXT("ClearPropertyWatchTooltip", "Stop Watching This Variable"),
-			FSlateIcon(),
-			ClearThisWatch
-		);
+			MenuBuilder.AddMenuEntry(
+				LOCTEXT("ClearPropertyWatch", "Stop Watching"),
+				LOCTEXT("ClearPropertyWatchTooltip", "Stop Watching This Variable"),
+				FSlateIcon(),
+				ClearThisWatch
+			);
+		}
 
 		if (!bInDebuggerTab)
 		{
@@ -2964,6 +2970,7 @@ protected:
 void SKismetDebugTreeView::Construct(const FArguments& InArgs)
 {
 	bFilteredItemsDirty = false;
+	bInDebuggerTab = InArgs._InDebuggerTab;
 	SearchString = MakeShared<FString>();
 
 	ChildSlot
