Fix visual lag when summoning SPinValueInspector

diff --git a/Engine/Source/Editor/Kismet/Private/SPinValueInspector.cpp b/Engine/Source/Editor/Kismet/Private/SPinValueInspector.cpp
--- a/Engine/Source/Editor/Kismet/Private/SPinValueInspector.cpp
+++ b/Engine/Source/Editor/Kismet/Private/SPinValueInspector.cpp
@@ -18,6 +18,7 @@
 
 TSharedPtr<SWindow> FPinValueInspectorTooltip::TooltipWindow = nullptr;
 TSharedPtr<SToolTip> FPinValueInspectorTooltip::TooltipWidget = nullptr;
+TSharedPtr<SPinValueInspector> FPinValueInspectorTooltip::ValueInspectorWidget = nullptr;
 TSharedPtr<FPinValueInspectorTooltip> FPinValueInspectorTooltip::Instance = nullptr;
 
 ///////////////////////////////////////////////////////////////////////////////
@@ -113,10 +114,8 @@ private:
 	mutable TOptional<float> CachedXSize;
 };
 
-void SPinValueInspector::Construct(const FArguments& InArgs, const FEdGraphPinReference& InPinRef)
+void SPinValueInspector::Construct(const FArguments& InArgs)
 {
-	PinRef = InPinRef;
-
 	ChildSlot
 	[
 		SNew(SVerticalBox)
@@ -147,14 +146,6 @@ void SPinValueInspector::Construct(const FArguments& InArgs, const FEdGraphPinRe
 				]
 		]
 	];
-
-	PopulateTreeView();
-
-	// Expand all root-level nodes that have children by default.
-	for (const FDebugTreeItemPtr& RootNode : TreeViewWidget->GetRootTreeItems())
-	{
-		TreeViewWidget->SetItemExpansion(RootNode, RootNode->HasChildren());
-	}
 }
 
 bool SPinValueInspector::ShouldShowSearchFilter() const
@@ -193,7 +184,6 @@ void SPinValueInspector::PopulateTreeView()
 	const UEdGraphPin* GraphPin = PinRef.Get();
 	if (!GraphPin)
 	{
-		TreeViewWidget->AddTreeItemUnique(SKismetDebugTreeView::MakeMessageItem(LOCTEXT("InvalidPin", "Pin Not Found").ToString()));
 		return;
 	}
 
@@ -240,6 +230,14 @@ void SPinValueInspector::PopulateTreeView()
 	}
 }
 
+void SPinValueInspector::SetPinRef(const FEdGraphPinReference& InPinRef)
+{
+	TreeViewWidget->ClearTreeItems();
+	PinRef = InPinRef;
+	PopulateTreeView();
+	ConstrainedBox->RequestResize();
+}
+
 TWeakPtr<FPinValueInspectorTooltip> FPinValueInspectorTooltip::SummonTooltip(FEdGraphPinReference InPinRef)
 {
 	if (!FSlateApplication::Get().GetVisibleMenuWindow().Get())
@@ -256,9 +254,9 @@ TWeakPtr<FPinValueInspectorTooltip> FPinValueInspectorTooltip::SummonTooltip(FEd
 			}
 
 			Instance = MakeShared<FPinValueInspectorTooltip>();
+			ValueInspectorWidget->SetPinRef(InPinRef);
 			TooltipWindow->ShowWindow();
 			TooltipWindow->SetAllowFastUpdate(true);
-			TooltipWidget->SetContentWidget(SAssignNew(Instance->ValueInspectorWidget, SPinValueInspector, InPinRef));
 			return Instance;
 		}
 	}
@@ -276,10 +274,10 @@ void FPinValueInspectorTooltip::MoveTooltip(const FVector2D& InNewLocation)
 
 void FPinValueInspectorTooltip::TryDismissTooltip()
 {
-	if (Instance.IsValid() && Instance->ValueInspectorWidget.IsValid())
+	if (Instance.IsValid() && ValueInspectorWidget.IsValid())
 	{
 		// if we can't inspect the pin, force it to close
-		if (!FKismetDebugUtilities::CanInspectPinValue(Instance->ValueInspectorWidget->PinRef.Get()) || TooltipCanClose())
+		if (!FKismetDebugUtilities::CanInspectPinValue(ValueInspectorWidget->PinRef.Get()) || TooltipCanClose())
 		{
 			DismissTooltip();
 		}
@@ -292,7 +290,7 @@ void FPinValueInspectorTooltip::DismissTooltip()
 	{
 		FSlateApplication::Get().DismissAllMenus();
 	}
-	TooltipWidget->ResetContentWidget();
+	ValueInspectorWidget->SetPinRef(FEdGraphPinReference());
 	TooltipWindow->HideWindow();
 	TooltipWindow->SetAllowFastUpdate(false);
 	Instance.Reset();
@@ -312,6 +310,7 @@ void FPinValueInspectorTooltip::CreatePinValueTooltipWindow()
 
 	FSlateApplication::Get().AddWindow(TooltipWindow.ToSharedRef());
 	TooltipWindow->SetContent(SAssignNew(TooltipWidget, SToolTip));
+	TooltipWidget->SetContentWidget(SAssignNew(ValueInspectorWidget, SPinValueInspector));
 }
 
 bool FPinValueInspectorTooltip::TooltipCanClose()
diff --git a/Engine/Source/Editor/Kismet/Public/SPinValueInspector.h b/Engine/Source/Editor/Kismet/Public/SPinValueInspector.h
--- a/Engine/Source/Editor/Kismet/Public/SPinValueInspector.h
+++ b/Engine/Source/Editor/Kismet/Public/SPinValueInspector.h
@@ -22,7 +22,7 @@ public:
 	{}
 	SLATE_END_ARGS()
 
-		void Construct(const FArguments& InArgs, const FEdGraphPinReference& InPinRef);
+		void Construct(const FArguments& InArgs);
 
 	/** Whether the search filter UI should be visible. */
 	bool ShouldShowSearchFilter() const;
@@ -42,6 +42,9 @@ protected:
 	/** Adds the pin to the tree view */
 	void PopulateTreeView();
 
+	/** Sets the current watched pin */
+	void SetPinRef(const FEdGraphPinReference& InPinRef);
+
 private:
 	/** Holds a weak reference to the target pin. */
 	FEdGraphPinReference PinRef;
@@ -69,9 +72,6 @@ public:
 	 */
 	void TryDismissTooltip();
 
-	/** Inspector widget in the tooltip */
-	TSharedPtr<SPinValueInspector> ValueInspectorWidget;
-
 private:
 	/** Dismisses the current tooltip (internal implementation) */
 	void DismissTooltip();
@@ -86,6 +86,9 @@ public:
 	/** Summons a new tooltip in the shared window */
 	static TWeakPtr<FPinValueInspectorTooltip> SummonTooltip(FEdGraphPinReference InPinRef);
 
+	/** Inspector widget in the tooltip */
+	static TSharedPtr<SPinValueInspector> ValueInspectorWidget;
+
 private:
 	/** Handles Creating a custom tooltip window for all PinValueInspector tooltips */
 	static void CreatePinValueTooltipWindow();
