Disable the pin type selector box on Multicast delegates. This was a regression introduced by 14378784, users should not be able to change the type of a multicast delegate.

#rnx

diff --git a/Engine/Source/Editor/Kismet/Private/SBlueprintPalette.cpp b/Engine/Source/Editor/Kismet/Private/SBlueprintPalette.cpp
--- a/Engine/Source/Editor/Kismet/Private/SBlueprintPalette.cpp
+++ b/Engine/Source/Editor/Kismet/Private/SBlueprintPalette.cpp
@@ -760,18 +760,24 @@ private:
 			}
 		}
 		
-		this->ChildSlot
-		[
-			SNew(SPinTypeSelector, FGetPinTypeTree::CreateUObject(GetDefault<UEdGraphSchema_K2>(), &UEdGraphSchema_K2::GetVariableTypeTree))
-			.ReadOnly(InArgs._ReadOnly)
-			.Schema(Schema)
-			.SchemaAction(ActionPtr)
-			.TargetPinType(this, &SPinTypeSelectorHelper::OnGetVarType)
-			.OnPinTypeChanged(this, &SPinTypeSelectorHelper::OnVarTypeChanged)
-			.TypeTreeFilter(ETypeTreeFilter::None)
-			.SelectorType(BlueprintEditorPtr.IsValid() ? SPinTypeSelector::ESelectorType::Partial : SPinTypeSelector::ESelectorType::None)
-			.CustomFilter(CustomPinTypeFilter)
-		];
+		const bool bIsDelegate = ActionPtr.IsValid() && ActionPtr.Pin()->GetTypeId() == FEdGraphSchemaAction_K2Delegate::StaticGetTypeId();
+
+		// You cannot change the type of multicast delegates in blueprints
+		if(!bIsDelegate)
+		{
+			this->ChildSlot
+			[
+				SNew(SPinTypeSelector, FGetPinTypeTree::CreateUObject(GetDefault<UEdGraphSchema_K2>(), &UEdGraphSchema_K2::GetVariableTypeTree))
+				.ReadOnly(InArgs._ReadOnly)
+				.Schema(Schema)
+				.SchemaAction(ActionPtr)
+				.TargetPinType(this, &SPinTypeSelectorHelper::OnGetVarType)
+				.OnPinTypeChanged(this, &SPinTypeSelectorHelper::OnVarTypeChanged)
+				.TypeTreeFilter(ETypeTreeFilter::None)
+				.SelectorType(BlueprintEditorPtr.IsValid() && !bIsDelegate ? SPinTypeSelector::ESelectorType::Partial : SPinTypeSelector::ESelectorType::None)
+				.CustomFilter(CustomPinTypeFilter)
+			];	
+		}
 	}
 	
 	FEdGraphPinType OnGetVarType() const
