Add a value customization for the Blueprint namespace property and combine it with the UI for adding imports.

Also add input validation for manual namespace entry, UI to view/select from the full list of registered namespace paths (rather than blindly typing for suggestions), and prevent users from importing an unregistered namespace.

diff --git a/Engine/Source/Editor/Kismet/Private/BlueprintDetailsCustomization.cpp b/Engine/Source/Editor/Kismet/Private/BlueprintDetailsCustomization.cpp
--- a/Engine/Source/Editor/Kismet/Private/BlueprintDetailsCustomization.cpp
+++ b/Engine/Source/Editor/Kismet/Private/BlueprintDetailsCustomization.cpp
@@ -65,6 +65,7 @@
 #include "BlueprintNamespaceHelper.h"
 #include "BlueprintNamespaceRegistry.h"
 #include "BlueprintNamespaceUtilities.h"
+#include "SBlueprintNamespaceEntry.h"
 #include "Widgets/Input/SSuggestionTextBox.h"
 #include "DragAndDrop/DecoratedDragDropOp.h"
 
@@ -5678,25 +5679,15 @@ FBlueprintImportsLayout::FBlueprintImportsLayout(TWeakPtr<class FBlueprintGlobal
 
 TSharedPtr<SWidget> FBlueprintImportsLayout::MakeAddItemRowWidget()
 {
-	return SNew(SHorizontalBox)
-	+ SHorizontalBox::Slot()
-	.FillWidth(1.0f)
-	[
-		SAssignNew(ImportEntryTextBox, SSuggestionTextBox)
-		.MinDesiredWidth(200.0f)
-		.Font(IDetailLayoutBuilder::GetDetailFont())
-		.OnTextCommitted(this, &FBlueprintImportsLayout::HandleImportEntryTextCommitted)
-		.OnShowingSuggestions(this, &FBlueprintImportsLayout::OnShowingImportSuggestions)
-	]
-	+ SHorizontalBox::Slot()
-	.AutoWidth()
-	.Padding(4.0f, 0.0f)
-	[
-		SNew(SButton)
-		.Text(LOCTEXT("BlueprintAddImportButtonText", "Add"))
-		.IsEnabled(this, &FBlueprintImportsLayout::IsAddImportEntryButtonEnabled)
-		.OnClicked(this, &FBlueprintImportsLayout::OnAddImportEntryButtonClicked)
-	];
+	return SNew(SBlueprintNamespaceEntry)
+		.AllowTextEntry(false)
+		.OnNamespaceSelected(this, &FBlueprintImportsLayout::OnNamespaceSelected)
+		.OnFilterNamespaceList(this, &FBlueprintImportsLayout::OnFilterNamespaceList)
+		.ButtonContent()
+		[
+			SNew(STextBlock)
+			.Text(LOCTEXT("BlueprintAddImportButton", "Add"))
+		];
 }
 
 void FBlueprintImportsLayout::GetManagedListItems(TArray<FManagedListItem>& OutListItems) const
@@ -5737,29 +5728,16 @@ void FBlueprintImportsLayout::OnRemoveItem(const FManagedListItem& Item)
 	OnRefreshInDetailsView();
 }
 
-bool FBlueprintImportsLayout::IsAddImportEntryButtonEnabled() const
-{
-	FString CurrentEntryAsString = ImportEntryTextBox->GetText().ToString();
-	return FBlueprintNamespaceRegistry::Get().IsRegisteredPath(CurrentEntryAsString);
-}
-
-FReply FBlueprintImportsLayout::OnAddImportEntryButtonClicked()
-{
-	HandleImportEntryTextCommitted(ImportEntryTextBox->GetText(), ETextCommit::OnEnter);
-	return FReply::Handled();
-}
-
-void FBlueprintImportsLayout::HandleImportEntryTextCommitted(const FText& NewLabel, ETextCommit::Type CommitType)
+void FBlueprintImportsLayout::OnNamespaceSelected(const FString& InNamespace)
 {
 	bool bWasAdded = false;
-	const FString& Namespace = NewLabel.ToString();
 
 	// Add to the blueprint's list of imports.
-	if (!Namespace.IsEmpty())
+	if (!InNamespace.IsEmpty())
 	{
 		UBlueprint* Blueprint = GetBlueprintObjectChecked();
 
-		if (FBlueprintEditorUtils::AddNamespaceToImportList(Blueprint, Namespace))
+		if (FBlueprintEditorUtils::AddNamespaceToImportList(Blueprint, InNamespace))
 		{
 			bWasAdded = true;
 		}
@@ -5775,50 +5753,18 @@ void FBlueprintImportsLayout::HandleImportEntryTextCommitted(const FText& NewLab
 		TSharedPtr<FBlueprintEditor> BlueprintEditorPtr = GetPinnedBlueprintEditorPtr();
 		if (BlueprintEditorPtr.IsValid())
 		{
-			BlueprintEditorPtr->ImportNamespace(Namespace);
+			BlueprintEditorPtr->ImportNamespace(InNamespace);
 		}
 	}
 }
 
-void FBlueprintImportsLayout::OnShowingImportSuggestions(const FString& InputText, TArray<FString>& OutSuggestions)
+void FBlueprintImportsLayout::OnFilterNamespaceList(TArray<FString>& InOutNamespaceList)
 {
-	int32 PathEnd;
-	FString CurrentPath;
-	FString CurrentName;
-	if (InputText.FindLastChar(TEXT('.'), PathEnd))
-	{
-		CurrentPath = InputText.LeftChop(InputText.Len() - PathEnd);
-		CurrentName = InputText.RightChop(PathEnd + 1);
-	}
-	else
+	const UBlueprint* Blueprint = GetBlueprintObjectChecked();
+	InOutNamespaceList.RemoveSwap(Blueprint->BlueprintNamespace);
+	for (const FString& ImportedNamespace : Blueprint->ImportedNamespaces)
 	{
-		CurrentName = InputText;
-	}
-
-	TArray<FName> SuggestedNames;
-	FBlueprintNamespaceRegistry::Get().GetNamesUnderPath(CurrentPath, SuggestedNames);
-
-	Algo::Sort(SuggestedNames, FNameLexicalLess());
-
-	TStringBuilder<128> PathBuilder;
-	UBlueprint* Blueprint = GetBlueprintObjectChecked();
-	for (FName SuggestedName : SuggestedNames)
-	{
-		FString SuggestedNameAsString = SuggestedName.ToString();
-		if (CurrentName.IsEmpty() || SuggestedNameAsString.StartsWith(CurrentName))
-		{
-			PathBuilder += CurrentPath;
-			PathBuilder += TEXT(".");
-			PathBuilder += SuggestedNameAsString;
-			
-			FString SuggestedNamespace = PathBuilder.ToString();
-			if (!Blueprint->ImportedNamespaces.Contains(SuggestedNamespace))
-			{
-				OutSuggestions.Add(SuggestedNamespace);
-			}
-
-			PathBuilder.Reset();
-		}
+		InOutNamespaceList.RemoveSwap(ImportedNamespace);
 	}
 }
 
@@ -5975,6 +5921,9 @@ TSharedRef<SWidget> FBlueprintInterfaceLayout::OnGetAddInterfaceMenuContent()
 		];
 }
 
+// Double the size of the default details view width for string values, which is otherwise too narrow since the customization adds a combo button.
+float FBlueprintGlobalOptionsDetails::NamespacePropertyValueCustomization_MinDesiredWidth = 250.0f;
+
 UBlueprint* FBlueprintGlobalOptionsDetails::GetBlueprintObj() const
 {
 	if(BlueprintEditorPtr.IsValid())
@@ -6072,6 +6021,46 @@ FText FBlueprintGlobalOptionsDetails::GetDeprecatedTooltip() const
 	return LOCTEXT("DisabledDeprecateBlueprintTooltip", "This Blueprint is deprecated because of a parent, it is not possible to remove deprecation from it!");
 }
 
+void FBlueprintGlobalOptionsDetails::OnNamespaceValueCommitted(const FString& InNamespace)
+{
+	if (UBlueprint* Blueprint = GetBlueprintObj())
+	{
+		// Skip if the value has not been changed.
+		if (Blueprint->BlueprintNamespace == InNamespace)
+		{
+			return;
+		}
+
+		FProperty* NamespaceProperty = UBlueprint::StaticClass()->FindPropertyByName(GET_MEMBER_NAME_CHECKED(UBlueprint, BlueprintNamespace));
+		check(NamespaceProperty);
+
+		// Notify that the property is about to change.
+		Blueprint->PreEditChange(NamespaceProperty);
+
+		// Update the Blueprint's current namespace value.
+		Blueprint->BlueprintNamespace = InNamespace;
+		FBlueprintEditorUtils::MarkBlueprintAsModified(Blueprint);
+
+		// Call PostEditChange() to alert others to the value change.
+		FPropertyChangedEvent PropertyChangedEvent(NamespaceProperty, EPropertyChangeType::ValueSet);
+		Blueprint->PostEditChangeProperty(PropertyChangedEvent);
+
+		// Ensure that the new path is added into the namespace registry.
+		// @todo_namespaces - Find and remove old path if no longer in use.
+		FBlueprintNamespaceRegistry::Get().RegisterNamespace(InNamespace);
+
+		if (GetDefault<UBlueprintEditorSettings>()->bEnableNamespaceImportingFeatures)
+		{
+			// Auto-import the namespace into the current editor context.
+			TSharedPtr<FBlueprintEditor> BlueprintEditor = GetBlueprintEditorPtr().Pin();
+			if (BlueprintEditor.IsValid())
+			{
+				BlueprintEditor->ImportNamespace(InNamespace);
+			}
+		}
+	}
+}
+
 void FBlueprintGlobalOptionsDetails::CustomizeDetails(IDetailLayoutBuilder& DetailLayout)
 {
 	const UBlueprint* Blueprint = GetBlueprintObj();
@@ -6182,6 +6171,23 @@ void FBlueprintGlobalOptionsDetails::CustomizeDetails(IDetailLayoutBuilder& Deta
 					.ToolTipText( this, &FBlueprintGlobalOptionsDetails::GetDeprecatedTooltip )
 				];
 		}
+
+		static FName BlueprintNamespacePropertyName = GET_MEMBER_NAME_CHECKED(UBlueprint, BlueprintNamespace);
+		TSharedRef<IPropertyHandle> NamespacePropertyHandle = DetailLayout.GetProperty(BlueprintNamespacePropertyName);
+		DetailLayout.EditDefaultProperty(NamespacePropertyHandle)->CustomWidget()
+			.NameContent()
+			[
+				NamespacePropertyHandle->CreatePropertyNameWidget()
+			]
+			.ValueContent()
+			.MinDesiredWidth(NamespacePropertyValueCustomization_MinDesiredWidth)
+			[
+				SNew(SBlueprintNamespaceEntry)
+				.Font(IDetailLayoutBuilder::GetDetailFont())
+				.AllowTextEntry(true)
+				.CurrentNamespace(Blueprint->BlueprintNamespace)
+				.OnNamespaceSelected(this, &FBlueprintGlobalOptionsDetails::OnNamespaceValueCommitted)
+			];
 	}
 }
 
diff --git a/Engine/Source/Editor/Kismet/Private/BlueprintDetailsCustomization.h b/Engine/Source/Editor/Kismet/Private/BlueprintDetailsCustomization.h
--- a/Engine/Source/Editor/Kismet/Private/BlueprintDetailsCustomization.h
+++ b/Engine/Source/Editor/Kismet/Private/BlueprintDetailsCustomization.h
@@ -769,10 +769,9 @@ protected:
 	virtual void OnRemoveItem(const FManagedListItem& Item) override;
 	/** END FBlueprintManagedListDetails interface */
 
-	bool IsAddImportEntryButtonEnabled() const;
-	FReply OnAddImportEntryButtonClicked();
+	void OnNamespaceSelected(const FString& InNamespace);
+	void OnFilterNamespaceList(TArray<FString>& InOutNamespaceList);
 	void HandleImportEntryTextCommitted(const FText& NewLabel, ETextCommit::Type CommitType);
-	void OnShowingImportSuggestions(const FString& InputText, TArray<FString>& OutSuggestions);
 
 private:
 	TSharedPtr<SSuggestionTextBox> ImportEntryTextBox;
@@ -857,12 +856,18 @@ protected:
 	/** Returns the tooltip explaining deprecation */
 	FText GetDeprecatedTooltip() const;
 
+	/** Callback for when a new Blueprint namespace value is entered */
+	void OnNamespaceValueCommitted(const FString& InNamespace);
+
 private:
 	/** Weak reference to the Blueprint editor */
 	TWeakPtr<FBlueprintEditor> BlueprintEditorPtr;
 
 	/** Combo button used to choose a parent class */
 	TSharedPtr<SComboButton> ParentClassComboButton;
+
+	/** Desired width for namespace value customization */
+	static float NamespacePropertyValueCustomization_MinDesiredWidth;
 };
 
 
diff --git a/Engine/Source/Editor/Kismet/Private/BlueprintEditor.cpp b/Engine/Source/Editor/Kismet/Private/BlueprintEditor.cpp
--- a/Engine/Source/Editor/Kismet/Private/BlueprintEditor.cpp
+++ b/Engine/Source/Editor/Kismet/Private/BlueprintEditor.cpp
@@ -2024,7 +2024,7 @@ void FBlueprintEditor::ImportNamespace(const FString& InNamespace)
 		if (UBlueprint* BlueprintObj = Cast<UBlueprint>(EditingObj))
 		{
 			// Add it into the Blueprint's user-facing import set.
-			if (!BlueprintObj->ImportedNamespaces.Contains(InNamespace))
+			if (BlueprintObj->BlueprintNamespace != InNamespace && !BlueprintObj->ImportedNamespaces.Contains(InNamespace))
 			{
 				BlueprintObj->Modify();
 				BlueprintObj->ImportedNamespaces.Add(InNamespace);
diff --git a/Engine/Source/Editor/Kismet/Private/BlueprintNamespaceRegistry.cpp b/Engine/Source/Editor/Kismet/Private/BlueprintNamespaceRegistry.cpp
--- a/Engine/Source/Editor/Kismet/Private/BlueprintNamespaceRegistry.cpp
+++ b/Engine/Source/Editor/Kismet/Private/BlueprintNamespaceRegistry.cpp
@@ -128,6 +128,29 @@ void FBlueprintNamespaceRegistry::GetNamesUnderPath(const FString& InPath, TArra
 	}
 }
 
+void FBlueprintNamespaceRegistry::GetAllRegisteredPaths(TArray<FString>& OutPaths) const
+{
+	PathTree->ForeachNode([&OutPaths](const TArray<FName>& CurrentPath, TSharedRef<FBlueprintNamespacePathTree::FNode> Node)
+	{
+		if (Node->bIsAddedPath)
+		{
+			// Note: This is not a hard limit on namespace path identifier string length, it's an optimization to try and avoid reallocation during path construction.
+			TStringBuilder<128> PathBuilder;
+			for (FName PathSegment : CurrentPath)
+			{
+				if (PathBuilder.Len() > 0)
+				{
+					PathBuilder += TEXT(".");
+				}
+				PathBuilder += PathSegment.ToString();
+			}
+
+			FString FullPath = PathBuilder.ToString();
+			OutPaths.Add(MoveTemp(FullPath));
+		}
+	});
+}
+
 void FBlueprintNamespaceRegistry::FindAndRegisterAllNamespaces()
 {
 	// Register loaded class type namespace identifiers.
@@ -191,22 +214,24 @@ void FBlueprintNamespaceRegistry::FindAndRegisterAllNamespaces()
 	}
 }
 
-void FBlueprintNamespaceRegistry::RegisterNamespace(const UObject* InObject)
+void FBlueprintNamespaceRegistry::RegisterNamespace(const FString& InPath)
 {
-	FString ObjectNamespace = FBlueprintNamespaceUtilities::GetObjectNamespace(InObject);
-	if (!ObjectNamespace.IsEmpty())
+	if (!InPath.IsEmpty())
 	{
-		PathTree->AddPath(ObjectNamespace);
+		PathTree->AddPath(InPath);
 	}
 }
 
+void FBlueprintNamespaceRegistry::RegisterNamespace(const UObject* InObject)
+{
+	FString ObjectNamespace = FBlueprintNamespaceUtilities::GetObjectNamespace(InObject);
+	RegisterNamespace(ObjectNamespace);
+}
+
 void FBlueprintNamespaceRegistry::RegisterNamespace(const FAssetData& AssetData)
 {
 	FString AssetNamespace = FBlueprintNamespaceUtilities::GetAssetNamespace(AssetData);
-	if (!AssetNamespace.IsEmpty())
-	{
-		PathTree->AddPath(AssetNamespace);
-	}
+	RegisterNamespace(AssetNamespace);
 }
 
 void FBlueprintNamespaceRegistry::ToggleDefaultNamespace()
@@ -229,30 +254,17 @@ void FBlueprintNamespaceRegistry::DumpAllRegisteredPaths()
 		Initialize();
 	}
 
+	TArray<FString> AllPaths;
+	GetAllRegisteredPaths(AllPaths);
+
 	UE_LOG(LogNamespace, Log, TEXT("=== Registered Blueprint namespace paths:"));
 
-	int32 PathCounter = 0;
-	PathTree->ForeachNode([&PathCounter](const TArray<FName>& CurrentPath, TSharedRef<FBlueprintNamespacePathTree::FNode> Node)
+	for (const FString& Path : AllPaths)
 	{
-		if (Node->bIsAddedPath)
-		{
-			TStringBuilder<128> PathBuilder;
-			for (FName PathSegment : CurrentPath)
-			{
-				if (PathBuilder.Len() > 0)
-				{
-					PathBuilder += TEXT(".");
-				}
-				PathBuilder += PathSegment.ToString();
-			}
-
-			++PathCounter;
-
-			UE_LOG(LogNamespace, Log, TEXT("%s"), PathBuilder.ToString());
-		}
-	});
+		UE_LOG(LogNamespace, Log, TEXT("%s"), *Path);
+	}
 
-	UE_LOG(LogNamespace, Log, TEXT("=== (end) %d total paths ==="), PathCounter);
+	UE_LOG(LogNamespace, Log, TEXT("=== (end) %d total paths ==="), AllPaths.Num());
 }
 
 void FBlueprintNamespaceRegistry::OnDefaultNamespaceTypeChanged()
diff --git a/Engine/Source/Editor/Kismet/Private/SBlueprintNamespaceEntry.cpp b/Engine/Source/Editor/Kismet/Private/SBlueprintNamespaceEntry.cpp
new file mode 100644
--- /dev/null
+++ b/Engine/Source/Editor/Kismet/Private/SBlueprintNamespaceEntry.cpp
@@ -0,0 +1,294 @@
+// Copyright Epic Games, Inc. All Rights Reserved.
+
+#include "SBlueprintNamespaceEntry.h"
+#include "BlueprintNamespaceRegistry.h"
+#include "Framework/Application/SlateApplication.h"
+#include "Misc/TextFilterExpressionEvaluator.h"
+#include "Widgets/Input/SComboButton.h"
+#include "Widgets/Input/SSearchBox.h"
+#include "Widgets/Input/SSuggestionTextBox.h"
+#include "Widgets/Layout/SSeparator.h"
+
+#define LOCTEXT_NAMESPACE "SBlueprintNamespaceEntry"
+
+float SBlueprintNamespaceEntry::NamespaceListBorderPadding = 1.0f;
+float SBlueprintNamespaceEntry::NamespaceListMinDesiredWidth = 350.0f;
+
+void SBlueprintNamespaceEntry::Construct(const FArguments& InArgs)
+{
+	CurrentNamespace = InArgs._CurrentNamespace;
+	OnNamespaceSelected = InArgs._OnNamespaceSelected;
+	OnFilterNamespaceList = InArgs._OnFilterNamespaceList;
+
+	ChildSlot
+	[
+		SNew(SHorizontalBox)
+		+SHorizontalBox::Slot()
+		.FillWidth(1.0f)
+		[
+			SAssignNew(TextBox, SSuggestionTextBox)
+			.Font(InArgs._Font)
+			.ForegroundColor(FSlateColor::UseForeground())
+			.Visibility(InArgs._AllowTextEntry ? EVisibility::Visible : EVisibility::Collapsed)
+			.Text(FText::FromString(CurrentNamespace))
+			.OnTextChanged(this, &SBlueprintNamespaceEntry::OnTextChanged)
+			.OnTextCommitted(this, &SBlueprintNamespaceEntry::OnTextCommitted)
+			.OnShowingSuggestions(this, &SBlueprintNamespaceEntry::OnShowingSuggestions)
+		]
+		+SHorizontalBox::Slot()
+		.AutoWidth()
+		[
+			SNew(SComboButton)
+			.OnGetMenuContent(this, &SBlueprintNamespaceEntry::OnGetNamespaceListMenuContent)
+			.ButtonContent()
+			[
+				InArgs._ButtonContent.Widget
+			]
+		]
+	];
+}
+
+void SBlueprintNamespaceEntry::OnTextChanged(const FText& InText)
+{
+	// Note: Empty string is valid (i.e. global namespace).
+	bool bIsValidString = true;
+
+	// Only allow alphanumeric characters, '.' and '_'.
+	FString NewString = InText.ToString();
+	for (const TCHAR& NewChar : NewString)
+	{
+		if (!FChar::IsAlnum(NewChar) && NewChar != TEXT('_') && NewChar != TEXT('.'))
+		{
+			bIsValidString = false;
+			break;
+		}
+	}
+
+	FString ErrorText;
+	if (bIsValidString)
+	{
+		// Keep the current namespace in sync with the last-known valid text box value.
+		CurrentNamespace = MoveTemp(NewString);
+	}
+	else
+	{
+		ErrorText = LOCTEXT("InvalidNamespaceIdentifierStringError", "Invalid namespace identifier string.").ToString();
+	}
+
+	// Set the error text regardless of whether or not the path is valid; this will clear the error state if the string is valid.
+	if (TextBox.IsValid())
+	{
+		TextBox->SetError(ErrorText);
+	}
+}
+
+void SBlueprintNamespaceEntry::OnTextCommitted(const FText& NewText, ETextCommit::Type InTextCommit)
+{
+	// Not using the current textbox value here because it might be invalid, and we want to revert to the last-known valid namespace string on commit.
+	SelectNamespace(CurrentNamespace);
+}
+
+void SBlueprintNamespaceEntry::OnShowingSuggestions(const FString& InputText, TArray<FString>& OutSuggestions)
+{
+	int32 PathEnd;
+	FString CurrentPath;
+	FString CurrentName;
+	if (InputText.FindLastChar(TEXT('.'), PathEnd))
+	{
+		CurrentPath = InputText.LeftChop(InputText.Len() - PathEnd);
+		CurrentName = InputText.RightChop(PathEnd + 1);
+	}
+	else
+	{
+		CurrentName = InputText;
+	}
+
+	// Find all names (path segments) that fall under the current path prefix.
+	TArray<FName> SuggestedNames;
+	FBlueprintNamespaceRegistry::Get().GetNamesUnderPath(CurrentPath, SuggestedNames);
+
+	// Sort the list alphabetically.
+	Algo::Sort(SuggestedNames, FNameLexicalLess());
+
+	// Build the suggestion set based on the set of matching names we found above.
+	TStringBuilder<128> PathBuilder;
+	for (FName SuggestedName : SuggestedNames)
+	{
+		FString SuggestedNameAsString = SuggestedName.ToString();
+		if (CurrentName.IsEmpty() || SuggestedNameAsString.StartsWith(CurrentName))
+		{
+			if (CurrentPath.Len() > 0)
+			{
+				PathBuilder += CurrentPath;
+				PathBuilder += TEXT(".");
+			}
+
+			PathBuilder += SuggestedNameAsString;
+
+			FString SuggestedNamespace = PathBuilder.ToString();
+			OutSuggestions.Add(MoveTemp(SuggestedNamespace));
+
+			PathBuilder.Reset();
+		}
+	}
+}
+
+TSharedRef<SWidget> SBlueprintNamespaceEntry::OnGetNamespaceListMenuContent()
+{
+	// Find and filter all registered paths.
+	PopulateNamespaceList();
+
+	// Construct the list view widget that we'll use for the menu content.
+	SAssignNew(ListView, SListView<TSharedPtr<FString>>)
+		.SelectionMode(ESelectionMode::Single)
+		.ListItemsSource(&ListItems)
+		.OnGenerateRow(this, &SBlueprintNamespaceEntry::OnGenerateRowForNamespaceList)
+		.OnSelectionChanged(this, &SBlueprintNamespaceEntry::OnNamespaceListSelectionChanged);
+
+	// If the current namespace is non-empty, look for a matching item in the set.
+	if (!CurrentNamespace.IsEmpty())
+	{
+		const TSharedPtr<FString>* CurrentItemPtr = ListItems.FindByPredicate([&CurrentNamespace = this->CurrentNamespace](const TSharedPtr<FString>& Item)
+		{
+			return Item.IsValid() && *Item == CurrentNamespace;
+		});
+
+		// If we found a match, make it the initial selection.
+		if (CurrentItemPtr)
+		{
+			ListView->SetSelection(*CurrentItemPtr);
+		}
+	}
+
+	return SNew(SBorder)
+		.Padding(NamespaceListBorderPadding)
+		[
+			SNew(SBox)
+			.MinDesiredWidth(NamespaceListMinDesiredWidth)
+			[
+				SNew(SVerticalBox)
+				+SVerticalBox::Slot()
+				.AutoHeight()
+				[
+					SAssignNew(SearchBox, SSearchBox)
+					.OnTextChanged(this, &SBlueprintNamespaceEntry::OnNamespaceListFilterTextChanged)
+				]
+				+SVerticalBox::Slot()
+				.AutoHeight()
+				[
+					SNew(SSeparator)
+				]
+				+SVerticalBox::Slot()
+				.FillHeight(1.0f)
+				[
+					ListView.ToSharedRef()
+				]
+			]
+		];
+}
+
+TSharedRef<ITableRow> SBlueprintNamespaceEntry::OnGenerateRowForNamespaceList(TSharedPtr<FString> Item, const TSharedRef<STableViewBase>& OwnerTable)
+{
+	check(Item.IsValid());
+
+	// Check for an empty list and add a single (disabled) item if found.
+	bool bIsEnabled = true;
+	FText ItemText = FText::FromString(*Item);
+	if (ItemText.IsEmpty() && ListItems.Num() == 1)
+	{
+		bIsEnabled = false;
+		ItemText = LOCTEXT("BlueprintNamespaceList_NoItems", "No Matching Items");
+	}
+
+	// Construct a new row widget, highlighting any text that matches the search filter.
+	return SNew(STableRow<TSharedPtr<FName>>, OwnerTable)
+		.IsEnabled(bIsEnabled)
+		[
+			SNew(STextBlock)
+			.Text(ItemText)
+			.HighlightText(bIsEnabled && SearchBox.IsValid() ? SearchBox->GetText() : FText::GetEmpty())
+		];
+}
+
+void SBlueprintNamespaceEntry::OnNamespaceListFilterTextChanged(const FText& InText)
+{
+	// Gather/filter all registered paths.
+	PopulateNamespaceList();
+
+	// Refresh the namespace item list view.
+	if (ListView.IsValid())
+	{
+		ListView->RequestListRefresh();
+	}
+}
+
+void SBlueprintNamespaceEntry::OnNamespaceListSelectionChanged(TSharedPtr<FString> Item, ESelectInfo::Type SelectInfo)
+{
+	// These actions should not trigger a selection.
+	if (SelectInfo == ESelectInfo::OnNavigation || SelectInfo == ESelectInfo::Direct)
+	{
+		return;
+	}
+
+	if (Item.IsValid())
+	{
+		SelectNamespace(*Item);
+	}
+
+	// Switch focus back to the text box; this will close the popup as a result.
+	if (TextBox.IsValid())
+	{
+		FSlateApplication::Get().SetKeyboardFocus(TextBox);
+		FSlateApplication::Get().SetUserFocus(0, TextBox);
+	}
+}
+
+void SBlueprintNamespaceEntry::PopulateNamespaceList()
+{
+	// Clear the current list.
+	ListItems.Empty();
+
+	// Gather the full set of registered namespace paths.
+	TArray<FString> AllPaths;
+	FBlueprintNamespaceRegistry::Get().GetAllRegisteredPaths(AllPaths);
+
+	// Invoke the delegate to allow owners to filter the list as needed.
+	OnFilterNamespaceList.ExecuteIfBound(AllPaths);
+
+	// Sort the list alphabetically.
+	AllPaths.Sort();
+
+	// Set up an expression evaluator to further trim the list according to the search filter.
+	FTextFilterExpressionEvaluator SearchFilter(ETextFilterExpressionEvaluatorMode::BasicString);
+	SearchFilter.SetFilterText(SearchBox.IsValid() ? SearchBox->GetText() : FText::GetEmpty());
+
+	// Build the source item list for the list view widget.
+	for (const FString& Path : AllPaths)
+	{
+		// Only include items that match the current search filter text.
+		if (SearchFilter.TestTextFilter(FBasicStringFilterExpressionContext(Path)))
+		{
+			ListItems.Add(MakeShared<FString>(Path));
+		}
+	}
+
+	// If no items were added, we signal this by adding a single blank entry.
+	if (ListItems.Num() == 0)
+	{
+		ListItems.Add(MakeShared<FString>());
+	}
+}
+
+void SBlueprintNamespaceEntry::SelectNamespace(const FString& InNamespace)
+{
+	if (TextBox.IsValid())
+	{
+		// Update the textbox to reflect the selected value. Note that this should also clear any error state via OnTextChanged().
+		TextBox->SetText(FText::FromString(InNamespace));
+	}
+
+	// Invoke the delegate in response to the new selection.
+	OnNamespaceSelected.ExecuteIfBound(InNamespace);
+}
+
+#undef LOCTEXT_NAMESPACE
\ No newline at end of file
diff --git a/Engine/Source/Editor/Kismet/Public/BlueprintNamespaceRegistry.h b/Engine/Source/Editor/Kismet/Public/BlueprintNamespaceRegistry.h
--- a/Engine/Source/Editor/Kismet/Public/BlueprintNamespaceRegistry.h
+++ b/Engine/Source/Editor/Kismet/Public/BlueprintNamespaceRegistry.h
@@ -41,6 +41,18 @@ public:
 	 */
 	void GetNamesUnderPath(const FString& InPath, TArray<FName>& OutNames) const;
 
+	/**
+	 * @param OutPaths	On output, contains the full set of all currently-registered namespace identifier paths.
+	 */
+	void GetAllRegisteredPaths(TArray<FString>& OutPaths) const;
+
+	/**
+	 * Adds an explicit namespace identifier to the registry if not already included.
+	 * 
+	 * @param InPath	Path identifier string (e.g. "X.Y").
+	 */
+	void RegisterNamespace(const FString& InPath);
+
 protected:
 	FBlueprintNamespaceRegistry();
 
diff --git a/Engine/Source/Editor/Kismet/Public/SBlueprintNamespaceEntry.h b/Engine/Source/Editor/Kismet/Public/SBlueprintNamespaceEntry.h
new file mode 100644
--- /dev/null
+++ b/Engine/Source/Editor/Kismet/Public/SBlueprintNamespaceEntry.h
@@ -0,0 +1,70 @@
+// Copyright Epic Games, Inc. All Rights Reserved.
+
+#pragma once
+
+#include "CoreMinimal.h"
+#include "Widgets/SCompoundWidget.h"
+#include "Widgets/Views/SListView.h"
+#include "SlateFwd.h"
+
+class SSuggestionTextBox;
+
+/** A custom widget class that provides support for Blueprint namespace entry and/or selection. */
+class KISMET_API SBlueprintNamespaceEntry : public SCompoundWidget
+{
+public:
+	DECLARE_DELEGATE_OneParam(FOnNamespaceSelected, const FString&);
+	DECLARE_DELEGATE_OneParam(FOnFilterNamespaceList, TArray<FString>&);
+
+	SLATE_BEGIN_ARGS(SBlueprintNamespaceEntry)
+	: _Font(FCoreStyle::Get().GetFontStyle(TEXT("NormalFont")))
+	, _AllowTextEntry(true)
+	{}
+		/** Current namespace value. */
+		SLATE_ARGUMENT(FString, CurrentNamespace)
+
+		/** Font color and opacity. */
+		SLATE_ATTRIBUTE(FSlateFontInfo, Font)
+
+		/** Allow text input to manually set arbitrary values. */
+		SLATE_ARGUMENT(bool, AllowTextEntry)
+
+		/** Allow external code to set custom combo button content. */
+		SLATE_NAMED_SLOT(FArguments, ButtonContent)
+
+		/** Called when a valid namespace is either entered or selected. */
+		SLATE_EVENT(FOnNamespaceSelected, OnNamespaceSelected)
+
+		/** Called to allow external code to filter out the namespace list. */
+		SLATE_EVENT(FOnFilterNamespaceList, OnFilterNamespaceList)
+
+	SLATE_END_ARGS()
+
+	void Construct(const FArguments& InArgs);
+
+protected:
+	void OnTextChanged(const FText& InText);
+	void OnTextCommitted(const FText& NewText, ETextCommit::Type InTextCommit);
+	void OnShowingSuggestions(const FString& InputText, TArray<FString>& OutSuggestions);
+	TSharedRef<SWidget> OnGetNamespaceListMenuContent();
+	TSharedRef<ITableRow> OnGenerateRowForNamespaceList(TSharedPtr<FString> Item, const TSharedRef<STableViewBase>& OwnerTable);
+	void OnNamespaceListFilterTextChanged(const FText& InText);
+	void OnNamespaceListSelectionChanged(TSharedPtr<FString> Item, ESelectInfo::Type SelectInfo);
+
+	void PopulateNamespaceList();
+	void SelectNamespace(const FString& InNamespace);
+
+private:
+	FString CurrentNamespace;
+	TArray<TSharedPtr<FString>> ListItems;
+
+	TSharedPtr<SSuggestionTextBox> TextBox;
+	TSharedPtr<SSearchBox> SearchBox;
+	TSharedPtr<SListView<TSharedPtr<FString>>> ListView;
+
+	FOnNamespaceSelected OnNamespaceSelected;
+	FOnFilterNamespaceList OnFilterNamespaceList;
+
+	static float NamespaceListBorderPadding;
+	static float NamespaceListMinDesiredWidth;
+};
\ No newline at end of file
