Fix visual lag in Debugger search

diff --git a/Engine/Source/Editor/Kismet/Private/Debugging/SKismetDebugTreeView.cpp b/Engine/Source/Editor/Kismet/Private/Debugging/SKismetDebugTreeView.cpp
--- a/Engine/Source/Editor/Kismet/Private/Debugging/SKismetDebugTreeView.cpp
+++ b/Engine/Source/Editor/Kismet/Private/Debugging/SKismetDebugTreeView.cpp
@@ -570,7 +570,6 @@ protected:
 		TArray<FLineItemWithChildren*>& Parents,
 		ESearchFlags SearchFlags = SF_None)
 	{
-		TSharedPtr<ITableRow> Row = DebugTreeView->WidgetFromItem(SharedThis(this));
 		bVisible = false;
 
 		UpdateSearch(InSearchString, SearchFlags);
@@ -624,10 +623,7 @@ protected:
 		Parents.Pop(/*bAllowShrinking =*/ false);
 		if (bChildMatch)
 		{
-			if (Row && !Row->IsItemExpanded())
-			{
-				Row->ToggleExpansion();
-			}
+			DebugTreeView->SetItemExpansion(SharedThis(this), true);
 		}
 
 		return bVisible;
