Moved BaseClassesToAllowRecompilingDuringPlayInEditor from experimental to *both* blueprint project settings and blueprint editor preferences (so it can be enabled for everyone in a given project for safe base classes as well as the existing 'local power user' cases)
- Also got rid of the loaded class cache in favor of a cached state for just opened BP editors that gets invalidated when settings are edited or the BP is (directly) reparented

diff --git a/Engine/Source/Editor/BlueprintGraph/BlueprintGraph.Build.cs b/Engine/Source/Editor/BlueprintGraph/BlueprintGraph.Build.cs
--- a/Engine/Source/Editor/BlueprintGraph/BlueprintGraph.Build.cs
+++ b/Engine/Source/Editor/BlueprintGraph/BlueprintGraph.Build.cs
@@ -15,7 +15,7 @@ public class BlueprintGraph : ModuleRules
 		);
 
 		OverridePackageType = PackageOverrideType.EngineDeveloper;
-
+		 
 		PublicDependencyModuleNames.AddRange(
 			new string[] { 
 				"Core", 
@@ -25,6 +25,7 @@ public class BlueprintGraph : ModuleRules
 				"Slate",
                 "EditorStyle",
 				"EditorSubsystem",
+				"DeveloperSettings"
 			}
 		);
 
diff --git a/Engine/Source/Editor/BlueprintGraph/Private/BlueprintEditorSettings.cpp b/Engine/Source/Editor/BlueprintGraph/Private/BlueprintEditorSettings.cpp
--- a/Engine/Source/Editor/BlueprintGraph/Private/BlueprintEditorSettings.cpp
+++ b/Engine/Source/Editor/BlueprintGraph/Private/BlueprintEditorSettings.cpp
@@ -126,6 +126,4 @@ void UBlueprintEditorSettings::PostEditChangeProperty(FPropertyChangedEvent& Pro
 	}
 
 	Super::PostEditChangeProperty(PropertyChangedEvent);
-
-	OnSettingsChange.Broadcast(this, PropertyChangedEvent.ChangeType);
 }
\ No newline at end of file
diff --git a/Engine/Source/Editor/BlueprintGraph/Public/BlueprintEditorSettings.h b/Engine/Source/Editor/BlueprintGraph/Public/BlueprintEditorSettings.h
--- a/Engine/Source/Editor/BlueprintGraph/Public/BlueprintEditorSettings.h
+++ b/Engine/Source/Editor/BlueprintGraph/Public/BlueprintEditorSettings.h
@@ -50,8 +50,7 @@ template<> struct TStructOpsTypeTraits<FPerBlueprintSettings> : public TStructOp
 };
 
 UCLASS(config=EditorPerProjectUserSettings)
-class BLUEPRINTGRAPH_API UBlueprintEditorSettings
-	:	public UObject
+class BLUEPRINTGRAPH_API UBlueprintEditorSettings : public UDeveloperSettings
 {
 	GENERATED_UCLASS_BODY()
 
@@ -247,20 +246,12 @@ public:
 	UPROPERTY(config)
 	TMap<int32, FEditedDocumentInfo> GraphEditorQuickJumps;
 
-public:
-	DECLARE_MULTICAST_DELEGATE_TwoParams(FOnUpdateSettingsMulticaster, const UBlueprintEditorSettings*, EPropertyChangeType::Type);
-	FOnUpdateSettingsMulticaster OnSettingsChange;
-
-	FDelegateHandle RegisterOnUpdateSettings(const FOnUpdateSettingsMulticaster::FDelegate& Delegate)
-	{
-		return OnSettingsChange.Add(Delegate);
-	}
-
-	void UnregisterOnUpdateSettings(FDelegateHandle Object)
-	{
-		OnSettingsChange.Remove(Object);
-	}
-	
+	/**
+	 * Any blueprint deriving from one of these base classes will be allowed to recompile during Play-in-Editor
+	 * (This setting exists both as an editor preference and project setting, and will be allowed if listed in either place) 
+	 */
+	UPROPERTY(EditAnywhere, config, Category=Play, meta=(AllowAbstract))
+	TArray<TSoftClassPtr<UObject>> BaseClassesToAllowRecompilingDuringPlayInEditor;
 
 protected:
 	//~ Begin UObject Interface
diff --git a/Engine/Source/Editor/Blutility/Classes/EditorUtilityBlueprint.h b/Engine/Source/Editor/Blutility/Classes/EditorUtilityBlueprint.h
--- a/Engine/Source/Editor/Blutility/Classes/EditorUtilityBlueprint.h
+++ b/Engine/Source/Editor/Blutility/Classes/EditorUtilityBlueprint.h
@@ -19,6 +19,5 @@ class BLUTILITY_API UEditorUtilityBlueprint : public UBlueprint
 	// UBlueprint interface
 	virtual bool SupportedByDefaultBlueprintFactory() const override;
 	virtual bool AlwaysCompileOnLoad() const override;
-	virtual bool CanRecompileWhilePlayingInEditor() const override;
 	// End of UBlueprint interface
 };
diff --git a/Engine/Source/Editor/Blutility/Private/EditorUtilityBlueprint.cpp b/Engine/Source/Editor/Blutility/Private/EditorUtilityBlueprint.cpp
--- a/Engine/Source/Editor/Blutility/Private/EditorUtilityBlueprint.cpp
+++ b/Engine/Source/Editor/Blutility/Private/EditorUtilityBlueprint.cpp
@@ -20,8 +20,3 @@ bool UEditorUtilityBlueprint::AlwaysCompileOnLoad() const
 {
 	return true;
 }
-
-bool UEditorUtilityBlueprint::CanRecompileWhilePlayingInEditor() const
-{
-	return false;
-}
\ No newline at end of file
diff --git a/Engine/Source/Editor/Kismet/Kismet.Build.cs b/Engine/Source/Editor/Kismet/Kismet.Build.cs
--- a/Engine/Source/Editor/Kismet/Kismet.Build.cs
+++ b/Engine/Source/Editor/Kismet/Kismet.Build.cs
@@ -56,6 +56,7 @@ public class Kismet : ModuleRules
                 "UMGEditor",
                 "UMG", // for SBlueprintDiff
                 "WorkspaceMenuStructure",
+				"DeveloperSettings",
 				"ToolMenus",
 				"SubobjectEditor",
 				"SubobjectDataInterface",
diff --git a/Engine/Source/Editor/Kismet/Private/BlueprintEditor.cpp b/Engine/Source/Editor/Kismet/Private/BlueprintEditor.cpp
--- a/Engine/Source/Editor/Kismet/Private/BlueprintEditor.cpp
+++ b/Engine/Source/Editor/Kismet/Private/BlueprintEditor.cpp
@@ -34,6 +34,7 @@
 #include "EdGraphNode_Comment.h"
 #include "Editor/UnrealEdEngine.h"
 #include "Settings/EditorExperimentalSettings.h"
+#include "Settings/BlueprintEditorProjectSettings.h"
 #include "GeneralProjectSettings.h"
 #include "Kismet/GameplayStatics.h"
 #include "Components/TimelineComponent.h"
@@ -801,7 +802,60 @@ void FBlueprintEditor::OnClose()
 bool FBlueprintEditor::InEditingMode() const
 {
 	UBlueprint* Blueprint = GetBlueprintObj();
-	return !FSlateApplication::Get().InKismetDebuggingMode() && (!InDebuggingMode() || (Blueprint && Blueprint->CanRecompileWhilePlayingInEditor()));
+	if (!FSlateApplication::Get().InKismetDebuggingMode())
+	{
+		if (!IsPlayInEditorActive())
+		{
+			return true;
+		}
+		else
+		{
+			if (Blueprint)
+			{
+				if (Blueprint->CanAlwaysRecompileWhilePlayingInEditor())
+				{
+					return true;
+				}
+				else
+				{
+					if (ModifyDuringPIEStatus == ESafeToModifyDuringPIEStatus::Unknown)
+					{
+						// Check for project settings or editor preferences that allow this anyways
+						ModifyDuringPIEStatus = ESafeToModifyDuringPIEStatus::NotSafe;
+
+						auto CheckClassList = [](UClass* TestClass, const TArray<TSoftClassPtr<UObject>>& ClassList)
+						{
+							for (const TSoftClassPtr<UObject>& SoftBaseClass : ClassList)
+							{
+								// Safe to call Get() instead of doing a load, as we can't possibly be derived from an unloaded class
+								if (UClass* BaseClass = SoftBaseClass.Get())
+								{
+									if (TestClass->IsChildOf(BaseClass))
+									{
+										return true;
+									}
+								}
+							}
+							return false;
+						};
+
+						if (UClass* TestClass = Blueprint->GeneratedClass)
+						{
+							if (CheckClassList(TestClass, GetDefault<UBlueprintEditorSettings>()->BaseClassesToAllowRecompilingDuringPlayInEditor) ||
+								CheckClassList(TestClass, GetDefault<UBlueprintEditorProjectSettings>()->BaseClassesToAllowRecompilingDuringPlayInEditor))
+							{
+								ModifyDuringPIEStatus = ESafeToModifyDuringPIEStatus::Safe;
+							}
+						}
+					}
+
+					return ModifyDuringPIEStatus == ESafeToModifyDuringPIEStatus::Safe;
+				}
+			}
+		}
+	}
+
+	return false;
 }
 
 bool FBlueprintEditor::IsCompilingEnabled() const
@@ -810,14 +864,14 @@ bool FBlueprintEditor::IsCompilingEnabled() const
 	return Blueprint && Blueprint->BlueprintType != BPTYPE_MacroLibrary && InEditingMode();
 }
 
-bool FBlueprintEditor::InDebuggingMode() const
+bool FBlueprintEditor::IsPlayInEditorActive() const
 {
 	return GEditor->PlayWorld != nullptr;
 }
 
 EVisibility FBlueprintEditor::IsDebuggerVisible() const
 {
-	return InDebuggingMode() ? EVisibility::Visible : EVisibility::Collapsed;
+	return IsPlayInEditorActive() ? EVisibility::Visible : EVisibility::Collapsed;
 }
 
 int32 FBlueprintEditor::GetNumberOfSelectedNodes() const
@@ -2167,6 +2221,12 @@ void FBlueprintEditor::InitBlueprintEditor(
 			DumpMessagesToCompilerLog(Blueprint->UpgradeNotesLog->Messages, true);
 		}
 	}
+
+	// Register for notifications when settings change
+	BlueprintEditorSettingsChangedHandle = GetMutableDefault<UBlueprintEditorSettings>()->OnSettingChanged()
+		.AddRaw(this, &FBlueprintEditor::OnBlueprintEditorPreferencesChanged);
+	BlueprintProjectSettingsChangedHandle = GetMutableDefault<UBlueprintEditorProjectSettings>()->OnSettingChanged()
+		.AddRaw(this, &FBlueprintEditor::OnBlueprintProjectSettingsChanged);
 }
 
 void FBlueprintEditor::InitToolMenuContext(FToolMenuContext& MenuContext)
@@ -2562,6 +2622,10 @@ void FBlueprintEditor::SetupViewForBlueprintEditingMode()
 
 FBlueprintEditor::~FBlueprintEditor()
 {
+	// Stop watching the settings
+	GetMutableDefault<UBlueprintEditorSettings>()->OnSettingChanged().Remove(BlueprintEditorSettingsChangedHandle);
+	GetMutableDefault<UBlueprintEditorProjectSettings>()->OnSettingChanged().Remove(BlueprintProjectSettingsChangedHandle);
+
 	// Clean up the preview
 	DestroyPreview();
 
@@ -3219,6 +3283,9 @@ void FBlueprintEditor::ReparentBlueprint_NewParentChosen(UClass* ChosenClass)
 		}
 	}
 
+	//@TODO: This is probably insufficient as you could reparent the parent instead and we wouldn't get notified
+	ModifyDuringPIEStatus = ESafeToModifyDuringPIEStatus::Unknown;
+
 	FSlateApplication::Get().DismissAllMenus();
 }
 
@@ -3496,7 +3563,8 @@ void FBlueprintEditor::UndoGraphAction()
 
 bool FBlueprintEditor::CanUndoGraphAction() const
 {
-	return !InDebuggingMode();
+	//@TODO: Should probably allow this for BPs that can be edited during PIE (basically returning InEditingMode instead)
+	return !IsPlayInEditorActive();
 }
 
 void FBlueprintEditor::RedoGraphAction()
@@ -3506,7 +3574,8 @@ void FBlueprintEditor::RedoGraphAction()
 
 bool FBlueprintEditor::CanRedoGraphAction() const
 {
-	return !InDebuggingMode();
+	//@TODO: Should probably allow this for BPs that can be edited during PIE (basically returning InEditingMode instead)
+	return !IsPlayInEditorActive();
 }
 
 void FBlueprintEditor::OnActiveTabChanged(TSharedPtr<SDockTab> PreviouslyActive, TSharedPtr<SDockTab> NewlyActivated)
@@ -10231,6 +10300,16 @@ bool FBlueprintEditor::IsNonImportedObject(const UObject* InObject) const
 	return bNotImported;
 }
 
+void FBlueprintEditor::OnBlueprintProjectSettingsChanged(UObject*, struct FPropertyChangedEvent&)
+{
+	ModifyDuringPIEStatus = ESafeToModifyDuringPIEStatus::Unknown;
+}
+
+void FBlueprintEditor::OnBlueprintEditorPreferencesChanged(UObject*, struct FPropertyChangedEvent&)
+{
+	ModifyDuringPIEStatus = ESafeToModifyDuringPIEStatus::Unknown;
+}
+
 /////////////////////////////////////////////////////
 /////////////////////////////////////////////////////
 /////////////////////////////////////////////////////
diff --git a/Engine/Source/Editor/Kismet/Public/BlueprintEditor.h b/Engine/Source/Editor/Kismet/Public/BlueprintEditor.h
--- a/Engine/Source/Editor/Kismet/Public/BlueprintEditor.h
+++ b/Engine/Source/Editor/Kismet/Public/BlueprintEditor.h
@@ -53,6 +53,8 @@ class UToolMenu;
 class FBlueprintNamespaceHelper;
 class SSubobjectEditor;
 class FSubobjectEditorTreeNode;
+class UBlueprintEditorSettings;
+class UBlueprintEditorProjectSettings;
 
 /* Enums to use when grouping the blueprint members in the list panel. The order here will determine the order in the list */
 namespace NodeSectionID
@@ -333,8 +335,8 @@ public:
 	/** Dumps messages to the compiler log, with an option to force it to display/come to front */
 	void DumpMessagesToCompilerLog(const TArray<TSharedRef<class FTokenizedMessage>>& Messages, bool bForceMessageDisplay);
 
-	/** Returns true if in debugging mode */
-	bool InDebuggingMode() const;
+	/** Returns true if PIE is active */
+	bool IsPlayInEditorActive() const;
 
 	/** Get the currently selected set of nodes */
 	FGraphPanelSelectionSet GetSelectedNodes() const;
@@ -1508,8 +1510,24 @@ private:
 	/** Handle to the registered OnActiveTabChanged delegate */
 	FDelegateHandle OnActiveTabChangedDelegateHandle;
 
+	/** Delegate handle registered for when settings change */
+	FDelegateHandle BlueprintEditorSettingsChangedHandle;
+	FDelegateHandle BlueprintProjectSettingsChangedHandle;
+
+	enum class ESafeToModifyDuringPIEStatus
+	{
+		Unknown,
+		Safe,
+		NotSafe
+	};
+
+	mutable ESafeToModifyDuringPIEStatus ModifyDuringPIEStatus = ESafeToModifyDuringPIEStatus::Unknown;
+
 	// Allow derived editors to add command mappings 
 	virtual void OnCreateGraphEditorCommands(TSharedPtr<FUICommandList> GraphEditorCommandsList) {}
+
+	virtual void OnBlueprintProjectSettingsChanged(UObject*, struct FPropertyChangedEvent&);
+	virtual void OnBlueprintEditorPreferencesChanged(UObject*, struct FPropertyChangedEvent&);
 };
 
 #undef LOCTEXT_NAMESPACE
diff --git a/Engine/Source/Editor/UMGEditor/Private/TabFactory/AnimationTabSummoner.cpp b/Engine/Source/Editor/UMGEditor/Private/TabFactory/AnimationTabSummoner.cpp
--- a/Engine/Source/Editor/UMGEditor/Private/TabFactory/AnimationTabSummoner.cpp
+++ b/Engine/Source/Editor/UMGEditor/Private/TabFactory/AnimationTabSummoner.cpp
@@ -741,7 +741,7 @@ private:
 
 	bool CanExecuteContextMenuAction() const
 	{
-		return AnimationListView->GetNumItemsSelected() == 1 && !BlueprintEditor.Pin()->InDebuggingMode();
+		return AnimationListView->GetNumItemsSelected() == 1 && !BlueprintEditor.Pin()->IsPlayInEditorActive();
 	}
 
 	void OnDuplicateAnimation()
diff --git a/Engine/Source/Editor/UnrealEd/Classes/Settings/EditorExperimentalSettings.h b/Engine/Source/Editor/UnrealEd/Classes/Settings/EditorExperimentalSettings.h
--- a/Engine/Source/Editor/UnrealEd/Classes/Settings/EditorExperimentalSettings.h
+++ b/Engine/Source/Editor/UnrealEd/Classes/Settings/EditorExperimentalSettings.h
@@ -68,17 +68,6 @@ public:
 	UPROPERTY(EditAnywhere, config, Category=Blueprints, meta=(DisplayName="Blueprint Break on Exceptions"))
 	bool bBreakOnExceptions;
 
-protected:
-	/** Any blueprint deriving from one of these base classes will be allowed to recompile during Play-in-Editor */
-	UPROPERTY(EditAnywhere, config, Category = Blueprints, meta=(AllowAbstract))
-	TArray<TSoftClassPtr<UObject>> BaseClassesToAllowRecompilingDuringPlayInEditor;
-
-	UPROPERTY(Transient)
-	mutable TArray<TObjectPtr<UClass>> ResolvedBaseClassesToAllowRecompilingDuringPlayInEditor;
-
-public:
-	bool IsClassAllowedToRecompileDuringPIE(UClass* TestClass) const;
-
 	/** Should arrows indicating data/execution flow be drawn halfway along wires? */
 	UPROPERTY(/*EditAnywhere - deprecated (moved into UBlueprintEditorSettings), */config/*, Category=Blueprints, meta=(DisplayName="Draw midpoint arrows in Blueprints")*/)
 	bool bDrawMidpointArrowsInBlueprints;
diff --git a/Engine/Source/Editor/UnrealEd/Private/Settings/SettingsClasses.cpp b/Engine/Source/Editor/UnrealEd/Private/Settings/SettingsClasses.cpp
--- a/Engine/Source/Editor/UnrealEd/Private/Settings/SettingsClasses.cpp
+++ b/Engine/Source/Editor/UnrealEd/Private/Settings/SettingsClasses.cpp
@@ -162,35 +162,6 @@ UEditorExperimentalSettings::UEditorExperimentalSettings( const FObjectInitializ
 {
 }
 
-bool UEditorExperimentalSettings::IsClassAllowedToRecompileDuringPIE(UClass* TestClass) const
-{
-	if (TestClass != nullptr)
-	{
-		// Rebuild the list if necessary (if the list was edited, either number of entires or value, ResolvedBaseClassesToAllowRecompilingDuringPlayInEditor will get reset below)
-		if (ResolvedBaseClassesToAllowRecompilingDuringPlayInEditor.Num() != BaseClassesToAllowRecompilingDuringPlayInEditor.Num())
-		{
-			ResolvedBaseClassesToAllowRecompilingDuringPlayInEditor.Reset();
-			for (TSoftClassPtr<UObject> BaseClassPtr : BaseClassesToAllowRecompilingDuringPlayInEditor)
-			{
-				if (UClass* BaseClass = BaseClassPtr.Get())
-				{
-					ResolvedBaseClassesToAllowRecompilingDuringPlayInEditor.Add(BaseClass);
-				}
-			}
-		}
-
-		// See if the test class matches any of the enabled base classes
-		for (UClass* BaseClass : ResolvedBaseClassesToAllowRecompilingDuringPlayInEditor)
-		{
-			if ((BaseClass != nullptr) && TestClass->IsChildOf(BaseClass))
-			{
-				return true;
-			}
-		}
-	}
-	return false;
-}
-
 void UEditorExperimentalSettings::PostInitProperties()
 {
 	CVarEditorHDRSupport->Set(bHDREditor ? 1 : 0, ECVF_SetByProjectSetting);
@@ -236,8 +207,6 @@ void UEditorExperimentalSettings::PostEditChangeProperty( struct FPropertyChange
 		SaveConfig();
 	}
 
-	ResolvedBaseClassesToAllowRecompilingDuringPlayInEditor.Reset();
-
 	SettingChangedEvent.Broadcast(Name);
 }
 
diff --git a/Engine/Source/Editor/UnrealEd/Public/Settings/BlueprintEditorProjectSettings.h b/Engine/Source/Editor/UnrealEd/Public/Settings/BlueprintEditorProjectSettings.h
--- a/Engine/Source/Editor/UnrealEd/Public/Settings/BlueprintEditorProjectSettings.h
+++ b/Engine/Source/Editor/UnrealEd/Public/Settings/BlueprintEditorProjectSettings.h
@@ -60,6 +60,14 @@ public:
 	UPROPERTY(EditAnywhere, config, Category=Experimental, meta=(EditCondition="bEnableChildActorExpansionInTreeView"))
 	EChildActorComponentTreeViewVisualizationMode DefaultChildActorTreeViewMode;
 
+	/**
+	 * Any blueprint deriving from one of these base classes will be allowed to recompile during Play-in-Editor
+	 * (This setting exists both as an editor preference and project setting, and will be allowed if listed in either place) 
+	 */
+	UPROPERTY(EditAnywhere, config, Category=Play, meta=(AllowAbstract))
+	TArray<TSoftClassPtr<UObject>> BaseClassesToAllowRecompilingDuringPlayInEditor;
+
+public:
 	// UObject interface
 	virtual void PostEditChangeProperty(struct FPropertyChangedEvent& PropertyChangedEvent) override;
 	// End of UObject interface
diff --git a/Engine/Source/Runtime/Engine/Classes/Animation/AnimBlueprint.h b/Engine/Source/Runtime/Engine/Classes/Animation/AnimBlueprint.h
--- a/Engine/Source/Runtime/Engine/Classes/Animation/AnimBlueprint.h
+++ b/Engine/Source/Runtime/Engine/Classes/Animation/AnimBlueprint.h
@@ -142,7 +142,7 @@ class ENGINE_API UAnimBlueprint : public UBlueprint, public IInterface_PreviewMe
 	}
 
 	virtual bool IsValidForBytecodeOnlyRecompile() const override { return false; }
-	virtual bool CanRecompileWhilePlayingInEditor() const override;
+	virtual bool CanAlwaysRecompileWhilePlayingInEditor() const override;
 	// End of UBlueprint interface
 
 	// Finds the index of the specified group, or creates a new entry for it (unless the name is NAME_None, which will return INDEX_NONE)
diff --git a/Engine/Source/Runtime/Engine/Classes/Engine/Blueprint.h b/Engine/Source/Runtime/Engine/Classes/Engine/Blueprint.h
--- a/Engine/Source/Runtime/Engine/Classes/Engine/Blueprint.h
+++ b/Engine/Source/Runtime/Engine/Classes/Engine/Blueprint.h
@@ -827,8 +827,19 @@ public:
 	*/
 	virtual bool AlwaysCompileOnLoad() const { return false; }
 
-	/** Some Blueprints (and classes) can recompile while we are debugging a live session. This function controls whether this can occur. */
-	virtual bool CanRecompileWhilePlayingInEditor() const;
+	/**
+	 * Some Blueprints (and classes) can recompile while we are debugging a live session (play in editor).
+	 * This function controls whether this can always occur.
+	 * There are also editor preferences and project settings that can be used to opt-in other classes even
+	 * when this returns false
+	 */
+	virtual bool CanAlwaysRecompileWhilePlayingInEditor() const;
+
+	UE_DEPRECATED(5.0, "CanRecompileWhilePlayingInEditor was renamed to CanAlwaysRecompileWhilePlayingInEditor to better explain usage.")
+	bool CanRecompileWhilePlayingInEditor() const
+	{
+		return CanAlwaysRecompileWhilePlayingInEditor();
+	}
 
 	/**
 	 * Check whether this blueprint can be nativized or not
diff --git a/Engine/Source/Runtime/Engine/Private/AnimBlueprint.cpp b/Engine/Source/Runtime/Engine/Private/AnimBlueprint.cpp
--- a/Engine/Source/Runtime/Engine/Private/AnimBlueprint.cpp
+++ b/Engine/Source/Runtime/Engine/Private/AnimBlueprint.cpp
@@ -199,7 +199,7 @@ void UAnimBlueprint::PostLoad()
 #endif
 }
 
-bool UAnimBlueprint::CanRecompileWhilePlayingInEditor() const
+bool UAnimBlueprint::CanAlwaysRecompileWhilePlayingInEditor() const
 {
 	return true;
 }
diff --git a/Engine/Source/Runtime/Engine/Private/Blueprint.cpp b/Engine/Source/Runtime/Engine/Private/Blueprint.cpp
--- a/Engine/Source/Runtime/Engine/Private/Blueprint.cpp
+++ b/Engine/Source/Runtime/Engine/Private/Blueprint.cpp
@@ -23,7 +23,6 @@
 #if WITH_EDITOR
 #include "BlueprintCompilationManager.h"
 #include "Settings/ProjectPackagingSettings.h"
-#include "Editor/UnrealEd/Classes/Settings/EditorExperimentalSettings.h"
 #include "Engine/SimpleConstructionScript.h"
 #include "Engine/SCS_Node.h"
 #include "Kismet2/BlueprintEditorUtils.h"
@@ -875,9 +874,9 @@ void UBlueprint::GetReparentingRules(TSet< const UClass* >& AllowedChildrenOfCla
 
 }
 
-bool UBlueprint::CanRecompileWhilePlayingInEditor() const
+bool UBlueprint::CanAlwaysRecompileWhilePlayingInEditor() const
 {
-	return GetDefault<UEditorExperimentalSettings>()->IsClassAllowedToRecompileDuringPIE(ParentClass);
+	return false;
 }
 
 void UBlueprint::GetAssetRegistryTags(TArray<FAssetRegistryTag>& OutTags) const
