Add a message if no items match the search in debugger

diff --git a/Engine/Source/Editor/Kismet/Private/Debugging/SKismetDebugTreeView.cpp b/Engine/Source/Editor/Kismet/Private/Debugging/SKismetDebugTreeView.cpp
--- a/Engine/Source/Editor/Kismet/Private/Debugging/SKismetDebugTreeView.cpp
+++ b/Engine/Source/Editor/Kismet/Private/Debugging/SKismetDebugTreeView.cpp
@@ -2968,6 +2968,7 @@ void SKismetDebugTreeView::Construct(const FArguments& InArgs)
 	bFilteredItemsDirty = false;
 	bInDebuggerTab = InArgs._InDebuggerTab;
 	SearchString = MakeShared<FString>();
+	SearchMessageItem = MakeMessageItem(LOCTEXT("NoItemsMatchSearch", "No entries match the search text").ToString());
 
 	ChildSlot
 	[
@@ -3080,6 +3081,11 @@ void SKismetDebugTreeView::UpdateFilteredItems()
 		}
 	}
 
+	if (FilteredTreeRoots.IsEmpty())
+	{
+		FilteredTreeRoots.Add(SearchMessageItem);
+	}
+
 	TreeView->RequestTreeRefresh();
 }
 
diff --git a/Engine/Source/Editor/Kismet/Public/Debugging/SKismetDebugTreeView.h b/Engine/Source/Editor/Kismet/Public/Debugging/SKismetDebugTreeView.h
--- a/Engine/Source/Editor/Kismet/Public/Debugging/SKismetDebugTreeView.h
+++ b/Engine/Source/Editor/Kismet/Public/Debugging/SKismetDebugTreeView.h
@@ -222,6 +222,9 @@ private:
 
 	bool bFilteredItemsDirty;
 
+	/** message to display when no entries in the tree match the search text */
+	FDebugTreeItemPtr SearchMessageItem;
+
 	/** whether this tree is held within the blueprint debugger tab, used to hide context menu options */
 	bool bInDebuggerTab;
 public:
