Extend the Blueprint graph context menu to allow users to load additional actions from a non-imported namespace.

Change summary:
- Added a customizable widget area to SBlueprintContextTargetMenu.
- Extended SBlueprintActionMenu to include a new "Add Import" UI into the context targets menu (currently visible only when namespace importing features are enabled).
- Fixed up a few misc. UI issues found while testing the Blueprint namespace entry widget.
- Added missing "Reset to Default" behavior associated with the Blueprint namespace property default value editor customization.

diff --git a/Engine/Source/Editor/Kismet/Private/BlueprintDetailsCustomization.cpp b/Engine/Source/Editor/Kismet/Private/BlueprintDetailsCustomization.cpp
--- a/Engine/Source/Editor/Kismet/Private/BlueprintDetailsCustomization.cpp
+++ b/Engine/Source/Editor/Kismet/Private/BlueprintDetailsCustomization.cpp
@@ -6033,19 +6033,9 @@ void FBlueprintGlobalOptionsDetails::OnNamespaceValueCommitted(const FString& In
 			return;
 		}
 
-		FProperty* NamespaceProperty = UBlueprint::StaticClass()->FindPropertyByName(GET_MEMBER_NAME_CHECKED(UBlueprint, BlueprintNamespace));
-		check(NamespaceProperty);
-
-		// Notify that the property is about to change.
-		Blueprint->PreEditChange(NamespaceProperty);
-
-		// Update the Blueprint's current namespace value.
-		Blueprint->BlueprintNamespace = InNamespace;
-		FBlueprintEditorUtils::MarkBlueprintAsModified(Blueprint);
-
-		// Call PostEditChange() to alert others to the value change.
-		FPropertyChangedEvent PropertyChangedEvent(NamespaceProperty, EPropertyChangeType::ValueSet);
-		Blueprint->PostEditChangeProperty(PropertyChangedEvent);
+		// Update the current namespace value. This will handle pre/post-edit change notifications, etc.
+		check(NamespacePropertyHandle.IsValid());
+		NamespacePropertyHandle->SetValue(InNamespace);
 
 		// Ensure that the new path is added into the namespace registry.
 		// @todo_namespaces - Find and remove old path if no longer in use.
@@ -6063,6 +6053,29 @@ void FBlueprintGlobalOptionsDetails::OnNamespaceValueCommitted(const FString& In
 	}
 }
 
+bool FBlueprintGlobalOptionsDetails::ShouldShowNamespaceResetToDefault() const
+{
+	check(NamespacePropertyHandle.IsValid());
+	return NamespacePropertyHandle->CanResetToDefault();
+}
+
+void FBlueprintGlobalOptionsDetails::OnNamespaceResetToDefaultValue()
+{
+	// Standard reset-to-default path.
+	check(NamespacePropertyHandle.IsValid());
+	NamespacePropertyHandle->ResetToDefault();
+
+	// Get the value after having been reset.
+	FString DefaultNamespaceValue;
+	NamespacePropertyHandle->GetValue(DefaultNamespaceValue);
+
+	// Update the entry widget to reflect the new value.
+	if (NamespaceValueWidget.IsValid())
+	{
+		NamespaceValueWidget->SetCurrentNamespace(DefaultNamespaceValue);
+	}
+}
+
 void FBlueprintGlobalOptionsDetails::CustomizeDetails(IDetailLayoutBuilder& DetailLayout)
 {
 	const UBlueprint* Blueprint = GetBlueprintObj();
@@ -6175,7 +6188,7 @@ void FBlueprintGlobalOptionsDetails::CustomizeDetails(IDetailLayoutBuilder& Deta
 		}
 
 		static FName BlueprintNamespacePropertyName = GET_MEMBER_NAME_CHECKED(UBlueprint, BlueprintNamespace);
-		TSharedRef<IPropertyHandle> NamespacePropertyHandle = DetailLayout.GetProperty(BlueprintNamespacePropertyName);
+		NamespacePropertyHandle = DetailLayout.GetProperty(BlueprintNamespacePropertyName);
 		DetailLayout.EditDefaultProperty(NamespacePropertyHandle)->CustomWidget()
 			.NameContent()
 			[
@@ -6184,12 +6197,16 @@ void FBlueprintGlobalOptionsDetails::CustomizeDetails(IDetailLayoutBuilder& Deta
 			.ValueContent()
 			.MinDesiredWidth(NamespacePropertyValueCustomization_MinDesiredWidth)
 			[
-				SNew(SBlueprintNamespaceEntry)
+				SAssignNew(NamespaceValueWidget, SBlueprintNamespaceEntry)
 				.Font(IDetailLayoutBuilder::GetDetailFont())
 				.AllowTextEntry(true)
 				.CurrentNamespace(Blueprint->BlueprintNamespace)
 				.OnNamespaceSelected(this, &FBlueprintGlobalOptionsDetails::OnNamespaceValueCommitted)
-			];
+			]
+			.OverrideResetToDefault(FResetToDefaultOverride::Create(
+				TAttribute<bool>(this, &FBlueprintGlobalOptionsDetails::ShouldShowNamespaceResetToDefault),
+				FSimpleDelegate::CreateSP(this, &FBlueprintGlobalOptionsDetails::OnNamespaceResetToDefaultValue))
+			);
 	}
 }
 
diff --git a/Engine/Source/Editor/Kismet/Private/BlueprintDetailsCustomization.h b/Engine/Source/Editor/Kismet/Private/BlueprintDetailsCustomization.h
--- a/Engine/Source/Editor/Kismet/Private/BlueprintDetailsCustomization.h
+++ b/Engine/Source/Editor/Kismet/Private/BlueprintDetailsCustomization.h
@@ -30,7 +30,7 @@ class SMultiLineEditableTextBox;
 class UEdGraphNode_Documentation;
 class UK2Node_Variable;
 class FSubobjectEditorTreeNode;
-class SSuggestionTextBox;
+class SBlueprintNamespaceEntry;
 
 /**
  * Variable network replication options.
@@ -772,9 +772,6 @@ protected:
 	void OnNamespaceSelected(const FString& InNamespace);
 	void OnFilterNamespaceList(TArray<FString>& InOutNamespaceList);
 	void HandleImportEntryTextCommitted(const FText& NewLabel, ETextCommit::Type CommitType);
-
-private:
-	TSharedPtr<SSuggestionTextBox> ImportEntryTextBox;
 };
 
 /** Blueprint Interface List Details */
@@ -859,6 +856,12 @@ protected:
 	/** Callback for when a new Blueprint namespace value is entered */
 	void OnNamespaceValueCommitted(const FString& InNamespace);
 
+	/** Determine Blueprint namespace "reset to default" button visibility */
+	bool ShouldShowNamespaceResetToDefault() const;
+
+	/** Invoked when the Blueprint namespace "reset to default" button is clicked */
+	void OnNamespaceResetToDefaultValue();
+
 private:
 	/** Weak reference to the Blueprint editor */
 	TWeakPtr<FBlueprintEditor> BlueprintEditorPtr;
@@ -866,6 +869,12 @@ private:
 	/** Combo button used to choose a parent class */
 	TSharedPtr<SComboButton> ParentClassComboButton;
 
+	/** Handle to the customized namespace property */
+	TSharedPtr<IPropertyHandle> NamespacePropertyHandle;
+
+	/** Widget used for namespace value customization */
+	TSharedPtr<SBlueprintNamespaceEntry> NamespaceValueWidget;
+
 	/** Desired width for namespace value customization */
 	static float NamespacePropertyValueCustomization_MinDesiredWidth;
 };
diff --git a/Engine/Source/Editor/Kismet/Private/SBlueprintActionMenu.cpp b/Engine/Source/Editor/Kismet/Private/SBlueprintActionMenu.cpp
--- a/Engine/Source/Editor/Kismet/Private/SBlueprintActionMenu.cpp
+++ b/Engine/Source/Editor/Kismet/Private/SBlueprintActionMenu.cpp
@@ -25,6 +25,8 @@
 #include "IDocumentation.h"
 #include "SSubobjectEditor.h"
 #include "SBlueprintContextTargetMenu.h"
+#include "SBlueprintNamespaceEntry.h"
+#include "BlueprintEditorSettings.h"
 
 #define LOCTEXT_NAMESPACE "SBlueprintGraphContextMenu"
 
@@ -249,6 +251,32 @@ void SBlueprintActionMenu::Construct( const FArguments& InArgs, TSharedPtr<FBlue
 	FBlueprintActionContext MenuContext;
 	ConstructActionContext(MenuContext);
 
+	TSharedPtr<SWidget> AddImportTargetContent = SNullWidget::NullWidget;
+	if (GetDefault<UBlueprintEditorSettings>()->bEnableNamespaceImportingFeatures)
+	{
+		SAssignNew(AddImportTargetContent, SBox)
+			.ToolTipText(LOCTEXT("ImportActionLabelTooltip", "Choose a namespace to import and load additional actions."))
+		[
+			SNew(SHorizontalBox)
+			+SHorizontalBox::Slot()
+				.AutoWidth()
+				.VAlign(VAlign_Center)
+			[
+				SNew(STextBlock)
+				.Text(LOCTEXT("ImportActionButtonLabel", "Import Actions From:"))
+			]
+			+SHorizontalBox::Slot()
+				.AutoWidth()
+				.Padding(4.f, 0.f)
+			[
+				SNew(SBlueprintNamespaceEntry)
+					.AllowTextEntry(false)
+					.OnFilterNamespaceList(this, &SBlueprintActionMenu::OnFilterImportNamespaceList)
+					.OnNamespaceSelected(this, &SBlueprintActionMenu::OnNamespaceSelectedForImport)
+			]
+		];
+	}
+
 	TSharedPtr<SComboButton> TargetContextSubMenuButton;
 	// @TODO: would be nice if we could use a checkbox style for this, and have a different state for open/closed
 	SAssignNew(TargetContextSubMenuButton, SComboButton)
@@ -260,6 +288,10 @@ void SBlueprintActionMenu::Construct( const FArguments& InArgs, TSharedPtr<FBlue
 		[
 			SAssignNew(ContextTargetSubMenu, SBlueprintContextTargetMenu, MenuContext)
 				.OnTargetMaskChanged(this, &SBlueprintActionMenu::OnContextTargetsChanged)
+				.CustomTargetContent()
+				[
+					AddImportTargetContent.ToSharedRef()
+				]
 		];
 
 	// Build the widget layout
@@ -572,6 +604,58 @@ void SBlueprintActionMenu::TryInsertPromoteToVariable(FBlueprintActionContext co
 	}
 }
 
+void SBlueprintActionMenu::OnFilterImportNamespaceList(TArray<FString>& InOutNamespaceList)
+{
+	FBlueprintActionContext MenuContext;
+	ConstructActionContext(MenuContext);
+
+	for (const UBlueprint* Blueprint : MenuContext.Blueprints)
+	{
+		InOutNamespaceList.RemoveSwap(Blueprint->BlueprintNamespace);
+
+		for (const FString& ImportedNamespace : Blueprint->ImportedNamespaces)
+		{
+			InOutNamespaceList.RemoveSwap(ImportedNamespace);
+		}
+	}
+}
+
+void SBlueprintActionMenu::OnNamespaceSelectedForImport(const FString& InNamespace)
+{
+	bool bWasAdded = false;
+
+	FBlueprintActionContext MenuContext;
+	ConstructActionContext(MenuContext);
+
+	// Add to the blueprint's list of imports.
+	if (!InNamespace.IsEmpty())
+	{
+		for (UBlueprint* Blueprint : MenuContext.Blueprints)
+		{
+			if (FBlueprintEditorUtils::AddNamespaceToImportList(Blueprint, InNamespace))
+			{
+				bWasAdded = true;
+			}
+		}
+	}
+
+	if (bWasAdded)
+	{
+		// Import the namespace into the current editor context. This may load additional type assets.
+		TSharedPtr<FBlueprintEditor> BlueprintEditorPtr = EditorPtr.Pin();
+		if (BlueprintEditorPtr.IsValid())
+		{
+			BlueprintEditorPtr->ImportNamespace(InNamespace);
+			BlueprintEditorPtr->RefreshInspector();
+		}
+
+		// Now that additional types have been loaded/imported, update the menu to include any additional action(s).
+		const bool bPreserveExpansion = true;
+		const bool bHandleOnSelectionEvent = false;
+		GraphActionMenu->RefreshAllActions(bPreserveExpansion, bHandleOnSelectionEvent);
+	}
+}
+
 /*******************************************************************************
 * SBlueprintActionMenuExpander
 *******************************************************************************/
diff --git a/Engine/Source/Editor/Kismet/Private/SBlueprintActionMenu.h b/Engine/Source/Editor/Kismet/Private/SBlueprintActionMenu.h
--- a/Engine/Source/Editor/Kismet/Private/SBlueprintActionMenu.h
+++ b/Engine/Source/Editor/Kismet/Private/SBlueprintActionMenu.h
@@ -69,9 +69,15 @@ protected:
 	/**  */
 	void ConstructActionContext(FBlueprintActionContext& ContextDescOut);
 
-	/** Functioin to try to insert a promote to variable entry if it is possible to do so. */
+	/** Function to try to insert a promote to variable entry if it is possible to do so. */
 	void TryInsertPromoteToVariable(FBlueprintActionContext const& Context, FGraphActionListBuilderBase& OutAllActions);
 
+	/** Callback to filter the list of namespaces available for import. */
+	void OnFilterImportNamespaceList(TArray<FString>& InOutNamespaceList);
+
+	/** Invoked when an available namespace is selected for import. */
+	void OnNamespaceSelectedForImport(const FString& InNamespace);
+
 private:
 	UEdGraph* GraphObj;
 	TArray<UEdGraphPin*> DraggedFromPins;
diff --git a/Engine/Source/Editor/Kismet/Private/SBlueprintContextTargetMenu.cpp b/Engine/Source/Editor/Kismet/Private/SBlueprintContextTargetMenu.cpp
--- a/Engine/Source/Editor/Kismet/Private/SBlueprintContextTargetMenu.cpp
+++ b/Engine/Source/Editor/Kismet/Private/SBlueprintContextTargetMenu.cpp
@@ -7,6 +7,7 @@
 #include "Styling/SlateTypes.h"
 #include "Widgets/Text/STextBlock.h"
 #include "Widgets/Layout/SBox.h"
+#include "Widgets/Layout/SSeparator.h"
 #include "Widgets/Input/SCheckBox.h"
 #include "EditorStyleSet.h"
 #include "Components/ActorComponent.h"
@@ -495,6 +496,27 @@ void SBlueprintContextTargetMenu::Construct(const FArguments& InArgs, const FBlu
 
 	FText const MenuToolTip = LOCTEXT("MenuToolTip", "Select whose functions/variables you want to see.\nNOTE: Unchecking everything is akin to 'SHOW EVERYTHING' (you're choosing to have NO target context and to not limit the scope)");
 
+	TSharedRef<SWidget> CustomContentWidget = InArgs._CustomTargetContent.Widget;
+	if (CustomContentWidget != SNullWidget::NullWidget)
+	{
+		SAssignNew(CustomContentWidget, SBox)
+		[
+			SNew(SVerticalBox)
+			+SVerticalBox::Slot()
+				.AutoHeight()
+				.Padding(0.f, 4.f, 0.f, 0.f)
+			[
+				SNew(SSeparator)
+			]
+			+SVerticalBox::Slot()
+				.AutoHeight()
+				.Padding(4.f)
+			[
+				CustomContentWidget
+			]
+		];
+	}
+
 	TSharedPtr<SHorizontalBox> MenuBody;
 	SBorder::Construct(SBorder::FArguments()
 		.BorderImage(FEditorStyle::GetBrush("Menu.Background"))
@@ -527,6 +549,11 @@ void SBlueprintContextTargetMenu::Construct(const FArguments& InArgs, const FBlu
 					SAssignNew(MenuBody, SHorizontalBox)
 						.ToolTipText(MenuToolTip)
 				]
+				+SVerticalBox::Slot()
+					.AutoHeight()
+				[
+					CustomContentWidget
+				]
 			]
 		]
 	);
diff --git a/Engine/Source/Editor/Kismet/Private/SBlueprintContextTargetMenu.h b/Engine/Source/Editor/Kismet/Private/SBlueprintContextTargetMenu.h
--- a/Engine/Source/Editor/Kismet/Private/SBlueprintContextTargetMenu.h
+++ b/Engine/Source/Editor/Kismet/Private/SBlueprintContextTargetMenu.h
@@ -48,6 +48,7 @@ class SBlueprintContextTargetMenu : public SBorder
 public:
 	SLATE_BEGIN_ARGS(SBlueprintContextTargetMenu){}
 		SLATE_EVENT(FOnTargetMaskChanged, OnTargetMaskChanged)
+		SLATE_NAMED_SLOT(FArguments, CustomTargetContent)
 	SLATE_END_ARGS()
 
 	// SWidget interface
diff --git a/Engine/Source/Editor/Kismet/Private/SBlueprintNamespaceEntry.cpp b/Engine/Source/Editor/Kismet/Private/SBlueprintNamespaceEntry.cpp
--- a/Engine/Source/Editor/Kismet/Private/SBlueprintNamespaceEntry.cpp
+++ b/Engine/Source/Editor/Kismet/Private/SBlueprintNamespaceEntry.cpp
@@ -38,7 +38,8 @@ void SBlueprintNamespaceEntry::Construct(const FArguments& InArgs)
 		+SHorizontalBox::Slot()
 		.AutoWidth()
 		[
-			SNew(SComboButton)
+			SAssignNew(ComboButton, SComboButton)
+			.CollapseMenuOnParentFocus(true)
 			.OnGetMenuContent(this, &SBlueprintNamespaceEntry::OnGetNamespaceListMenuContent)
 			.ButtonContent()
 			[
@@ -48,6 +49,15 @@ void SBlueprintNamespaceEntry::Construct(const FArguments& InArgs)
 	];
 }
 
+void SBlueprintNamespaceEntry::SetCurrentNamespace(const FString& InNamespace)
+{
+	// Pass through the text box in order to validate the string before committing it to the current value.
+	if (TextBox.IsValid())
+	{
+		TextBox->SetText(FText::FromString(InNamespace));
+	}
+}
+
 void SBlueprintNamespaceEntry::OnTextChanged(const FText& InText)
 {
 	// Note: Empty string is valid (i.e. global namespace).
@@ -140,7 +150,7 @@ TSharedRef<SWidget> SBlueprintNamespaceEntry::OnGetNamespaceListMenuContent()
 
 	// Construct the list view widget that we'll use for the menu content.
 	SAssignNew(ListView, SListView<TSharedPtr<FString>>)
-		.SelectionMode(ESelectionMode::Single)
+		.SelectionMode(ESelectionMode::SingleToggle)
 		.ListItemsSource(&ListItems)
 		.OnGenerateRow(this, &SBlueprintNamespaceEntry::OnGenerateRowForNamespaceList)
 		.OnSelectionChanged(this, &SBlueprintNamespaceEntry::OnNamespaceListSelectionChanged);
@@ -235,8 +245,20 @@ void SBlueprintNamespaceEntry::OnNamespaceListSelectionChanged(TSharedPtr<FStrin
 		SelectNamespace(*Item);
 	}
 
-	// Switch focus back to the text box; this will close the popup as a result.
-	if (TextBox.IsValid())
+	// Clear the search filter text.
+	if (SearchBox.IsValid())
+	{
+		SearchBox->SetText(FText::GetEmpty());
+	}
+
+	// Close the combo button menu after a selection.
+	if (ComboButton.IsValid())
+	{
+		ComboButton->SetIsOpen(false);
+	}
+
+	// Switch focus back to the text box if present and visible.
+	if (TextBox.IsValid() && TextBox->GetVisibility() == EVisibility::Visible)
 	{
 		FSlateApplication::Get().SetKeyboardFocus(TextBox);
 		FSlateApplication::Get().SetUserFocus(0, TextBox);
diff --git a/Engine/Source/Editor/Kismet/Public/BlueprintActionMenuUtils.h b/Engine/Source/Editor/Kismet/Public/BlueprintActionMenuUtils.h
--- a/Engine/Source/Editor/Kismet/Public/BlueprintActionMenuUtils.h
+++ b/Engine/Source/Editor/Kismet/Public/BlueprintActionMenuUtils.h
@@ -21,7 +21,7 @@ namespace EContextTargetFlags
 		TARGET_PinObject			= 0x00000008 UMETA(DisplayName="Pin Type Class", ToolTip="Include functions and variables that belong to this pin type."),
 		TARGET_SiblingPinObjects	= 0x00000010 UMETA(DisplayName="Other Object Outputs", ToolTip="Include functions and variables that belong to any of this node's output types."),
 		TARGET_BlueprintLibraries	= 0x00000020 UMETA(DisplayName="Libraries", ToolTip="Include static functions that are globally accessible (belonging to function/macro libraries, etc.)."),
-		TARGET_NonImportedTypes		= 0x00000040 UMETA(DisplayName="Non-Imported Types", ToolTip="Include functions from namespaces that are not imported by this Blueprint."),
+		TARGET_NonImportedTypes		= 0x00000040 UMETA(DisplayName="Non-Imported Types", ToolTip="Include functions from loaded types that are not imported by this Blueprint."),
 
 		// +1 to the last flag (so we can easily iterate these flags)
 		ContextTargetFlagsEnd UMETA(Hidden),
diff --git a/Engine/Source/Editor/Kismet/Public/SBlueprintNamespaceEntry.h b/Engine/Source/Editor/Kismet/Public/SBlueprintNamespaceEntry.h
--- a/Engine/Source/Editor/Kismet/Public/SBlueprintNamespaceEntry.h
+++ b/Engine/Source/Editor/Kismet/Public/SBlueprintNamespaceEntry.h
@@ -7,6 +7,7 @@
 #include "Widgets/Views/SListView.h"
 #include "SlateFwd.h"
 
+class SComboButton;
 class SSuggestionTextBox;
 
 /** A custom widget class that provides support for Blueprint namespace entry and/or selection. */
@@ -42,6 +43,13 @@ public:
 
 	void Construct(const FArguments& InArgs);
 
+	/**
+	 * Set the current namespace to the given identifier.
+	 * 
+	 * @param InNamespace	New namespace identifier. May be an empty string.
+	 */
+	void SetCurrentNamespace(const FString& InNamespace);
+
 protected:
 	void OnTextChanged(const FText& InText);
 	void OnTextCommitted(const FText& NewText, ETextCommit::Type InTextCommit);
@@ -58,6 +66,7 @@ private:
 	FString CurrentNamespace;
 	TArray<TSharedPtr<FString>> ListItems;
 
+	TSharedPtr<SComboButton> ComboButton;
 	TSharedPtr<SSuggestionTextBox> TextBox;
 	TSharedPtr<SSearchBox> SearchBox;
 	TSharedPtr<SListView<TSharedPtr<FString>>> ListView;
