Expose an editor setting that will more easily allow users to toggle Blueprint namespace-related editor features on/off.

diff --git a/Engine/Source/Editor/BlueprintGraph/Private/BlueprintEditorSettings.cpp b/Engine/Source/Editor/BlueprintGraph/Private/BlueprintEditorSettings.cpp
--- a/Engine/Source/Editor/BlueprintGraph/Private/BlueprintEditorSettings.cpp
+++ b/Engine/Source/Editor/BlueprintGraph/Private/BlueprintEditorSettings.cpp
@@ -10,6 +10,8 @@
 #include "Editor.h"
 #include "BlueprintTypePromotion.h"
 #include "AssetRegistry/AssetRegistryModule.h"
+#include "BlueprintNamespaceHelper.h"
+#include "BlueprintNamespaceUtilities.h"
 
 UBlueprintEditorSettings::UBlueprintEditorSettings(const FObjectInitializer& ObjectInitializer)
 	: Super(ObjectInitializer)
@@ -37,6 +39,7 @@ UBlueprintEditorSettings::UBlueprintEditorSettings(const FObjectInitializer& Obj
 	, BreakpointReloadMethod(EBlueprintBreakpointReloadMethod::RestoreAll)
 	, bEnablePinValueInspectionTooltips(true)
 	// Experimental
+	, bEnableNamespaceEditorFeatures(false)
 	, bEnableNamespaceFilteringFeatures(false)
 	, bEnableNamespaceImportingFeatures(false)
 	, bFavorPureCastNodes(false)
@@ -100,6 +103,19 @@ void UBlueprintEditorSettings::OnAssetRemoved(UObject* Object)
 	}
 }
 
+void UBlueprintEditorSettings::PostInitProperties()
+{
+	Super::PostInitProperties();
+
+	// Initialize transient flags and console variables for namespace editor features from the config.
+	// @todo_namespaces - May be removed once dependent code is changed to utilize the config setting.
+	bEnableNamespaceFilteringFeatures = bEnableNamespaceEditorFeatures;
+	bEnableNamespaceImportingFeatures = bEnableNamespaceEditorFeatures;
+
+	// Update console flags to match the current configuration.
+	FBlueprintNamespaceHelper::RefreshEditorFeatureConsoleFlags();
+}
+
 void UBlueprintEditorSettings::PostEditChangeProperty(FPropertyChangedEvent& PropertyChangedEvent)
 {
 	const FName PropertyName = (PropertyChangedEvent.Property != nullptr) ? PropertyChangedEvent.Property->GetFName() : NAME_None;
@@ -125,5 +141,19 @@ void UBlueprintEditorSettings::PostEditChangeProperty(FPropertyChangedEvent& Pro
 		FBlueprintActionDatabase::Get().RefreshAll();
 	}
 
+	if (PropertyName == GET_MEMBER_NAME_CHECKED(UBlueprintEditorSettings, bEnableNamespaceEditorFeatures))
+	{
+		// Update transient settings and console variable flags to reflect the new config setting value.
+		// @todo_namespaces - May be removed once dependent code is changed to utilize the config setting.
+		bEnableNamespaceFilteringFeatures = bEnableNamespaceEditorFeatures;
+		bEnableNamespaceImportingFeatures = bEnableNamespaceEditorFeatures;
+
+		// Update console flags to match the current configuration.
+		FBlueprintNamespaceHelper::RefreshEditorFeatureConsoleFlags();
+
+		// Refresh the Blueprint editor UI environment to match current settings.
+		FBlueprintNamespaceUtilities::RefreshBlueprintEditorFeatures();
+	}
+
 	Super::PostEditChangeProperty(PropertyChangedEvent);
 }
\ No newline at end of file
diff --git a/Engine/Source/Editor/BlueprintGraph/Public/BlueprintEditorSettings.h b/Engine/Source/Editor/BlueprintGraph/Public/BlueprintEditorSettings.h
--- a/Engine/Source/Editor/BlueprintGraph/Public/BlueprintEditorSettings.h
+++ b/Engine/Source/Editor/BlueprintGraph/Public/BlueprintEditorSettings.h
@@ -145,14 +145,18 @@ public:
 
 // Experimental
 public:
+	/** Whether to enable namespace importing and filtering features in the Blueprint editor */
+	UPROPERTY(config, EditAnywhere, Category = Experimental)
+	bool bEnableNamespaceEditorFeatures;
+
 	/** Whether to enable namespace filtering features in the Blueprint editor */
-	// @todo_namespaces - Make this a non-transient editable config property.
-	UPROPERTY(Transient/*, Category = Experimental*/)
+	// @todo_namespaces - Remove this if/when dependent code is changed to utilize the single setting above.
+	UPROPERTY(Transient)
 	bool bEnableNamespaceFilteringFeatures;
 
 	/** Whether to enable namespace importing features in the Blueprint editor */
-	// @todo_namespaces - Make this a non-transient editable config property.
-	UPROPERTY(Transient/*, Category = Experimental*/)
+	// @todo_namespaces - Remove this if/when dependent code is changed to utilize the single setting above.
+	UPROPERTY(Transient)
 	bool bEnableNamespaceImportingFeatures;
 
 	// The list of namespaces to always expose in any Blueprint (local per-user)
@@ -255,6 +259,7 @@ public:
 
 protected:
 	//~ Begin UObject Interface
+	virtual void PostInitProperties();
 	virtual void PostEditChangeProperty(struct FPropertyChangedEvent& PropertyChangedEvent) override;
 	//~ End UObject Interface
 	
diff --git a/Engine/Source/Editor/Kismet/Private/BlueprintNamespaceHelper.cpp b/Engine/Source/Editor/Kismet/Private/BlueprintNamespaceHelper.cpp
--- a/Engine/Source/Editor/Kismet/Private/BlueprintNamespaceHelper.cpp
+++ b/Engine/Source/Editor/Kismet/Private/BlueprintNamespaceHelper.cpp
@@ -37,60 +37,6 @@ static TAutoConsoleVariable<bool> CVarBPImportParentClassNamespaces(
 	false,
 	TEXT("Enables import of parent class namespaces when opening a Blueprint for editing."));
 
-static void UpdateNamespaceFeatureSettingsCVarSinkFunction()
-{
-	// Note: Do NOT try to access settings objects below during the initial editor load! They rely on the config being loaded, which may not have occurred yet.
-	if (GIsInitialLoad || !GEditor)
-	{
-		return;
-	}
-
-	auto CheckAndUpdateSettingValueLambda = [](bool& CurValue, const bool NewValue) -> bool
-	{
-		if (CurValue != NewValue)
-		{
-			CurValue = NewValue;
-			return true;
-		}
-
-		return false;
-	};
-
-	bool bWasUpdated = false;
-
-	// Blueprint editor settings.
-	UBlueprintEditorSettings* BlueprintEditorSettingsPtr = GetMutableDefault<UBlueprintEditorSettings>();
-	bWasUpdated |= CheckAndUpdateSettingValueLambda(BlueprintEditorSettingsPtr->bEnableNamespaceFilteringFeatures, CVarBPEnableNamespaceFilteringFeatures.GetValueOnGameThread());
-	bWasUpdated |= CheckAndUpdateSettingValueLambda(BlueprintEditorSettingsPtr->bEnableNamespaceImportingFeatures, CVarBPEnableNamespaceImportingFeatures.GetValueOnGameThread());
-
-	if (bWasUpdated)
-	{
-		// Refresh all relevant open Blueprint editor UI elements.
-		// @todo_namespaces - Move this into PostEditChangeProperty() on the appropriate settings object(s).
-		if (UAssetEditorSubsystem* AssetEditorSubsystem = GEditor->GetEditorSubsystem<UAssetEditorSubsystem>())
-		{
-			TArray<UObject*> EditedAssets = AssetEditorSubsystem->GetAllEditedAssets();
-			for (UObject* Asset : EditedAssets)
-			{
-				if (Asset && Asset->IsA<UBlueprint>())
-				{
-					TSharedPtr<IToolkit> AssetEditorPtr = FToolkitManager::Get().FindEditorForAsset(Asset);
-					if (AssetEditorPtr.IsValid() && AssetEditorPtr->IsBlueprintEditor())
-					{
-						TSharedPtr<IBlueprintEditor> BlueprintEditorPtr = StaticCastSharedPtr<IBlueprintEditor>(AssetEditorPtr);
-						BlueprintEditorPtr->RefreshEditors();
-						BlueprintEditorPtr->RefreshInspector();
-					}
-				}
-			}
-		}
-	}
-}
-
-static FAutoConsoleVariableSink CVarUpdateNamespaceFeatureSettingsSink(
-	FConsoleCommandDelegate::CreateStatic(&UpdateNamespaceFeatureSettingsCVarSinkFunction)
-);
-
 // ---
 
 class FClassViewerNamespaceFilter : public IClassViewerFilter
@@ -276,4 +222,61 @@ bool FBlueprintNamespaceHelper::IsImportedAsset(const FAssetData& InAssetData) c
 	return IsIncludedInNamespaceList(Namespace);
 }
 
+namespace UE::Editor::Kismet::Private
+{
+	static void OnUpdateNamespaceEditorFeatureConsoleFlag(IConsoleVariable* InCVar, bool* InValuePtr)
+	{
+		check(InCVar);
+
+		// Skip if not set by console command; in that case we're updating the flag directly.
+		if ((InCVar->GetFlags() & ECVF_SetByMask) != ECVF_SetByConsole)
+		{
+			return;
+		}
+
+		// Update the editor setting (referenced) to match the console variable's new setting.
+		check(InValuePtr);
+		*InValuePtr = InCVar->GetBool();
+
+		// Refresh the Blueprint editor UI environment in response to the console variable change.
+		FBlueprintNamespaceUtilities::RefreshBlueprintEditorFeatures();
+	}
+}
+
+void FBlueprintNamespaceHelper::RefreshEditorFeatureConsoleFlags()
+{
+	UBlueprintEditorSettings* BlueprintEditorSettings = GetMutableDefault<UBlueprintEditorSettings>();
+
+	// Register callbacks to respond to flag changes via console.
+	static bool bIsInitialized = false;
+	if (!bIsInitialized)
+	{
+		auto InitCVarFlag = [](IConsoleVariable* InCVar, bool& InValueRef)
+		{
+			using namespace UE::Editor::Kismet::Private;
+			InCVar->OnChangedDelegate().AddStatic(&OnUpdateNamespaceEditorFeatureConsoleFlag, &InValueRef);
+		};
+
+		InitCVarFlag(CVarBPEnableNamespaceFilteringFeatures.AsVariable(), BlueprintEditorSettings->bEnableNamespaceFilteringFeatures);
+		InitCVarFlag(CVarBPEnableNamespaceImportingFeatures.AsVariable(), BlueprintEditorSettings->bEnableNamespaceImportingFeatures);
+
+		bIsInitialized = true;
+	}
+
+	// Update console variables to match current Blueprint editor settings.
+	static bool bIsUpdating = false;
+	if (!bIsUpdating)
+	{
+		TGuardValue<bool> ScopeGuard(bIsUpdating, true);
+
+		auto SetCVarFlag = [](IConsoleVariable* InCVar, bool& InValueRef)
+		{
+			InCVar->Set(InValueRef);
+		};
+
+		SetCVarFlag(CVarBPEnableNamespaceFilteringFeatures.AsVariable(), BlueprintEditorSettings->bEnableNamespaceFilteringFeatures);
+		SetCVarFlag(CVarBPEnableNamespaceImportingFeatures.AsVariable(), BlueprintEditorSettings->bEnableNamespaceImportingFeatures);
+	}
+}
+
 #undef LOCTEXT_NAMESPACE
\ No newline at end of file
diff --git a/Engine/Source/Editor/Kismet/Private/BlueprintNamespaceUtilities.cpp b/Engine/Source/Editor/Kismet/Private/BlueprintNamespaceUtilities.cpp
--- a/Engine/Source/Editor/Kismet/Private/BlueprintNamespaceUtilities.cpp
+++ b/Engine/Source/Editor/Kismet/Private/BlueprintNamespaceUtilities.cpp
@@ -6,6 +6,8 @@
 #include "AssetRegistry/AssetData.h"
 #include "EdGraphSchema_K2.h"
 #include "Misc/AssertionMacros.h"
+#include "BlueprintEditorModule.h"
+#include "Toolkits/ToolkitManager.h"
 
 namespace UE::Editor::Kismet::Private
 {
@@ -249,4 +251,31 @@ FBlueprintNamespaceUtilities::FOnDefaultBlueprintNamespaceTypeChanged& FBlueprin
 {
 	using namespace UE::Editor::Kismet::Private;
 	return OnDefaultBlueprintNamespaceTypeChangedDelegate;
+}
+
+void FBlueprintNamespaceUtilities::RefreshBlueprintEditorFeatures()
+{
+	if (!GEditor)
+	{
+		return;
+	}
+
+	// Refresh all relevant open Blueprint editor UI elements.
+	if (UAssetEditorSubsystem* AssetEditorSubsystem = GEditor->GetEditorSubsystem<UAssetEditorSubsystem>())
+	{
+		TArray<UObject*> EditedAssets = AssetEditorSubsystem->GetAllEditedAssets();
+		for (UObject* Asset : EditedAssets)
+		{
+			if (Asset && Asset->IsA<UBlueprint>())
+			{
+				TSharedPtr<IToolkit> AssetEditorPtr = FToolkitManager::Get().FindEditorForAsset(Asset);
+				if (AssetEditorPtr.IsValid() && AssetEditorPtr->IsBlueprintEditor())
+				{
+					TSharedPtr<IBlueprintEditor> BlueprintEditorPtr = StaticCastSharedPtr<IBlueprintEditor>(AssetEditorPtr);
+					BlueprintEditorPtr->RefreshEditors();
+					BlueprintEditorPtr->RefreshInspector();
+				}
+			}
+		}
+	}
 }
\ No newline at end of file
diff --git a/Engine/Source/Editor/Kismet/Public/BlueprintNamespaceHelper.h b/Engine/Source/Editor/Kismet/Public/BlueprintNamespaceHelper.h
--- a/Engine/Source/Editor/Kismet/Public/BlueprintNamespaceHelper.h
+++ b/Engine/Source/Editor/Kismet/Public/BlueprintNamespaceHelper.h
@@ -94,6 +94,11 @@ public:
 		return PinTypeSelectorFilter;
 	}
 
+	/**
+	 * Utility method used to keep console flags in sync with the current Blueprint editor settings environment.
+	 */
+	static void RefreshEditorFeatureConsoleFlags();
+
 private:
 	// Path tree that stores imported namespace identifiers for the associated Blueprint.
 	TUniquePtr<FBlueprintNamespacePathTree> NamespacePathTree;
diff --git a/Engine/Source/Editor/Kismet/Public/BlueprintNamespaceUtilities.h b/Engine/Source/Editor/Kismet/Public/BlueprintNamespaceUtilities.h
--- a/Engine/Source/Editor/Kismet/Public/BlueprintNamespaceUtilities.h
+++ b/Engine/Source/Editor/Kismet/Public/BlueprintNamespaceUtilities.h
@@ -67,4 +67,9 @@ public:
 	/** Delegate invoked whenever the default Blueprint namespace type changes. */
 	DECLARE_MULTICAST_DELEGATE(FOnDefaultBlueprintNamespaceTypeChanged);
 	static FOnDefaultBlueprintNamespaceTypeChanged& OnDefaultBlueprintNamespaceTypeChanged();
+
+	/**
+	 * Refresh the Blueprint editor environment to align with current namespace editor feature settings.
+	 */
+	static void RefreshBlueprintEditorFeatures();
 };
\ No newline at end of file
