Do not overwrite the default values with that of the generated class during compilation if delta serialization is enabled. This will make sure that if we are compiling a BP that has been reparented we use the parent class' default values instead of overwriting them with 0. This was a regression introduced by 16759714

#rnx

diff --git a/Engine/Source/Editor/Kismet/Private/BlueprintCompilationManager.cpp b/Engine/Source/Editor/Kismet/Private/BlueprintCompilationManager.cpp
--- a/Engine/Source/Editor/Kismet/Private/BlueprintCompilationManager.cpp
+++ b/Engine/Source/Editor/Kismet/Private/BlueprintCompilationManager.cpp
@@ -413,6 +413,7 @@ struct FCompilerData
 		InternalOptions.bSkipDefaultObjectValidation = (UserOptions & EBlueprintCompileOptions::SkipDefaultObjectValidation) != EBlueprintCompileOptions::None;
 		InternalOptions.bSkipFiBSearchMetaUpdate = (UserOptions & EBlueprintCompileOptions::SkipFiBSearchMetaUpdate) != EBlueprintCompileOptions::None;
 		InternalOptions.bUseDeltaSerializationDuringReinstancing = (UserOptions & EBlueprintCompileOptions::UseDeltaSerializationDuringReinstancing) != EBlueprintCompileOptions::None;
+		InternalOptions.bSkipNewVariableDefaultsDetection = (UserOptions & EBlueprintCompileOptions::SkipNewVariableDefaultsDetection) != EBlueprintCompileOptions::None;
 		InternalOptions.CompileType = bBytecodeOnly ? EKismetCompileType::BytecodeOnly : EKismetCompileType::Full;
 
 		Compiler = FKismetCompilerContext::GetCompilerForBP(BP, *ActiveResultsLog, InternalOptions);
@@ -434,7 +435,8 @@ struct FCompilerData
 	bool ShouldValidateClassDefaultObject() const { return JobType == ECompilationManagerJobType::Normal && !InternalOptions.bSkipDefaultObjectValidation; }
 	bool ShouldUpdateBlueprintSearchMetadata() const { return JobType == ECompilationManagerJobType::Normal && !InternalOptions.bSkipFiBSearchMetaUpdate; }
 	bool UseDeltaSerializationDuringReinstancing() const { return InternalOptions.bUseDeltaSerializationDuringReinstancing; }
-
+	bool ShouldSkipNewVariableDefaultsDetection() const { return InternalOptions.bSkipNewVariableDefaultsDetection; }
+	
 	UBlueprint* BP;
 	FCompilerResultsLog* ActiveResultsLog;
 	TUniquePtr<FCompilerResultsLog> ResultsLog;
@@ -1081,7 +1083,8 @@ void FBlueprintCompilationManagerImpl::FlushCompilationQueueImpl(bool bSuppressB
 			{
 				if (CompilerData.JobType == ECompilationManagerJobType::Normal &&
 					 CompilerData.BP->bHasBeenRegenerated &&		// Note: This ensures that we'll only do this after the Blueprint has been loaded/created; otherwise the class may not contain any properties to find.
-					 CompilerData.BP->GeneratedClass)
+					 CompilerData.BP->GeneratedClass &&
+					 !CompilerData.ShouldSkipNewVariableDefaultsDetection())
 				{
 					const UClass* ParentClass = CompilerData.BP->ParentClass;
 					while (const UBlueprint* ParentBP = UBlueprint::GetBlueprintFromClass(ParentClass))
diff --git a/Engine/Source/Editor/Kismet/Private/BlueprintEditor.cpp b/Engine/Source/Editor/Kismet/Private/BlueprintEditor.cpp
--- a/Engine/Source/Editor/Kismet/Private/BlueprintEditor.cpp
+++ b/Engine/Source/Editor/Kismet/Private/BlueprintEditor.cpp
@@ -3865,7 +3865,7 @@ void FBlueprintEditor::Compile()
 
 		if (bIsReparentingBlueprint)
 		{
-			CompileOptions |= EBlueprintCompileOptions::UseDeltaSerializationDuringReinstancing;
+			CompileOptions |= (EBlueprintCompileOptions::UseDeltaSerializationDuringReinstancing | EBlueprintCompileOptions::SkipNewVariableDefaultsDetection);
 		}
 
 		FKismetEditorUtilities::CompileBlueprint(BlueprintObj, CompileOptions, &LogResults);
diff --git a/Engine/Source/Editor/UnrealEd/Public/Kismet2/KismetEditorUtilities.h b/Engine/Source/Editor/UnrealEd/Public/Kismet2/KismetEditorUtilities.h
--- a/Engine/Source/Editor/UnrealEd/Public/Kismet2/KismetEditorUtilities.h
+++ b/Engine/Source/Editor/UnrealEd/Public/Kismet2/KismetEditorUtilities.h
@@ -56,6 +56,8 @@ enum class EBlueprintCompileOptions
 	SkipFiBSearchMetaUpdate = 0x200,
 	/** Allow the delta serialization during FBlueprintCompileReinstancer::CopyPropertiesForUnrelatedObjects */
 	UseDeltaSerializationDuringReinstancing = 0x400,
+	/** Skips the new variable defaults detection - in some cases we do not want to use the defaults from the generated class such as during a reparent */
+	SkipNewVariableDefaultsDetection = 0x800,
 };
 
 ENUM_CLASS_FLAGS(EBlueprintCompileOptions)
diff --git a/Engine/Source/Runtime/Engine/Classes/Engine/Blueprint.h b/Engine/Source/Runtime/Engine/Classes/Engine/Blueprint.h
--- a/Engine/Source/Runtime/Engine/Classes/Engine/Blueprint.h
+++ b/Engine/Source/Runtime/Engine/Classes/Engine/Blueprint.h
@@ -180,6 +180,9 @@ public:
 	/** Whether or not to use Delta Serialization when copying unrelated objects */
 	bool bUseDeltaSerializationDuringReinstancing;
 
+	/** Whether or not to skip new variable defaults detection */
+	bool bSkipNewVariableDefaultsDetection;
+
 	TSharedPtr<FString> OutHeaderSourceCode;
 	TSharedPtr<FString> OutCppSourceCode;
 
