Do not allow pass by reference to be changed for array properties as they are implicitly always pass by reference
Do not show default value field for containers

diff --git a/Engine/Source/Editor/Kismet/Private/BlueprintDetailsCustomization.cpp b/Engine/Source/Editor/Kismet/Private/BlueprintDetailsCustomization.cpp
--- a/Engine/Source/Editor/Kismet/Private/BlueprintDetailsCustomization.cpp
+++ b/Engine/Source/Editor/Kismet/Private/BlueprintDetailsCustomization.cpp
@@ -3199,7 +3199,8 @@ void FBlueprintGraphArgumentLayout::GenerateChildContent( IDetailChildrenBuilder
 			const UClass* ClassObject = Cast<UClass>(FoundPin->PinType.PinSubCategoryObject.Get());
 			const bool bTypeWithNoDefaults = (FoundPin->PinType.PinCategory == UEdGraphSchema_K2::PC_Object) || (FoundPin->PinType.PinCategory == UEdGraphSchema_K2::PC_Class) || (FoundPin->PinType.PinCategory == UEdGraphSchema_K2::PC_Interface) 
 				|| (FoundPin->PinType.PinCategory == UEdGraphSchema_K2::PC_SoftObject && ClassObject && ClassObject->IsChildOf(AActor::StaticClass()))
-				|| UEdGraphSchema_K2::IsExecPin(*FoundPin);
+				|| UEdGraphSchema_K2::IsExecPin(*FoundPin)
+				|| FoundPin->PinType.IsContainer();
 
 			if (!FoundPin->PinType.bIsReference && !bTypeWithNoDefaults)
 			{
@@ -3249,20 +3250,27 @@ void FBlueprintGraphArgumentLayout::GenerateChildContent( IDetailChildrenBuilder
 		// Exec pins can't be passed by reference
 		if (FoundPin && !UEdGraphSchema_K2::IsExecPin(*FoundPin) && !bMacroGraph)
 		{
+			auto ShouldPassByRefBeReadOnly = [this]
+			{
+				// Array types will always be implicitly passed by reference, regardless of
+				// the checkbox setting so make it readonly.
+				return OnGetPinInfo().IsArray() || ShouldPinBeReadOnly();
+			};
+
 			ChildrenBuilder.AddCustomRow(LOCTEXT("FunctionArgDetailsPassByReference", "Pass-by-Reference"))
 				.NameContent()
 				[
 					SNew(STextBlock)
 					.Text(LOCTEXT("FunctionArgDetailsPassByReference", "Pass-by-Reference"))
-				.ToolTipText(LOCTEXT("FunctionArgDetailsPassByReferenceTooltip", "Pass this parameter by reference?"))
-				.Font(IDetailLayoutBuilder::GetDetailFont())
+					.ToolTipText(LOCTEXT("FunctionArgDetailsPassByReferenceTooltip", "Pass this parameter by reference?"))
+					.Font(IDetailLayoutBuilder::GetDetailFont())
 				]
 			.ValueContent()
 				[
 					SNew(SCheckBox)
 					.IsChecked(this, &FBlueprintGraphArgumentLayout::IsRefChecked)
-				.OnCheckStateChanged(this, &FBlueprintGraphArgumentLayout::OnRefCheckStateChanged)
-				.IsEnabled(!ShouldPinBeReadOnly())
+					.OnCheckStateChanged(this, &FBlueprintGraphArgumentLayout::OnRefCheckStateChanged)
+					.IsEnabled(!ShouldPassByRefBeReadOnly())
 				];
 		}
 	}
@@ -3407,8 +3415,11 @@ UEdGraphPin* FBlueprintGraphArgumentLayout::GetPin() const
 
 ECheckBoxState FBlueprintGraphArgumentLayout::IsRefChecked() const
 {
-	FEdGraphPinType PinType = OnGetPinInfo();
-	return PinType.bIsReference? ECheckBoxState::Checked : ECheckBoxState::Unchecked;
+	const FEdGraphPinType PinType = OnGetPinInfo();
+
+	// Array types will always be implicitly passed by reference, regardless of
+	// the checkbox setting so show it as checked
+	return (PinType.bIsReference || PinType.IsArray())  ? ECheckBoxState::Checked : ECheckBoxState::Unchecked;
 }
 
 void FBlueprintGraphArgumentLayout::OnRefCheckStateChanged(ECheckBoxState InState)
