BlueprintDebugger: Fixed missing Window menu on Mac. Menus go to the main application bar on MacOS and thus the tabmanager needs to be aware of the menu to supply it to the application bar.

diff --git a/Engine/Source/Editor/Kismet/Private/BlueprintDebugger.cpp b/Engine/Source/Editor/Kismet/Private/BlueprintDebugger.cpp
--- a/Engine/Source/Editor/Kismet/Private/BlueprintDebugger.cpp
+++ b/Engine/Source/Editor/Kismet/Private/BlueprintDebugger.cpp
@@ -267,12 +267,14 @@ TSharedRef<SDockTab> FBlueprintDebuggerImpl::CreateBluprintDebuggerTab(const FSp
 			})
 	);
 
+	TSharedRef<SWidget> MenuBarWidget = MenuBarBuilder.MakeWidget();
+
 	NomadTab->SetContent(
 		SNew(SVerticalBox)
 		+ SVerticalBox::Slot()
 		.AutoHeight()
 		[
-			MenuBarBuilder.MakeWidget()
+			MenuBarWidget
 		]
 		+SVerticalBox::Slot()
 		[
@@ -285,6 +287,9 @@ TSharedRef<SDockTab> FBlueprintDebuggerImpl::CreateBluprintDebuggerTab(const FSp
 		]
 	);
 
+	// Tell tab-manager about the multi-box for platforms with a global menu bar
+	DebuggingToolsTabManager->SetMenuMultiBox(MenuBarBuilder.GetMultiBox(), MenuBarWidget);
+
 	return NomadTab;
 }
 
