Regenerate function signatures referenced by call sites after renaming input/output pins on function graph entry/result nodes.

diff --git a/Engine/Source/Editor/BlueprintGraph/Private/CallFunctionHandler.cpp b/Engine/Source/Editor/BlueprintGraph/Private/CallFunctionHandler.cpp
--- a/Engine/Source/Editor/BlueprintGraph/Private/CallFunctionHandler.cpp
+++ b/Engine/Source/Editor/BlueprintGraph/Private/CallFunctionHandler.cpp
@@ -704,9 +704,13 @@ UFunction* FKCHandler_CallFunction::FindFunction(FKismetFunctionContext& Context
 
 	if (CallingContext)
 	{
+		const UBlueprint* BlueprintContext = UBlueprint::GetBlueprintFromClass(CallingContext);
+
+		// Redirect the calling context to the most up-to-date class (when not up-to-date,
+		// this will redirect to the Blueprint's skeleton class)
 		// It may be advisable to always do this branch in GetMostUpToDateClass, but
 		// being conservative:
-		if (!FBlueprintCompilationManager::IsGeneratedClassLayoutReady())
+		if (!BlueprintContext || (BlueprintContext->bBeingCompiled && !FBlueprintCompilationManager::IsGeneratedClassLayoutReady()) || (!BlueprintContext->bBeingCompiled && !BlueprintContext->IsUpToDate()))
 		{
 			CallingContext = FBlueprintEditorUtils::GetMostUpToDateClass(CallingContext);
 		}
diff --git a/Engine/Source/Editor/Kismet/Private/BlueprintDetailsCustomization.cpp b/Engine/Source/Editor/Kismet/Private/BlueprintDetailsCustomization.cpp
--- a/Engine/Source/Editor/Kismet/Private/BlueprintDetailsCustomization.cpp
+++ b/Engine/Source/Editor/Kismet/Private/BlueprintDetailsCustomization.cpp
@@ -4685,23 +4685,30 @@ bool FBaseBlueprintGraphActionDetails::OnPinRenamed(UK2Node_EditablePinBase* Tar
 		{
 			TerminalNodes.Add(EntryNode);
 		}
+
+		bool bRequiresFunctionSignatureUpdate = false;
 		for (UK2Node_EditablePinBase* TerminalNode : TerminalNodes)
 		{
 			TerminalNode->Modify();
 			PinRenamedHelper.NodesToRename.Add(TerminalNode);
+
+			// Since function terminator node pins map to generated function properties, we need to
+			// regenerate the referenced function so that dependent pins can be reconstructed properly.
+			bRequiresFunctionSignatureUpdate |= TerminalNode->IsA<UK2Node_FunctionTerminator>();
 		}
 
-		PinRenamedHelper.ModifiedBlueprints.Add(GetBlueprintObj());
+		UBlueprint* TargetBlueprint = GetBlueprintObj();
+		PinRenamedHelper.ModifiedBlueprints.Add(TargetBlueprint);
 
 		// GATHER 
-		PinRenamedHelper.Broadcast(GetBlueprintObj(), TargetNode, Graph);
+		PinRenamedHelper.Broadcast(TargetBlueprint, TargetNode, Graph);
 
 		const FName NewFName = *NewName;
 
 		// TEST
 		for (UK2Node* NodeToRename : PinRenamedHelper.NodesToRename)
 		{
-			if (ERenamePinResult::ERenamePinResult_NameCollision == NodeToRename->RenameUserDefinedPin(OldName, NewFName, true))
+			if (ERenamePinResult::ERenamePinResult_NameCollision == NodeToRename->RenameUserDefinedPin(OldName, NewFName, /*bTest =*/ true))
 			{
 				return false;
 			}
@@ -4710,9 +4717,12 @@ bool FBaseBlueprintGraphActionDetails::OnPinRenamed(UK2Node_EditablePinBase* Tar
 		// UPDATE
 		for (UK2Node* NodeToRename : PinRenamedHelper.NodesToRename)
 		{
-			NodeToRename->RenameUserDefinedPin(OldName, NewFName, false);
+			// Note: This will internally call Modify() on any matching pin(s).
+			NodeToRename->RenameUserDefinedPin(OldName, NewFName, /*bTest =*/ false);
 		}
 
+		// Update the corresponding UserDefinedPins entry for each terminal node.
+		// Note: This array is not serialized, so a Modify() here isn't necessary.
 		for (UK2Node_EditablePinBase* TerminalNode : TerminalNodes)
 		{
 			TSharedPtr<FUserPinInfo>* UDPinPtr = TerminalNode->UserDefinedPins.FindByPredicate([&](TSharedPtr<FUserPinInfo>& Pin)
@@ -4724,6 +4734,18 @@ bool FBaseBlueprintGraphActionDetails::OnPinRenamed(UK2Node_EditablePinBase* Tar
 				(*UDPinPtr)->PinName = NewFName;
 			}
 		}
+
+		// If necessary, regenerate the skeleton class to update function properties.
+		if (bRequiresFunctionSignatureUpdate)
+		{
+			FBlueprintEditorUtils::MarkBlueprintAsStructurallyModified(TargetBlueprint);
+		}
+
+		// Trigger a change notification for Blueprints that were updated, in case there's anything we need to refresh.
+		for (UBlueprint* ModifiedBlueprint : PinRenamedHelper.ModifiedBlueprints)
+		{
+			ModifiedBlueprint->BroadcastChanged();
+		}
 	}
 	return true;
 }
