Fix an editor access violation crash caused by a reference to a stale property wrapper object after its underlying Blueprint skeleton class has gone out of scope.

diff --git a/Engine/Source/Editor/Kismet/Private/BlueprintDetailsCustomization.cpp b/Engine/Source/Editor/Kismet/Private/BlueprintDetailsCustomization.cpp
--- a/Engine/Source/Editor/Kismet/Private/BlueprintDetailsCustomization.cpp
+++ b/Engine/Source/Editor/Kismet/Private/BlueprintDetailsCustomization.cpp
@@ -245,19 +245,6 @@ FBlueprintVarActionDetails::~FBlueprintVarActionDetails()
 BEGIN_SLATE_FUNCTION_BUILD_OPTIMIZATION
 void FBlueprintVarActionDetails::CustomizeDetails( IDetailLayoutBuilder& DetailLayout )
 {
-	// update the property being edited
-	PropertyBeingEdited = nullptr;
-	TArray<TWeakObjectPtr<UObject>> ObjectsBeingEdited;
-	DetailLayout.GetObjectsBeingCustomized(ObjectsBeingEdited);
-	for(TWeakObjectPtr<UObject> Object : ObjectsBeingEdited)
-	{
-		if(UPropertyWrapper* Wrapper = Cast<UPropertyWrapper>(Object.Get()))
-		{
-			PropertyBeingEdited = Wrapper->GetProperty();
-			break;
-		}
-	}
-
 	CachedVariableProperty = SelectionAsProperty();
 
 	if(!CachedVariableProperty.IsValid())
@@ -1670,12 +1657,17 @@ FProperty* FBlueprintVarActionDetails::SelectionAsProperty() const
 	{
 		return LocalVarAction->GetProperty();
 	}
+	FEdGraphSchemaAction_K2Delegate* DelegateVar = MyBlueprintSelectionAsDelegate();
+	if (DelegateVar)
+	{
+		return DelegateVar->GetDelegateProperty();
+	}
 	UK2Node_Variable* GraphVar = EdGraphSelectionAsVar();
 	if(GraphVar)
 	{
 		return GraphVar->GetPropertyForVariable();
 	}
-	return PropertyBeingEdited;
+	return nullptr;
 }
 
 FName FBlueprintVarActionDetails::GetVariableName() const
@@ -1690,15 +1682,16 @@ FName FBlueprintVarActionDetails::GetVariableName() const
 	{
 		return LocalVarAction->GetVariableName();
 	}
+	FEdGraphSchemaAction_K2Delegate* DelegateVar = MyBlueprintSelectionAsDelegate();
+	if (DelegateVar)
+	{
+		return DelegateVar->GetDelegateName();
+	}
 	UK2Node_Variable* GraphVar = EdGraphSelectionAsVar();
 	if(GraphVar)
 	{
 		return GraphVar->GetVarName();
 	}
-	if(PropertyBeingEdited)
-	{
-		return PropertyBeingEdited->GetFName();
-	}
 	return NAME_None;
 }
 
diff --git a/Engine/Source/Editor/Kismet/Private/BlueprintDetailsCustomization.h b/Engine/Source/Editor/Kismet/Private/BlueprintDetailsCustomization.h
--- a/Engine/Source/Editor/Kismet/Private/BlueprintDetailsCustomization.h
+++ b/Engine/Source/Editor/Kismet/Private/BlueprintDetailsCustomization.h
@@ -122,8 +122,9 @@ public:
 
 private:
 	/** Accessors passed to parent */
-	FEdGraphSchemaAction_K2Var* MyBlueprintSelectionAsVar() const {return MyBlueprint.Pin()->SelectionAsVar();}
-	FEdGraphSchemaAction_K2LocalVar* MyBlueprintSelectionAsLocalVar() const {return MyBlueprint.Pin()->SelectionAsLocalVar();}
+	FEdGraphSchemaAction_K2Var* MyBlueprintSelectionAsVar() const { return MyBlueprint.Pin()->SelectionAsVar(); }
+	FEdGraphSchemaAction_K2LocalVar* MyBlueprintSelectionAsLocalVar() const { return MyBlueprint.Pin()->SelectionAsLocalVar(); }
+	FEdGraphSchemaAction_K2Delegate* MyBlueprintSelectionAsDelegate() const { return MyBlueprint.Pin()->SelectionAsDelegate(); }
 	UK2Node_Variable* EdGraphSelectionAsVar() const;
 	FProperty* SelectionAsProperty() const;
 	FName GetVariableName() const;
@@ -321,9 +322,6 @@ private:
 
 	/** External detail customizations */
 	TArray<TSharedPtr<IDetailCustomization>> ExternalDetailCustomizations;
-
-	/** If a property is being edited right now, this will cache it */
-	FProperty* PropertyBeingEdited;
 };
 
 class FBaseBlueprintGraphActionDetails : public IDetailCustomization
