Ensure that Blueprint graph nodes are always fixed up on load as part of the precompile phase, even for fast paths (e.g. data-only BPs).

This change also reverts the forced compilation of all Blueprints on load if below a certain asset version that was previously introduced to fixup disabled (default) nodes in a data-only BP.

diff --git a/Engine/Source/Editor/Kismet/Private/BlueprintCompilationManager.cpp b/Engine/Source/Editor/Kismet/Private/BlueprintCompilationManager.cpp
--- a/Engine/Source/Editor/Kismet/Private/BlueprintCompilationManager.cpp
+++ b/Engine/Source/Editor/Kismet/Private/BlueprintCompilationManager.cpp
@@ -425,7 +425,7 @@ struct FCompilerData
 	bool ShouldValidate() const { return JobType == ECompilationManagerJobType::Normal; }
 	bool ShouldRegenerateSkeleton() const { return JobType != ECompilationManagerJobType::RelinkOnly; }
 	bool ShouldMarkUpToDateAfterSkeletonStage() const { return IsSkeletonOnly(); }
-	bool ShouldReconstructNodes() const { return JobType == ECompilationManagerJobType::Normal; }
+	bool ShouldReconstructNodes() const { return JobType == ECompilationManagerJobType::Normal || BP->bIsRegeneratingOnLoad; }
 	bool ShouldSkipReinstancerCreation() const { return (IsSkeletonOnly() && (!BP->ParentClass || BP->ParentClass->IsNative())); }
 	bool ShouldInitiateReinstancing() const { return JobType == ECompilationManagerJobType::Normal || BP->bIsRegeneratingOnLoad; }
 	bool ShouldCompileClassLayout() const { return JobType == ECompilationManagerJobType::Normal; }
diff --git a/Engine/Source/Runtime/Engine/Classes/Engine/Blueprint.h b/Engine/Source/Runtime/Engine/Classes/Engine/Blueprint.h
--- a/Engine/Source/Runtime/Engine/Classes/Engine/Blueprint.h
+++ b/Engine/Source/Runtime/Engine/Classes/Engine/Blueprint.h
@@ -856,7 +856,7 @@ public:
 	/**
 	* Allows derived blueprints to require compilation on load, otherwise they may get treated as data only and not compiled on load.
 	*/
-	virtual bool AlwaysCompileOnLoad() const;
+	virtual bool AlwaysCompileOnLoad() const { return false; }
 
 	/**
 	 * Some Blueprints (and classes) can recompile while we are debugging a live session (play in editor).
diff --git a/Engine/Source/Runtime/Engine/Private/Blueprint.cpp b/Engine/Source/Runtime/Engine/Private/Blueprint.cpp
--- a/Engine/Source/Runtime/Engine/Private/Blueprint.cpp
+++ b/Engine/Source/Runtime/Engine/Private/Blueprint.cpp
@@ -876,16 +876,6 @@ void UBlueprint::GetReparentingRules(TSet< const UClass* >& AllowedChildrenOfCla
 
 }
 
-bool UBlueprint::AlwaysCompileOnLoad() const
-{
-	bool bShouldCompile = false;
-
-	bShouldCompile = 
-		GetLinkerCustomVersion(FUE5ReleaseStreamObjectVersion::GUID) <  FUE5ReleaseStreamObjectVersion::BlueprintPinsUseRealNumbers;
-
-	return bShouldCompile;
-}
-
 bool UBlueprint::CanAlwaysRecompileWhilePlayingInEditor() const
 {
 	return false;
