Cleaned up details customization for Blueprint variable actions.

diff --git a/Engine/Source/Editor/Kismet/Private/BlueprintDetailsCustomization.cpp b/Engine/Source/Editor/Kismet/Private/BlueprintDetailsCustomization.cpp
--- a/Engine/Source/Editor/Kismet/Private/BlueprintDetailsCustomization.cpp
+++ b/Engine/Source/Editor/Kismet/Private/BlueprintDetailsCustomization.cpp
@@ -1642,56 +1642,34 @@ UK2Node_Variable* FBlueprintVarActionDetails::EdGraphSelectionAsVar() const
 			}
 		}
 	}
-	return NULL;
+	return nullptr;
 }
 
 FProperty* FBlueprintVarActionDetails::SelectionAsProperty() const
 {
-	FEdGraphSchemaAction_K2Var* VarAction = MyBlueprintSelectionAsVar();
-	if(VarAction)
-	{
-		return VarAction->GetProperty();
-	}
-	FEdGraphSchemaAction_K2LocalVar* LocalVarAction = MyBlueprintSelectionAsLocalVar();
-	if(LocalVarAction)
-	{
-		return LocalVarAction->GetProperty();
-	}
-	FEdGraphSchemaAction_K2Delegate* DelegateVar = MyBlueprintSelectionAsDelegate();
-	if (DelegateVar)
+	if (FEdGraphSchemaAction_BlueprintVariableBase* BPVar = MyBlueprint.Pin()->SelectionAsBlueprintVariable())
 	{
-		return DelegateVar->GetDelegateProperty();
+		return BPVar->GetProperty();
 	}
-	UK2Node_Variable* GraphVar = EdGraphSelectionAsVar();
-	if(GraphVar)
+	else if (UK2Node_Variable* GraphVar = EdGraphSelectionAsVar())
 	{
 		return GraphVar->GetPropertyForVariable();
 	}
+
 	return nullptr;
 }
 
 FName FBlueprintVarActionDetails::GetVariableName() const
 {
-	FEdGraphSchemaAction_K2Var* VarAction = MyBlueprintSelectionAsVar();
-	if(VarAction)
-	{
-		return VarAction->GetVariableName();
-	}
-	FEdGraphSchemaAction_K2LocalVar* LocalVarAction = MyBlueprintSelectionAsLocalVar();
-	if(LocalVarAction)
+	if (FEdGraphSchemaAction_BlueprintVariableBase* BPVar = MyBlueprint.Pin()->SelectionAsBlueprintVariable())
 	{
-		return LocalVarAction->GetVariableName();
+		return BPVar->GetVariableName();
 	}
-	FEdGraphSchemaAction_K2Delegate* DelegateVar = MyBlueprintSelectionAsDelegate();
-	if (DelegateVar)
-	{
-		return DelegateVar->GetDelegateName();
-	}
-	UK2Node_Variable* GraphVar = EdGraphSelectionAsVar();
-	if(GraphVar)
+	else if (UK2Node_Variable* GraphVar = EdGraphSelectionAsVar())
 	{
 		return GraphVar->GetVarName();
 	}
+
 	return NAME_None;
 }
 
diff --git a/Engine/Source/Editor/Kismet/Private/BlueprintDetailsCustomization.h b/Engine/Source/Editor/Kismet/Private/BlueprintDetailsCustomization.h
--- a/Engine/Source/Editor/Kismet/Private/BlueprintDetailsCustomization.h
+++ b/Engine/Source/Editor/Kismet/Private/BlueprintDetailsCustomization.h
@@ -122,9 +122,6 @@ public:
 
 private:
 	/** Accessors passed to parent */
-	FEdGraphSchemaAction_K2Var* MyBlueprintSelectionAsVar() const { return MyBlueprint.Pin()->SelectionAsVar(); }
-	FEdGraphSchemaAction_K2LocalVar* MyBlueprintSelectionAsLocalVar() const { return MyBlueprint.Pin()->SelectionAsLocalVar(); }
-	FEdGraphSchemaAction_K2Delegate* MyBlueprintSelectionAsDelegate() const { return MyBlueprint.Pin()->SelectionAsDelegate(); }
 	UK2Node_Variable* EdGraphSelectionAsVar() const;
 	FProperty* SelectionAsProperty() const;
 	FName GetVariableName() const;
