UE-141316 - Real variables don't have Slider range and Value range settings in the details panel

FBlueprintVarActionDetails::RangeVisibility() was missing a condition for double properties.

diff --git a/Engine/Source/Editor/Kismet/Private/BlueprintDetailsCustomization.cpp b/Engine/Source/Editor/Kismet/Private/BlueprintDetailsCustomization.cpp
--- a/Engine/Source/Editor/Kismet/Private/BlueprintDetailsCustomization.cpp
+++ b/Engine/Source/Editor/Kismet/Private/BlueprintDetailsCustomization.cpp
@@ -2179,7 +2179,7 @@ EVisibility FBlueprintVarActionDetails::RangeVisibility() const
 	{
 		const bool bIsInteger = VariableProperty->IsA(FIntProperty::StaticClass());
 		const bool bIsNonEnumByte = (VariableProperty->IsA(FByteProperty::StaticClass()) && CastField<const FByteProperty>(VariableProperty)->Enum == nullptr);
-		const bool bIsFloat = VariableProperty->IsA(FFloatProperty::StaticClass());
+		const bool bIsReal = (VariableProperty->IsA(FFloatProperty::StaticClass()) || VariableProperty->IsA(FDoubleProperty::StaticClass()));
 
 		// If this is a struct property than we must check the name of the struct it points to, so we can check
 		// if it supports the editing of the UIMin/UIMax metadata
@@ -2187,7 +2187,7 @@ EVisibility FBlueprintVarActionDetails::RangeVisibility() const
 		const UStruct* InnerStruct = StructProp ? StructProp->Struct : nullptr;
 		const bool bIsSupportedStruct = InnerStruct ? RangeVisibilityUtils::StructsSupportingRangeVisibility.Contains(InnerStruct->GetFName()) : false;
 
-		if (IsABlueprintVariable(VariableProperty) && (bIsInteger || bIsNonEnumByte || bIsFloat || bIsSupportedStruct))
+		if (IsABlueprintVariable(VariableProperty) && (bIsInteger || bIsNonEnumByte || bIsReal || bIsSupportedStruct))
 		{
 			return EVisibility::Visible;
 		}
