Fix up delegate bindings after renaming a custom event node without requiring a full compilation pass.

Also fixes an issue that didn't allow returning a custom event node back to its previous name without also requiring a full compilation pass.

diff --git a/Engine/Source/Editor/BlueprintGraph/Private/K2Node_CustomEvent.cpp b/Engine/Source/Editor/BlueprintGraph/Private/K2Node_CustomEvent.cpp
--- a/Engine/Source/Editor/BlueprintGraph/Private/K2Node_CustomEvent.cpp
+++ b/Engine/Source/Editor/BlueprintGraph/Private/K2Node_CustomEvent.cpp
@@ -67,31 +67,67 @@ public:
 		: FKismetNameValidator(CustomEventIn->GetBlueprint(), CustomEventIn->CustomFunctionName)
 		, CustomEvent(CustomEventIn)
 	{
-		check(CustomEvent != NULL);
+		check(CustomEvent != nullptr);
 	}
 
 	// Begin INameValidatorInterface
 	virtual EValidatorResult IsValid(FString const& Name, bool bOriginal = false) override
 	{
+		UBlueprint* Blueprint = CustomEvent->GetBlueprint();
+		check(Blueprint != nullptr);
+
 		EValidatorResult NameValidity = FKismetNameValidator::IsValid(Name, bOriginal);
 		if ((NameValidity == EValidatorResult::Ok) || (NameValidity == EValidatorResult::ExistingName))
 		{
-			UBlueprint* Blueprint = CustomEvent->GetBlueprint();
-			check(Blueprint != NULL);
-
 			UFunction* ParentFunction = FindUField<UFunction>(Blueprint->ParentClass, *Name);
 			// if this custom-event is overriding a function belonging to the blueprint's parent
-			if (ParentFunction != NULL)
+			if (ParentFunction != nullptr)
 			{
 				UK2Node_CustomEvent const* OverriddenEvent = FindCustomEventNodeFromFunction(ParentFunction);
 				// if the function that we're overriding isn't another custom event,
 				// then we can't name it this (only allow custom-event to override other custom-events)
-				if (OverriddenEvent == NULL)
+				if (OverriddenEvent == nullptr)
 				{
 					NameValidity = EValidatorResult::AlreadyInUse;
 				}		
 			}
 		}
+		else if (NameValidity == EValidatorResult::AlreadyInUse)
+		{
+			auto Predicate_EventGraphs = [Name](const TObjectPtr<UEdGraph>& InEventGraph) -> bool
+			{
+				return !InEventGraph.IsNull() && InEventGraph->HasAnyFlags(RF_Transient) && InEventGraph->GetName() == Name;
+			};
+
+			// Allow a transient event subgraph (compiler artifact) that matches the existing name
+			// to pass if there are no other event nodes that would use this name at compile time.
+			// This type of collision is a false positive that won't result in a conflict, because
+			// the custom event node won't enter the Blueprint's namespace until the next compile,
+			// and the compiler will regenerate the transient event subgraphs array on a full pass.
+			if (Blueprint->EventGraphs.FindByPredicate(Predicate_EventGraphs))
+			{
+				TArray<UK2Node_Event*> AllEventNodes;
+				FBlueprintEditorUtils::GetAllNodesOfClass<UK2Node_Event>(Blueprint, AllEventNodes);
+				UK2Node_Event** MatchingNodePtr = AllEventNodes.FindByPredicate([Name](UK2Node_Event* InEventNode)
+				{
+					if (InEventNode->bOverrideFunction)
+					{
+						return InEventNode->EventReference.GetMemberName().ToString() == Name;
+					}
+					else if (InEventNode->CustomFunctionName != NAME_None)
+					{
+						return InEventNode->CustomFunctionName.ToString() == Name;
+					}
+					
+					return false;
+				});
+
+				if (!MatchingNodePtr || *MatchingNodePtr == CustomEvent)
+				{
+					NameValidity = EValidatorResult::Ok;
+				}
+			}
+		}
 		return NameValidity;
 	}
 	// End INameValidatorInterface
diff --git a/Engine/Source/Editor/Kismet/Private/BlueprintCompilationManager.cpp b/Engine/Source/Editor/Kismet/Private/BlueprintCompilationManager.cpp
--- a/Engine/Source/Editor/Kismet/Private/BlueprintCompilationManager.cpp
+++ b/Engine/Source/Editor/Kismet/Private/BlueprintCompilationManager.cpp
@@ -1494,60 +1494,63 @@ void FBlueprintCompilationManagerImpl::FlushCompilationQueueImpl(bool bSuppressB
 
 			UBlueprint* BP = CompilerData.BP;
 
-			if (!CompilerData.IsSkeletonOnly())
+			// Ensure that delegate nodes/pins conform to compiled function names/signatures. In general,
+			// these will rely on the skeleton class, so we must also include it for a skeleton-only pass.
+			FBlueprintEditorUtils::UpdateDelegatesInBlueprint(BP);
+
+			// These are post-compile validations that generally rely on a fully-generated up-to-date class,
+			// so we defer it on skeleton-only passes as well as during compile-on-load, where it's generally
+			// going to be handled during the PostLoad() phase after all dependencies have been loaded/compiled.
+			if (!CompilerData.IsSkeletonOnly() && !BP->bIsRegeneratingOnLoad && BP->GeneratedClass)
 			{
-				FBlueprintEditorUtils::UpdateDelegatesInBlueprint(BP);
-				if (!BP->bIsRegeneratingOnLoad && BP->GeneratedClass)
+				FKismetEditorUtilities::StripExternalComponents(BP);
+
+				if (BP->SimpleConstructionScript)
 				{
-					FKismetEditorUtilities::StripExternalComponents(BP);
+					BP->SimpleConstructionScript->FixupRootNodeParentReferences();
+				}
 
-					if (BP->SimpleConstructionScript)
+				if (BP->Status != BS_Error)
+				{
+					if (CompilerData.Compiler.IsValid())
 					{
-						BP->SimpleConstructionScript->FixupRootNodeParentReferences();
-					}
+						// Route through the compiler context in order to perform type-specific Blueprint class validation.
+						CompilerData.Compiler->ValidateGeneratedClass(CastChecked<UBlueprintGeneratedClass>(BP->GeneratedClass));
 
-					if (BP->Status != BS_Error)
-					{
-						if (CompilerData.Compiler.IsValid())
+						if (CompilerData.ShouldValidateClassDefaultObject())
 						{
-							// Route through the compiler context in order to perform type-specific Blueprint class validation.
-							CompilerData.Compiler->ValidateGeneratedClass(CastChecked<UBlueprintGeneratedClass>(BP->GeneratedClass));
-
-							if (CompilerData.ShouldValidateClassDefaultObject())
+							// Our CDO should be properly constructed by this point and should always exist
+							UObject* ClassDefaultObject = BP->GeneratedClass->GetDefaultObject(false);
+							if (ensureAlways(ClassDefaultObject))
 							{
-								// Our CDO should be properly constructed by this point and should always exist
-								UObject* ClassDefaultObject = BP->GeneratedClass->GetDefaultObject(false);
-								if (ensureAlways(ClassDefaultObject))
-								{
-									FKismetCompilerUtilities::ValidateEnumProperties(ClassDefaultObject, *CompilerData.ActiveResultsLog);
+								FKismetCompilerUtilities::ValidateEnumProperties(ClassDefaultObject, *CompilerData.ActiveResultsLog);
 
-									// Make sure any class-specific validation passes on the CDO
-									TArray<FText> ValidationErrors;
-									EDataValidationResult ValidateCDOResult = ClassDefaultObject->IsDataValid(/*out*/ ValidationErrors);
-									if (ValidateCDOResult == EDataValidationResult::Invalid)
+								// Make sure any class-specific validation passes on the CDO
+								TArray<FText> ValidationErrors;
+								EDataValidationResult ValidateCDOResult = ClassDefaultObject->IsDataValid(/*out*/ ValidationErrors);
+								if (ValidateCDOResult == EDataValidationResult::Invalid)
+								{
+									for (const FText& ValidationError : ValidationErrors)
 									{
-										for (const FText& ValidationError : ValidationErrors)
-										{
-											CompilerData.ActiveResultsLog->Error(*ValidationError.ToString());
-										}
+										CompilerData.ActiveResultsLog->Error(*ValidationError.ToString());
 									}
+								}
 
-									// Adjust Blueprint status to match anything new that was found during validation.
-									if (CompilerData.ActiveResultsLog->NumErrors > 0)
-									{
-										BP->Status = BS_Error;
-									}
-									else if (BP->Status == BS_UpToDate && CompilerData.ActiveResultsLog->NumWarnings > 0)
-									{
-										BP->Status = BS_UpToDateWithWarnings;
-									}
+								// Adjust Blueprint status to match anything new that was found during validation.
+								if (CompilerData.ActiveResultsLog->NumErrors > 0)
+								{
+									BP->Status = BS_Error;
+								}
+								else if (BP->Status == BS_UpToDate && CompilerData.ActiveResultsLog->NumWarnings > 0)
+								{
+									BP->Status = BS_UpToDateWithWarnings;
 								}
 							}
 						}
-						else
-						{
-							UBlueprint::ValidateGeneratedClass(BP->GeneratedClass);
-						}
+					}
+					else
+					{
+						UBlueprint::ValidateGeneratedClass(BP->GeneratedClass);
 					}
 				}
 			}
