Remove MultiLine metadata tag instead of setting it to false in blueprint variable details customization

diff --git a/Engine/Source/Editor/Kismet/Private/BlueprintDetailsCustomization.cpp b/Engine/Source/Editor/Kismet/Private/BlueprintDetailsCustomization.cpp
--- a/Engine/Source/Editor/Kismet/Private/BlueprintDetailsCustomization.cpp
+++ b/Engine/Source/Editor/Kismet/Private/BlueprintDetailsCustomization.cpp
@@ -2638,7 +2638,14 @@ void FBlueprintVarActionDetails::OnMultilineChanged(ECheckBoxState InNewState)
 	if (Property)
 	{
 		const bool bMultiline = (InNewState == ECheckBoxState::Checked);
-		FBlueprintEditorUtils::SetBlueprintVariableMetaData(GetBlueprintObj(), Property->GetFName(), GetLocalVariableScope(CachedVariableProperty.Get()), TEXT("MultiLine"), bMultiline ? TEXT("true") : TEXT("false"));
+		if (bMultiline)
+		{
+			FBlueprintEditorUtils::SetBlueprintVariableMetaData(GetBlueprintObj(), Property->GetFName(), GetLocalVariableScope(CachedVariableProperty.Get()), TEXT("MultiLine"), TEXT("true"));
+		}
+		else
+		{
+			FBlueprintEditorUtils::RemoveBlueprintVariableMetaData(GetBlueprintObj(), Property->GetFName(), GetLocalVariableScope(CachedVariableProperty.Get()), TEXT("MultiLine"));
+		}
 	}
 }
 
