SKismetInspector::IsPropertyVisible was failing to display members of structs that were members of instanced objects that were members of components, ie.:

class UComponent
{
   UPROPERTY(Instanced)
   class UInstanced
   {
      struct FStruct
      {
         int IncorrectlyHidden;
      }
   }
}

This was because only the instanced object property was stored in the SelectedObjectProperties array, and we were only checking the immediate parent of the property rather than the entire parent property chain.

UDN: https://epicgames.lightning.force.com/lightning/r/Case/5004z00001c3Lv8AAE/view

diff --git a/Engine/Source/Editor/Kismet/Private/SKismetInspector.cpp b/Engine/Source/Editor/Kismet/Private/SKismetInspector.cpp
--- a/Engine/Source/Editor/Kismet/Private/SKismetInspector.cpp
+++ b/Engine/Source/Editor/Kismet/Private/SKismetInspector.cpp
@@ -764,12 +764,17 @@ bool SKismetInspector::IsStructViewPropertyReadOnly(const struct FPropertyAndPar
 	return false;
 }
 
-bool SKismetInspector::IsAnyParentContainerSelected(const FPropertyAndParent& PropertyAndParent) const
+bool SKismetInspector::IsAnyParentOrContainerSelected(const FPropertyAndParent& PropertyAndParent) const
 {
 	for (const FProperty* CurrentProperty : PropertyAndParent.ParentProperties)
 	{
-		const FProperty* CurrentOuter = CurrentProperty->GetOwner<FProperty>();
+		if (SelectedObjectProperties.Find(const_cast<FProperty*>(CurrentProperty)))
+		{
+			return true;
+		}
 
+		// the property might be the Inner property of an array (or Key/Value of a map), so check if the outer property is selected
+		const FProperty* CurrentOuter = CurrentProperty->GetOwner<FProperty>();
 		if (CurrentOuter != nullptr && SelectedObjectProperties.Find(const_cast<FProperty*>(CurrentOuter)))
 		{
 			return true;
@@ -850,14 +855,7 @@ bool SKismetInspector::IsPropertyVisible( const FPropertyAndParent& PropertyAndP
 	}
 	else if ( PropertyAndParent.ParentProperties.Num() > 0 && SelectedObjectProperties.Num() > 0 )
 	{
-		const FProperty* ParentProperty = PropertyAndParent.ParentProperties[0];
-
-		if ( SelectedObjectProperties.Find( const_cast<FProperty*>( ParentProperty ) ) )
-		{
-			// If its parent is selected, it should be visible
-			return true;
-		}
-		else if ( IsAnyParentContainerSelected(PropertyAndParent) )
+		if (IsAnyParentOrContainerSelected(PropertyAndParent))
 		{
 			return true;
 		}
diff --git a/Engine/Source/Editor/Kismet/Public/SKismetInspector.h b/Engine/Source/Editor/Kismet/Public/SKismetInspector.h
--- a/Engine/Source/Editor/Kismet/Public/SKismetInspector.h
+++ b/Engine/Source/Editor/Kismet/Public/SKismetInspector.h
@@ -207,7 +207,7 @@ protected:
 	ECheckBoxState GetPublicViewCheckboxState() const;
 	void SetPublicViewCheckboxState(ECheckBoxState InIsChecked);
 
-	bool IsAnyParentContainerSelected(const FPropertyAndParent& PropertyAndParent) const;
+	bool IsAnyParentOrContainerSelected(const FPropertyAndParent& PropertyAndParent) const;
 
 	/** Callback invoked after a value change on the selected object(s) */
 	void OnFinishedChangingProperties(const FPropertyChangedEvent& InPropertyChangedEvent);
