Avoid stalling for long periods of time during map load and PIE
  - Remove slowtask throttling during load as it could prevent meaningful top level slowtask to be created
  - Fix delayed slowtask that was never showing when children were added before the delay
  - Add missing slowtask handling on long gamethread functions to improve UI responsiveness
  - Add PIE preparation slowtask scopes
  - Test to make sure the total amount of time spent in all slowtask functions is reasonable (1s on a 2 minute map load)

diff --git a/Engine/Source/Editor/Kismet/Private/BlueprintCompilationManager.cpp b/Engine/Source/Editor/Kismet/Private/BlueprintCompilationManager.cpp
--- a/Engine/Source/Editor/Kismet/Private/BlueprintCompilationManager.cpp
+++ b/Engine/Source/Editor/Kismet/Private/BlueprintCompilationManager.cpp
@@ -24,6 +24,7 @@
 #include "Kismet2/KismetEditorUtilities.h"
 #include "Kismet2/KismetReinstanceUtilities.h"
 #include "KismetCompiler.h"
+#include "Misc/ScopedSlowTask.h"
 #include "ProfilingDebugging/ScopedTimers.h"
 #include "Serialization/ArchiveHasReferences.h"
 #include "Serialization/ArchiveReplaceObjectRef.h"
@@ -516,6 +517,9 @@ void FBlueprintCompilationManagerImpl::FlushCompilationQueueImpl(bool bSuppressB
 		return;
 	}
 
+	FScopedSlowTask SlowTask(17.f /* Number of steps */, LOCTEXT("FlushCompilationQueue", "Compiling blueprints..."));
+	SlowTask.MakeDialogDelayed(1.0f);
+
 	TArray<FCompilerData> CurrentlyCompilingBPs;
 	{ // begin GTimeCompiling scope 
 		FScopedDurationTimer SetupTimer(GTimeCompiling); 
@@ -579,6 +583,8 @@ void FBlueprintCompilationManagerImpl::FlushCompilationQueueImpl(bool bSuppressB
 			}
 		}
 
+		SlowTask.EnterProgressFrame();
+
 		// then make sure any normal blueprints have their bytecode dependents recompiled, this is in case a function signature changes:
 		for(const FBPCompileRequestInternal& CompileJob : QueuedRequests)
 		{
@@ -615,6 +621,8 @@ void FBlueprintCompilationManagerImpl::FlushCompilationQueueImpl(bool bSuppressB
 			}
 		}
 
+		SlowTask.EnterProgressFrame();
+
 		// STAGE II: Filter out data only and interface blueprints:
 		for(int32 I = 0; I < QueuedRequests.Num(); ++I)
 		{
@@ -708,7 +716,9 @@ void FBlueprintCompilationManagerImpl::FlushCompilationQueueImpl(bool bSuppressB
 				BlueprintsToRecompile.Add(QueuedBP);
 			}
 		}
-			
+
+		SlowTask.EnterProgressFrame();
+
 		for(UBlueprint* BP : BlueprintsToRecompile)
 		{
 			// make sure all children are at least re-linked:
@@ -740,6 +750,10 @@ void FBlueprintCompilationManagerImpl::FlushCompilationQueueImpl(bool bSuppressB
 					}
 				}
 			}
+
+			// Don't report progress for substep but gives a chance to tick slate to improve the responsiveness of the 
+			// progress bar being shown. We expect slate to be ticked at regular intervals throughout the loading.
+			SlowTask.TickProgress();
 		}
 
 		/*	Prevent 'pending kill' blueprints from being recompiled. Dependency
@@ -781,6 +795,8 @@ void FBlueprintCompilationManagerImpl::FlushCompilationQueueImpl(bool bSuppressB
 		BlueprintsToRecompile.Empty();
 		QueuedRequests.Empty();
 
+		SlowTask.EnterProgressFrame();
+
 		// STAGE III: Sort into correct compilation order. We want to compile root types before their derived (child) types:
 		auto HierarchyDepthSortFn = [](const FCompilerData& CompilerDataA, const FCompilerData& CompilerDataB)
 		{
@@ -803,6 +819,8 @@ void FBlueprintCompilationManagerImpl::FlushCompilationQueueImpl(bool bSuppressB
 		};
 		CurrentlyCompilingBPs.Sort( HierarchyDepthSortFn );
 
+		SlowTask.EnterProgressFrame();
+
 		// STAGE IV: Set UBlueprint flags (bBeingCompiled, bIsRegeneratingOnLoad)
 		for (FCompilerData& CompilerData : CurrentlyCompilingBPs)
 		{
@@ -833,6 +851,8 @@ void FBlueprintCompilationManagerImpl::FlushCompilationQueueImpl(bool bSuppressB
 			}
 		}
 
+		SlowTask.EnterProgressFrame();
+
 		// STAGE V: Validate
 		for (FCompilerData& CompilerData : CurrentlyCompilingBPs)
 		{
@@ -844,6 +864,9 @@ void FBlueprintCompilationManagerImpl::FlushCompilationQueueImpl(bool bSuppressB
 			CompilerData.Compiler->ValidateVariableNames();
 			CompilerData.Compiler->ValidateClassPropertyDefaults();
 		}
+
+		SlowTask.EnterProgressFrame();
+
 		// STAGE V (phase 2): Give the blueprint the possibility for edits
 		for (FCompilerData& CompilerData : CurrentlyCompilingBPs)
 		{
@@ -855,6 +878,8 @@ void FBlueprintCompilationManagerImpl::FlushCompilationQueueImpl(bool bSuppressB
 			}
 		}
 
+		SlowTask.EnterProgressFrame();
+
 		// STAGE VI: Purge null graphs, misc. data fixup
 		for (FCompilerData& CompilerData : CurrentlyCompilingBPs)
 		{
@@ -873,6 +898,8 @@ void FBlueprintCompilationManagerImpl::FlushCompilationQueueImpl(bool bSuppressB
 			}
 		}
 
+		SlowTask.EnterProgressFrame();
+
 		// STAGE VII: safely throw away old skeleton CDOs:
 		{
 			TMap<UClass*, UClass*> NewSkeletonToOldSkeleton;
@@ -1108,6 +1135,7 @@ void FBlueprintCompilationManagerImpl::FlushCompilationQueueImpl(bool bSuppressB
 			}
 		}
 
+		SlowTask.EnterProgressFrame();
 
 		// STAGE IX: Reconstruct nodes and replace deprecated nodes, then broadcast 'precompile
 		for (FCompilerData& CompilerData : CurrentlyCompilingBPs)
@@ -1179,6 +1207,8 @@ void FBlueprintCompilationManagerImpl::FlushCompilationQueueImpl(bool bSuppressB
 			BP->bHasBeenRegenerated = true;
 		}
 	
+		SlowTask.EnterProgressFrame();
+
 		// STAGE X: reinstance every blueprint that is queued, note that this means classes in the hierarchy that are *not* being 
 		// compiled will be parented to REINST versions of the class, so type checks (IsA, etc) involving those types
 		// will be incoherent!
@@ -1256,6 +1286,8 @@ void FBlueprintCompilationManagerImpl::FlushCompilationQueueImpl(bool bSuppressB
 			}
 		}
 
+		SlowTask.EnterProgressFrame();
+
 		// STAGE XI: Reinstancing done, lets fix up child->parent pointers
 		for (FCompilerData& CompilerData : CurrentlyCompilingBPs)
 		{
@@ -1266,6 +1298,8 @@ void FBlueprintCompilationManagerImpl::FlushCompilationQueueImpl(bool bSuppressB
 			}
 		}
 
+		SlowTask.EnterProgressFrame();
+
 		// STAGE XII: Recompile every blueprint
 		bGeneratedClassLayoutReady = false;
 		for (FCompilerData& CompilerData : CurrentlyCompilingBPs)
@@ -1316,6 +1350,8 @@ void FBlueprintCompilationManagerImpl::FlushCompilationQueueImpl(bool bSuppressB
 		
 		ProcessExtensions(CurrentlyCompilingBPs);
 
+		SlowTask.EnterProgressFrame();
+
 		// STAGE XIII: Compile functions
 		UBlueprintEditorSettings* Settings = GetMutableDefault<UBlueprintEditorSettings>();
 		
@@ -1423,6 +1459,8 @@ void FBlueprintCompilationManagerImpl::FlushCompilationQueueImpl(bool bSuppressB
 		}
 	} // end GTimeCompiling scope
 
+	SlowTask.EnterProgressFrame();
+
 	// STAGE XIV: Now we can finish the first stage of the reinstancing operation, moving old classes to new classes:
 	{
 		{
@@ -1619,6 +1657,8 @@ void FBlueprintCompilationManagerImpl::FlushCompilationQueueImpl(bool bSuppressB
 		}
 	}
 
+	SlowTask.EnterProgressFrame();
+
 	for (FCompilerData& CompilerData : CurrentlyCompilingBPs)
 	{
 		if(CompilerData.ResultsLog)
@@ -1633,6 +1673,8 @@ void FBlueprintCompilationManagerImpl::FlushCompilationQueueImpl(bool bSuppressB
 		UEdGraphPin::Purge();
 	}
 
+	SlowTask.EnterProgressFrame();
+
 	UE_LOG(LogBlueprint, Display, TEXT("Time Compiling: %f, Time Reinstancing: %f"),  GTimeCompiling, GTimeReinstancing);
 	//GTimeCompiling = 0.0;
 	//GTimeReinstancing = 0.0;
diff --git a/Engine/Source/Editor/UnrealEd/Private/Kismet2/KismetReinstanceUtilities.cpp b/Engine/Source/Editor/UnrealEd/Private/Kismet2/KismetReinstanceUtilities.cpp
--- a/Engine/Source/Editor/UnrealEd/Private/Kismet2/KismetReinstanceUtilities.cpp
+++ b/Engine/Source/Editor/UnrealEd/Private/Kismet2/KismetReinstanceUtilities.cpp
@@ -18,6 +18,7 @@
 #include "Engine/BlueprintGeneratedClass.h"
 #include "Engine/SimpleConstructionScript.h"
 #include "FileHelpers.h"
+#include "Misc/ScopedSlowTask.h"
 #include "Kismet2/CompilerResultsLog.h"
 #include "Kismet2/KismetEditorUtilities.h"
 #include "Kismet2/BlueprintEditorUtils.h"
@@ -109,11 +110,17 @@ struct FReplaceReferenceHelper
 			Targets = FReferencerFinder::GetAllReferencers(SourceObjects, ObjectsThatShouldUseOldStuff);
 		}
 
+		if (Targets.Num())
 		{
 			BP_SCOPED_COMPILER_EVENT_STAT(EKismetReinstancerStats_ReplaceReferences);
 
+			FScopedSlowTask SlowTask(Targets.Num(), NSLOCTEXT("Kismet", "PerformingReplaceReferences", "Performing replace references..."));
+			SlowTask.MakeDialogDelayed(1.0f);
+
 			for (UObject* Obj : Targets)
 			{
+				SlowTask.EnterProgressFrame(1);
+
 				// Make sure we don't update properties in old objects, as they
 				// may take ownership of objects referenced in new objects (e.g.
 				// delete components owned by new actors)
diff --git a/Engine/Source/Runtime/Core/Private/Misc/SlowTask.cpp b/Engine/Source/Runtime/Core/Private/Misc/SlowTask.cpp
--- a/Engine/Source/Runtime/Core/Private/Misc/SlowTask.cpp
+++ b/Engine/Source/Runtime/Core/Private/Misc/SlowTask.cpp
@@ -29,6 +29,7 @@ FSlowTask::FSlowTask(float InAmountOfWork, const FText& InDefaultMessage, bool b
 	, bEnabled(bInEnabled && IsInGameThread())
 	, bCreatedDialog(false)		// only set to true if we create a dialog
 	, Context(InContext)
+	, bSkipRecursiveDialogCreation(false)
 {
 	// If we have no work to do ourselves, create an arbitrary scope so that any actions performed underneath this still contribute to this one.
 	if (TotalAmountOfWork == 0.f)
@@ -37,6 +38,44 @@ FSlowTask::FSlowTask(float InAmountOfWork, const FText& InDefaultMessage, bool b
 	}
 }
 
+void FSlowTask::MakeRecursiveDialogIfNeeded()
+{
+	if (bEnabled)
+	{
+		if (bSkipRecursiveDialogCreation)
+		{
+			MakeDialogIfNeeded();
+			return;
+		}
+
+		bSkipRecursiveDialogCreation = true;
+		for (FSlowTask* Scope : Context.ScopeStack)
+		{
+			if (Scope->MakeDialogIfNeeded())
+			{
+				// Some dialog in the hierarchy wants to be called back
+				bSkipRecursiveDialogCreation = false;
+			}
+		}
+	}
+}
+
+bool FSlowTask::MakeDialogIfNeeded()
+{
+	if (bEnabled && !bCreatedDialog && OpenDialogThreshold.IsSet())
+	{
+		if (static_cast<float>(FPlatformTime::Seconds() - StartTime) < OpenDialogThreshold.GetValue())
+		{
+			// Let our caller know that we need to be called back
+			return true;
+		}
+
+		MakeDialog(bDelayedDialogShowCancelButton, bDelayedDialogAllowInPIE);
+	}
+
+	return false;
+}
+
 void FSlowTask::Initialize()
 {
 	if (bEnabled)
@@ -100,34 +139,41 @@ void FSlowTask::MakeDialogDelayed(float Threshold, bool bShowCancelButton, bool
 
 void FSlowTask::EnterProgressFrame(float ExpectedWorkThisFrame, const FText& Text)
 {
-	check(!bEnabled || IsInGameThread());
-
-	// Should actually be FrameMessage = Text; but this code is to investigate crashes in FSlowTask::GetCurrentMessage()
-	if (!Text.IsEmpty())
-	{
-		FrameMessage = Text;
-	}
-	else
+	if (bEnabled)
 	{
-		FrameMessage = FText::GetEmpty();
-	}
-	CompletedWork += CurrentFrameScope;
+		check(IsInGameThread());
 
-	// Make sure OS events are getting through while the task is being processed
-	FPlatformMisc::PumpMessagesForSlowTask();
+		// Should actually be FrameMessage = Text; but this code is to investigate crashes in FSlowTask::GetCurrentMessage()
+		if (!Text.IsEmpty())
+		{
+			FrameMessage = Text;
+		}
+		else
+		{
+			FrameMessage = FText::GetEmpty();
+		}
+		CompletedWork += CurrentFrameScope;
 
-	const float WorkRemaining = TotalAmountOfWork - CompletedWork;
-	// Add a small threshold here because when there are a lot of tasks, numerical imprecision can add up and trigger this.
-	ensureMsgf(ExpectedWorkThisFrame <= 1.01f * TotalAmountOfWork - CompletedWork, TEXT("Work overflow in slow task. Please revise call-site to account for entire progress range."));
-	CurrentFrameScope = FMath::Min(WorkRemaining, ExpectedWorkThisFrame);
+		const float WorkRemaining = TotalAmountOfWork - CompletedWork;
+		// Add a small threshold here because when there are a lot of tasks, numerical imprecision can add up and trigger this.
+		ensureMsgf(ExpectedWorkThisFrame <= 1.01f * TotalAmountOfWork - CompletedWork, TEXT("Work overflow in slow task. Please revise call-site to account for entire progress range."));
+		CurrentFrameScope = FMath::Min(WorkRemaining, ExpectedWorkThisFrame);
 
-	if (!bCreatedDialog && OpenDialogThreshold.IsSet() && static_cast<float>(FPlatformTime::Seconds() - StartTime) > OpenDialogThreshold.GetValue())
-	{
-		MakeDialog(bDelayedDialogShowCancelButton, bDelayedDialogAllowInPIE);
+		TickProgress();
 	}
+}
 
+void FSlowTask::TickProgress()
+{
 	if (bEnabled)
 	{
+		check(IsInGameThread());
+
+		// Make sure OS events are getting through while the task is being processed
+		FPlatformMisc::PumpMessagesForSlowTask();
+
+		MakeRecursiveDialogIfNeeded();
+	
 		Context.RequestUpdateUI();
 	}
 }
diff --git a/Engine/Source/Runtime/Core/Public/Misc/SlowTask.h b/Engine/Source/Runtime/Core/Public/Misc/SlowTask.h
--- a/Engine/Source/Runtime/Core/Public/Misc/SlowTask.h
+++ b/Engine/Source/Runtime/Core/Public/Misc/SlowTask.h
@@ -63,12 +63,21 @@ private:
 	/** Specify whether the delayed dialog is allowed in PIE */
 	bool bDelayedDialogAllowInPIE : 1;
 
+	/** Avoid calling hierarchical dialog creation if not required */
+	bool bSkipRecursiveDialogCreation : 1;
+
 	/** Prevent copying */
 	FSlowTask(const FSlowTask&);
 
 	/** Whenever we encounter edge condition requiring to force a UI refresh */
 	static void ForceRefresh(FFeedbackContext& Context);
 
+	/** Handle dialog and delayed dialog creation logic */
+	bool MakeDialogIfNeeded();
+
+	/** Call MakeDialogIfNeeded on all scopes recursively to give them a chance to create a delayed dialog in the hierarchy */
+	void MakeRecursiveDialogIfNeeded();
+
 public:
 
 	/**
@@ -113,6 +122,11 @@ public:
 	 */
 	void EnterProgressFrame(float ExpectedWorkThisFrame = 1.f, const FText& Text = FText());
 
+	/**
+	 * Let the UI refresh but doesn't advance progress. This should be called at regular intervals even when no progress is being made to keep the UI responsive.
+	 */
+	void TickProgress();
+
 	/**
 	 * Get the frame message or default message if empty
 	 */
diff --git a/Engine/Source/Runtime/CoreUObject/Private/Serialization/AsyncLoading.cpp b/Engine/Source/Runtime/CoreUObject/Private/Serialization/AsyncLoading.cpp
--- a/Engine/Source/Runtime/CoreUObject/Private/Serialization/AsyncLoading.cpp
+++ b/Engine/Source/Runtime/CoreUObject/Private/Serialization/AsyncLoading.cpp
@@ -19,6 +19,7 @@
 #include "Misc/App.h"
 #include "Misc/MessageDialog.h"
 #include "Misc/PackageName.h"
+#include "Misc/ScopedSlowTask.h"
 #include "UObject/PackageFileSummary.h"
 #include "UObject/Linker.h"
 #include "UObject/CoreRedirects.h"
@@ -6283,6 +6284,12 @@ EAsyncPackageState::Type FAsyncPackage::PostLoadObjects()
 	const bool bAsyncPostLoadEnabled = FAsyncLoadingThreadSettings::Get().bAsyncPostLoadEnabled;
 	const bool bIsMultithreaded = AsyncLoadingThread.IsMultithreaded();
 
+	TOptional<FScopedSlowTask> SlowTask;
+	if (IsInGameThread())
+	{
+		SlowTask.Emplace(1, NSLOCTEXT("AsyncLoading", "PostloadingObjects", "Postloading objects..."));
+	}
+
 	// PostLoad objects.
 	while (PostLoadIndex < PackageObjLoaded.Num() && PostLoadIndex < PreLoadIndex && !IsTimeLimitExceeded())
 	{
@@ -6330,6 +6337,13 @@ EAsyncPackageState::Type FAsyncPackage::PostLoadObjects()
 			check(Object->IsValidLowLevelFast());
 			// Make sure all objects in DeferredFinalizeObjects are referenced too
 			AddObjectReference(Object);
+
+			if (SlowTask)
+			{
+				// Don't report progress but gives a chance to tick slate to improve the responsiveness of the 
+				// progress bar being shown. We expect slate to be ticked at regular intervals throughout the loading.
+				SlowTask->TickProgress();
+			}
 		}
 	}
 
diff --git a/Engine/Source/Runtime/CoreUObject/Private/UObject/LinkerLoad.cpp b/Engine/Source/Runtime/CoreUObject/Private/UObject/LinkerLoad.cpp
--- a/Engine/Source/Runtime/CoreUObject/Private/UObject/LinkerLoad.cpp
+++ b/Engine/Source/Runtime/CoreUObject/Private/UObject/LinkerLoad.cpp
@@ -77,7 +77,7 @@ Helpers
 #if WITH_EDITOR
 bool FLinkerLoad::ShouldCreateThrottledSlowTask() const
 {
-	return ShouldReportProgress() && FSlowTask::ShouldCreateThrottledSlowTask();
+	return ShouldReportProgress();
 }
 
 int32 GTreatVerifyImportErrorsAsWarnings = 0;
diff --git a/Engine/Source/Runtime/CoreUObject/Private/UObject/UObjectGlobals.cpp b/Engine/Source/Runtime/CoreUObject/Private/UObject/UObjectGlobals.cpp
--- a/Engine/Source/Runtime/CoreUObject/Private/UObject/UObjectGlobals.cpp
+++ b/Engine/Source/Runtime/CoreUObject/Private/UObject/UObjectGlobals.cpp
@@ -194,7 +194,7 @@ bool ShouldReportProgress()
 
 bool ShouldCreateThrottledSlowTask()
 {
-	return ShouldReportProgress() && FSlowTask::ShouldCreateThrottledSlowTask();
+	return ShouldReportProgress();
 }
 
 /**
@@ -1662,6 +1662,14 @@ void EndLoad(FUObjectSerializeContext* LoadContext, TArray<UPackage*>* OutLoaded
 					{
 						check(Obj->GetLinker());
 						UE_TRACK_REFERENCING_PACKAGE_DELAYED(AccessRefScope, Obj->GetOutermost());
+#if WITH_EDITOR
+						if (SlowTask)
+						{
+							// Don't report progress but gives a chance to tick slate to improve the responsiveness of the 
+							// progress bar being shown. We expect slate to be ticked at regular intervals throughout the loading.
+							SlowTask->TickProgress();
+						}
+#endif
 						Obj->GetLinker()->Preload(Obj);
 					}
 				}
@@ -1840,6 +1848,13 @@ void EndLoad(FUObjectSerializeContext* LoadContext, TArray<UPackage*>* OutLoaded
 	{
 		check(LoadedAsset);
 		FCoreUObjectDelegates::OnAssetLoaded.Broadcast(LoadedAsset);
+		
+		if (SlowTask)
+		{
+			// Don't report progress but gives a chance to tick slate to improve the responsiveness of the 
+			// progress bar being shown. We expect slate to be ticked at regular intervals throughout the loading.
+			SlowTask->TickProgress();
+		}
 	}
 #endif	// WITH_EDITOR
 
diff --git a/Engine/Source/Runtime/Engine/Private/GameInstance.cpp b/Engine/Source/Runtime/Engine/Private/GameInstance.cpp
--- a/Engine/Source/Runtime/Engine/Private/GameInstance.cpp
+++ b/Engine/Source/Runtime/Engine/Private/GameInstance.cpp
@@ -15,6 +15,7 @@
 #include "GameFramework/PlayerController.h"
 #include "Engine/Engine.h"
 #include "Engine/Console.h"
+#include "Misc/ScopedSlowTask.h"
 #include "Engine/GameEngine.h"
 #include "GameFramework/GameModeBase.h"
 #include "Engine/DemoNetDriver.h"
@@ -396,6 +397,9 @@ FGameInstancePIEResult UGameInstance::StartPlayInEditorGameInstance(ULocalPlayer
 	}
 	else
 	{
+		FScopedSlowTask SlowTask(100, NSLOCTEXT("UnrealEd", "StartPlayInEditor", "Starting PIE..."));
+		SlowTask.MakeDialogDelayed(1.0f, false, true);
+
 		// we're going to be playing in the current world, get it ready for play
 		UWorld* const PlayWorld = GetWorld();
 
@@ -454,11 +458,14 @@ FGameInstancePIEResult UGameInstance::StartPlayInEditorGameInstance(ULocalPlayer
 			return PostCreateGameModeResult;
 		}
 
+		SlowTask.EnterProgressFrame(10, NSLOCTEXT("UnrealEd", "PIEFlushingLevelStreaming", "Starting PIE (Loading always loaded objects)..."));
 		// Make sure "always loaded" sub-levels are fully loaded
 		PlayWorld->FlushLevelStreaming(EFlushLevelStreamingType::Visibility);
 
+		SlowTask.EnterProgressFrame(10, NSLOCTEXT("UnrealEd", "PIECreatingAISystem", "Starting PIE (Creating AI System)..."));
 		PlayWorld->CreateAISystem();
 
+		SlowTask.EnterProgressFrame(10, NSLOCTEXT("UnrealEd", "PIEInitializingActors", "Starting PIE (Initializing Actors)..."));
 		PlayWorld->InitializeActorsForPlay(URL);
 		// calling it after InitializeActorsForPlay has been called to have all potential bounding boxed initialized
 		FNavigationSystem::AddNavigationSystemToWorld(*PlayWorld, LocalPlayers.Num() > 0 ? FNavigationSystemRunMode::PIEMode : FNavigationSystemRunMode::SimulationMode);
@@ -476,6 +483,7 @@ FGameInstancePIEResult UGameInstance::StartPlayInEditorGameInstance(ULocalPlayer
 		UGameViewportClient* const GameViewport = GetGameViewportClient();
 		if (GameViewport != NULL && GameViewport->Viewport != NULL)
 		{
+			SlowTask.EnterProgressFrame(50, NSLOCTEXT("UnrealEd", "PIEWaitingForLevelStreaming", "Starting PIE (Waiting for level streaming)..."));
 			// Stream any levels now that need to be loaded before the game starts
 			GEngine->BlockTillLevelStreamingCompleted(PlayWorld);
 		}
@@ -494,8 +502,10 @@ FGameInstancePIEResult UGameInstance::StartPlayInEditorGameInstance(ULocalPlayer
 			ensureMsgf(EnableListenServer(true, ListenPort), TEXT("Starting Listen Server for Play in Editor failed!"));
 		}
 
+		SlowTask.EnterProgressFrame(10, NSLOCTEXT("UnrealEd", "PIEBeginPlay", "Starting PIE (Begin play)..."));
 		PlayWorld->BeginPlay();
 
+		SlowTask.EnterProgressFrame(10);
 #if WITH_EDITOR
 		if (PlayWorld->WorldType == EWorldType::PIE)
 		{
