Some minor refactoring to assist with handling removal/reassignment of Blueprint namespaces.

Change summary:
- Added a code path for removing a namespace from the BP editor context.
- Consolidated namespace helper instances into a single BP editor context.
- Extended the FBlueprintEditor::ImportNamespace() API to accept additional options and a post-import callback.
- Refactored the auto-import code path to fix issues with reassignment to/from an imported namespace.
- Decoupled SSubobjectBlueprintEditor from a namespace-aware callback in favor of a generalized "add" callback.
- Linked the CVar for inheriting parent BP imports to a transient BP setting (experimental/not user-facing).
- Consolidated code to gather global/default imports into a utility API so that namespace sets are consistent.
- Modified "Imports" view to include a "default imports" section (similar to inherited interfaces) - these are not removable.
- Moved "Inherited Interfaces" section to be sorted above the "Implemented Interfaces" section for consistency.
- Fixed SBlueprintNamespaceEntry to query for the current text rather than only setting it at construction time.

diff --git a/Engine/Source/Editor/BlueprintGraph/Private/BlueprintEditorSettings.cpp b/Engine/Source/Editor/BlueprintGraph/Private/BlueprintEditorSettings.cpp
--- a/Engine/Source/Editor/BlueprintGraph/Private/BlueprintEditorSettings.cpp
+++ b/Engine/Source/Editor/BlueprintGraph/Private/BlueprintEditorSettings.cpp
@@ -42,6 +42,7 @@ UBlueprintEditorSettings::UBlueprintEditorSettings(const FObjectInitializer& Obj
 	, bEnableNamespaceEditorFeatures(false)
 	, bEnableNamespaceFilteringFeatures(false)
 	, bEnableNamespaceImportingFeatures(false)
+	, bInheritImportedNamespacesFromParentBP(false)
 	, bFavorPureCastNodes(false)
 	// Compiler Settings
 	, SaveOnCompile(SoC_Never)
diff --git a/Engine/Source/Editor/BlueprintGraph/Public/BlueprintEditorSettings.h b/Engine/Source/Editor/BlueprintGraph/Public/BlueprintEditorSettings.h
--- a/Engine/Source/Editor/BlueprintGraph/Public/BlueprintEditorSettings.h
+++ b/Engine/Source/Editor/BlueprintGraph/Public/BlueprintEditorSettings.h
@@ -175,6 +175,11 @@ public:
 	UPROPERTY(Transient)
 	bool bEnableNamespaceImportingFeatures;
 
+	/** Whether to inherit the set of imported namespaces from the parent class hierarchy */
+	// @todo_namespaces - Remove this if/when this becomes a permanent setting. For now this is experimental.
+	UPROPERTY(Transient)
+	bool bInheritImportedNamespacesFromParentBP;
+
 	// The list of namespaces to always expose in any Blueprint (local per-user)
 	UPROPERTY(EditAnywhere, config, Category = Experimental, meta = (EditCondition = "bEnableNamespaceEditorFeatures"))
 	TArray<FString> NamespacesToAlwaysInclude;
diff --git a/Engine/Source/Editor/Kismet/Private/BlueprintDetailsCustomization.cpp b/Engine/Source/Editor/Kismet/Private/BlueprintDetailsCustomization.cpp
--- a/Engine/Source/Editor/Kismet/Private/BlueprintDetailsCustomization.cpp
+++ b/Engine/Source/Editor/Kismet/Private/BlueprintDetailsCustomization.cpp
@@ -2782,7 +2782,7 @@ void FBlueprintVarActionDetails::ImportNamespacesForPropertyValue(const UStruct*
 	}
 
 	// Auto-import any namespace(s) associated with the property's value into the current editor context.
-	TSet<FString> AssociatedNamespaces;
+	TArray<FString> AssociatedNamespaces;
 	FBlueprintNamespaceUtilities::GetPropertyValueNamespaces(InStruct, InProperty, InContainer, AssociatedNamespaces);
 	if (AssociatedNamespaces.Num() > 0)
 	{
@@ -2792,10 +2792,9 @@ void FBlueprintVarActionDetails::ImportNamespacesForPropertyValue(const UStruct*
 			TSharedPtr<FBlueprintEditor> BlueprintEditor = MyBlueprintPtr->GetBlueprintEditor().Pin();
 			if (BlueprintEditor.IsValid())
 			{
-				for (const FString& AssociatedNamespace : AssociatedNamespaces)
-				{
-					BlueprintEditor->ImportNamespace(AssociatedNamespace);
-				}
+				FBlueprintEditor::FImportNamespaceParameters Params;
+				Params.AdditionalNamespaces = MakeArrayView(AssociatedNamespaces).RightChop(1);
+				BlueprintEditor->ImportNamespace(AssociatedNamespaces[0], Params);
 			}
 		}
 	}
@@ -5688,10 +5687,13 @@ void FBlueprintManagedListDetails::OnRefreshInDetailsView()
 }
 
 
-FBlueprintImportsLayout::FBlueprintImportsLayout(TWeakPtr<class FBlueprintGlobalOptionsDetails> InGlobalOptionsDetails)
+FBlueprintImportsLayout::FBlueprintImportsLayout(TWeakPtr<class FBlueprintGlobalOptionsDetails> InGlobalOptionsDetails, bool bInShowDefaultImports)
 	:FBlueprintManagedListDetails(InGlobalOptionsDetails)
+	,bShouldShowDefaultImports(bInShowDefaultImports)
 {
-	DisplayOptions.TitleText = LOCTEXT("BlueprintImportedNamespaceTitle", "Imported Namespaces");
+	DisplayOptions.TitleText = bShouldShowDefaultImports ?
+		LOCTEXT("BlueprintDefaultNamespaceTitle", "Default Namespaces") :
+		LOCTEXT("BlueprintImportedNamespaceTitle", "Imported Namespaces");
 	DisplayOptions.NoItemsLabelText = LOCTEXT("NoBlueprintImports", "No Imports");
 	DisplayOptions.ItemRowFilterText = LOCTEXT("BlueprintImportsValue", "Imports Value");
 	DisplayOptions.AddItemRowFilterText = LOCTEXT("BlueprintAddImport", "Add Import");
@@ -5699,6 +5701,11 @@ FBlueprintImportsLayout::FBlueprintImportsLayout(TWeakPtr<class FBlueprintGlobal
 
 TSharedPtr<SWidget> FBlueprintImportsLayout::MakeAddItemRowWidget()
 {
+	if (bShouldShowDefaultImports)
+	{
+		return nullptr;
+	}
+
 	return SNew(SBlueprintNamespaceEntry)
 		.AllowTextEntry(false)
 		.OnNamespaceSelected(this, &FBlueprintImportsLayout::OnNamespaceSelected)
@@ -5712,37 +5719,47 @@ TSharedPtr<SWidget> FBlueprintImportsLayout::MakeAddItemRowWidget()
 
 void FBlueprintImportsLayout::GetManagedListItems(TArray<FManagedListItem>& OutListItems) const
 {
-	auto AddNamespaceItemsToOutputList = [&OutListItems](const TArray<FString>& NamespaceArray, bool bIsRemovable)
+	auto AddNamespaceItemsToOutputList = [&OutListItems](const TSet<FString>& NamespaceItems, bool bIsRemovable)
 	{
-		for (const FString& Namespace : NamespaceArray)
+		for (const FString& NamespaceItem : NamespaceItems)
 		{
 			FManagedListItem ItemDesc;
-			ItemDesc.ItemName = Namespace;
-			ItemDesc.DisplayName = FText::FromString(Namespace);
+			ItemDesc.ItemName = NamespaceItem;
+			ItemDesc.DisplayName = FText::FromString(NamespaceItem);
 			ItemDesc.bIsRemovable = bIsRemovable;
 
 			OutListItems.Add(MoveTemp(ItemDesc));
 		}
 	};
 
-	// Project imports (static).
-	bool bIsRemovable = false;
-	AddNamespaceItemsToOutputList(GetDefault<UBlueprintEditorProjectSettings>()->NamespacesToAlwaysInclude, bIsRemovable);
-
-	// Blueprint imports (removable).
-	bIsRemovable = true;
+	TSet<FString> NamespaceItems;
 	const UBlueprint* Blueprint = GetBlueprintObjectChecked();
-	AddNamespaceItemsToOutputList(Blueprint->ImportedNamespaces.Array(), bIsRemovable);
-}
 
-void FBlueprintImportsLayout::OnRemoveItem(const FManagedListItem& Item)
-{
-	UBlueprint* Blueprint = GetBlueprintObjectChecked();
+	if (bShouldShowDefaultImports)
+	{
+		FBlueprintNamespaceUtilities::GetSharedGlobalImports(NamespaceItems);
+		FBlueprintNamespaceUtilities::GetDefaultImportsForBlueprint(Blueprint, NamespaceItems);
+	}
+	else
+	{
+		// Blueprint imports (removable).
+		NamespaceItems.Append(Blueprint->ImportedNamespaces);
 
-	Blueprint->ImportedNamespaces.Remove(Item.ItemName);
+		// A Blueprint can explicitly import and also be assigned to the same namespace.
+		NamespaceItems.Remove(Blueprint->BlueprintNamespace);
+	}
 
-	FBlueprintEditorUtils::MarkBlueprintAsModified(Blueprint);
+	AddNamespaceItemsToOutputList(NamespaceItems, !bShouldShowDefaultImports);
+}
 
+void FBlueprintImportsLayout::OnRemoveItem(const FManagedListItem& Item)
+{
+	TSharedPtr<FBlueprintEditor> BlueprintEditorPtr = GetPinnedBlueprintEditorPtr();
+	if (BlueprintEditorPtr.IsValid())
+	{
+		BlueprintEditorPtr->RemoveNamespace(Item.ItemName);
+	}
+	
 	RegenerateChildContent();
 
 	OnRefreshInDetailsView();
@@ -5750,32 +5767,21 @@ void FBlueprintImportsLayout::OnRemoveItem(const FManagedListItem& Item)
 
 void FBlueprintImportsLayout::OnNamespaceSelected(const FString& InNamespace)
 {
-	bool bWasAdded = false;
-
-	// Add to the blueprint's list of imports.
-	if (!InNamespace.IsEmpty())
+	TSharedPtr<FBlueprintEditor> BlueprintEditorPtr = GetPinnedBlueprintEditorPtr();
+	if (BlueprintEditorPtr.IsValid())
 	{
-		UBlueprint* Blueprint = GetBlueprintObjectChecked();
-
-		if (FBlueprintEditorUtils::AddNamespaceToImportList(Blueprint, InNamespace))
+		FBlueprintEditor::FImportNamespaceParameters Params;
+		Params.bIsAutoImport = false;
+		Params.OnImportCallback = FSimpleDelegate::CreateLambda([this]()
 		{
-			bWasAdded = true;
-		}
+			RegenerateChildContent();
+		});
 
-		RegenerateChildContent();
+		// Add to edited Blueprint(s) and import into the current editor context.
+		BlueprintEditorPtr->ImportNamespace(InNamespace, Params);
 	}
 
 	OnRefreshInDetailsView();
-
-	// If added above, import the namespace into the current editor context.
-	if (bWasAdded)
-	{
-		TSharedPtr<FBlueprintEditor> BlueprintEditorPtr = GetPinnedBlueprintEditorPtr();
-		if (BlueprintEditorPtr.IsValid())
-		{
-			BlueprintEditorPtr->ImportNamespace(InNamespace);
-		}
-	}
 }
 
 void FBlueprintImportsLayout::OnFilterNamespaceList(TArray<FString>& InOutNamespaceList)
@@ -6046,7 +6052,8 @@ void FBlueprintGlobalOptionsDetails::OnNamespaceValueCommitted(const FString& In
 	if (UBlueprint* Blueprint = GetBlueprintObj())
 	{
 		// Skip if the value has not been changed.
-		if (Blueprint->BlueprintNamespace == InNamespace)
+		const FString OldNamespace = Blueprint->BlueprintNamespace;
+		if (OldNamespace == InNamespace)
 		{
 			return;
 		}
@@ -6056,15 +6063,21 @@ void FBlueprintGlobalOptionsDetails::OnNamespaceValueCommitted(const FString& In
 		NamespacePropertyHandle->SetValue(InNamespace);
 
 		// Ensure that the new path is added into the namespace registry.
-		// @todo_namespaces - Find and remove old path if no longer in use.
 		FBlueprintNamespaceRegistry::Get().RegisterNamespace(InNamespace);
 
-		if (GetDefault<UBlueprintEditorSettings>()->bEnableNamespaceImportingFeatures)
+		// Auto-import the namespace into the current editor context.
+		TSharedPtr<FBlueprintEditor> BlueprintEditor = GetBlueprintEditorPtr().Pin();
+		if (BlueprintEditor.IsValid())
 		{
-			// Auto-import the namespace into the current editor context.
-			TSharedPtr<FBlueprintEditor> BlueprintEditor = GetBlueprintEditorPtr().Pin();
-			if (BlueprintEditor.IsValid())
+			// We need to refresh the details view if we flipped from an imported namespace to an assigned
+			// namespace (or vice-versa). In that case, it will switch to/from a default import in the view.
+			if (Blueprint->ImportedNamespaces.Contains(OldNamespace) || Blueprint->ImportedNamespaces.Contains(InNamespace))
+			{
+				BlueprintEditor->RefreshInspector();
+			}
+			else if (GetDefault<UBlueprintEditorSettings>()->bEnableNamespaceImportingFeatures)
 			{
+				// This will also refresh the details view.
 				BlueprintEditor->ImportNamespace(InNamespace);
 			}
 		}
@@ -6079,8 +6092,13 @@ bool FBlueprintGlobalOptionsDetails::ShouldShowNamespaceResetToDefault() const
 
 void FBlueprintGlobalOptionsDetails::OnNamespaceResetToDefaultValue()
 {
-	// Standard reset-to-default path.
 	check(NamespacePropertyHandle.IsValid());
+
+	// Get the current value.
+	FString OriginalValue;
+	NamespacePropertyHandle->GetValue(OriginalValue);
+
+	// Standard reset-to-default path.
 	NamespacePropertyHandle->ResetToDefault();
 
 	// Get the value after having been reset.
@@ -6092,6 +6110,18 @@ void FBlueprintGlobalOptionsDetails::OnNamespaceResetToDefaultValue()
 	{
 		NamespaceValueWidget->SetCurrentNamespace(DefaultNamespaceValue);
 	}
+
+	TSharedPtr<FBlueprintEditor> BlueprintEditor = GetBlueprintEditorPtr().Pin();
+	if (BlueprintEditor.IsValid())
+	{
+		// We need to refresh the details view if we reassigned away from an already-imported
+		// namespace. In that case, it needs to switch back to a user-added import in the view.
+		const UBlueprint* Blueprint = GetBlueprintObj();
+		if (Blueprint && Blueprint->ImportedNamespaces.Contains(OriginalValue))
+		{
+			BlueprintEditor->RefreshInspector();
+		}
+	}
 }
 
 void FBlueprintGlobalOptionsDetails::CustomizeDetails(IDetailLayoutBuilder& DetailLayout)
@@ -6150,8 +6180,11 @@ void FBlueprintGlobalOptionsDetails::CustomizeDetails(IDetailLayoutBuilder& Deta
 			// Imported namespace details
 			IDetailCategoryBuilder& ImportsCategory = DetailLayout.EditCategory("Imports", LOCTEXT("BlueprintImportDetailsCategory", "Imports"));
 
-			TSharedRef<FBlueprintImportsLayout> ImportsLayout = MakeShareable(new FBlueprintImportsLayout(SharedThis(this)));
-			ImportsCategory.AddCustomBuilder(ImportsLayout);
+			TSharedRef<FBlueprintImportsLayout> DefaultImportsLayout = MakeShareable(new FBlueprintImportsLayout(SharedThis(this), /*bShowDefaultImports = */true));
+			ImportsCategory.AddCustomBuilder(DefaultImportsLayout);
+
+			TSharedRef<FBlueprintImportsLayout> LocalImportsLayout = MakeShareable(new FBlueprintImportsLayout(SharedThis(this), /*bShowDefaultImports = */false));
+			ImportsCategory.AddCustomBuilder(LocalImportsLayout);
 		}
 
 		if (bSupportsInterfaces)
@@ -6159,11 +6192,11 @@ void FBlueprintGlobalOptionsDetails::CustomizeDetails(IDetailLayoutBuilder& Deta
 			// Interface details customization
 			IDetailCategoryBuilder& InterfacesCategory = DetailLayout.EditCategory("Interfaces", LOCTEXT("BlueprintInterfacesDetailsCategory", "Interfaces"));
 		
-			TSharedRef<FBlueprintInterfaceLayout> InterfaceLayout = MakeShareable(new FBlueprintInterfaceLayout(SharedThis(this), false));
-			InterfacesCategory.AddCustomBuilder(InterfaceLayout);
-		
-			TSharedRef<FBlueprintInterfaceLayout> InheritedInterfaceLayout = MakeShareable(new FBlueprintInterfaceLayout(SharedThis(this), true));
-			InterfacesCategory.AddCustomBuilder(InheritedInterfaceLayout);
+			TSharedRef<FBlueprintInterfaceLayout> InheritedInterfacesLayout = MakeShareable(new FBlueprintInterfaceLayout(SharedThis(this), /*bShowInheritedInterfaces = */true));
+			InterfacesCategory.AddCustomBuilder(InheritedInterfacesLayout);
+
+			TSharedRef<FBlueprintInterfaceLayout> LocalInterfacesLayout = MakeShareable(new FBlueprintInterfaceLayout(SharedThis(this), /*bShowInheritedInterfaces = */false));
+			InterfacesCategory.AddCustomBuilder(LocalInterfacesLayout);
 		}
 
 		// Hide the bDeprecate, we override the functionality.
diff --git a/Engine/Source/Editor/Kismet/Private/BlueprintDetailsCustomization.h b/Engine/Source/Editor/Kismet/Private/BlueprintDetailsCustomization.h
--- a/Engine/Source/Editor/Kismet/Private/BlueprintDetailsCustomization.h
+++ b/Engine/Source/Editor/Kismet/Private/BlueprintDetailsCustomization.h
@@ -755,7 +755,7 @@ private:
 class FBlueprintImportsLayout : public FBlueprintManagedListDetails, public TSharedFromThis<FBlueprintImportsLayout>
 {
 public:
-	FBlueprintImportsLayout(TWeakPtr<class FBlueprintGlobalOptionsDetails> InGlobalOptionsDetails);
+	FBlueprintImportsLayout(TWeakPtr<class FBlueprintGlobalOptionsDetails> InGlobalOptionsDetails, bool bInShowDefaultImports);
 
 protected:
 	/** FBlueprintManagedListDetails interface*/
@@ -767,6 +767,10 @@ protected:
 	void OnNamespaceSelected(const FString& InNamespace);
 	void OnFilterNamespaceList(TArray<FString>& InOutNamespaceList);
 	void HandleImportEntryTextCommitted(const FText& NewLabel, ETextCommit::Type CommitType);
+
+private:
+	/** Whether we should show default (i.e. global) imports versus custom (i.e. local) imports */
+	bool bShouldShowDefaultImports;
 };
 
 /** Blueprint Interface List Details */
diff --git a/Engine/Source/Editor/Kismet/Private/BlueprintEditor.cpp b/Engine/Source/Editor/Kismet/Private/BlueprintEditor.cpp
--- a/Engine/Source/Editor/Kismet/Private/BlueprintEditor.cpp
+++ b/Engine/Source/Editor/Kismet/Private/BlueprintEditor.cpp
@@ -103,6 +103,7 @@
 #include "Stats/StatsHierarchical.h"
 #include "BlueprintEditorLibrary.h"
 #include "BlueprintNamespaceHelper.h"
+#include "BlueprintNamespaceUtilities.h"
 
 #include "BlueprintEditorTabs.h"
 
@@ -230,23 +231,14 @@ namespace BlueprintEditorImpl
 {
 	static const float InstructionFadeDuration = 0.5f;
 
-	/** Consolidates the class viewer filters for each edited BP asset into a single filter option */
-	class FImportedClassViewerFilter : public IClassViewerFilter, public TSharedFromThis<FImportedClassViewerFilter>
+	/** Class viewer filter proxy for imported namespace type selectors, controlled by a custom filter option */
+	class FImportedClassViewerFilterProxy : public IClassViewerFilter, public TSharedFromThis<FImportedClassViewerFilterProxy>
 	{
 	public:
-		FImportedClassViewerFilter(const TArray<TSharedRef<FBlueprintNamespaceHelper>>& InNamespaceHelpers)
+		FImportedClassViewerFilterProxy(TSharedPtr<IClassViewerFilter> InClassViewerFilter)
+			: ClassViewerFilter(InClassViewerFilter)
+			, bIsFilterEnabled(true)
 		{
-			bIsFilterEnabled = true;
-
-			ClassViewerFilters.Reserve(InNamespaceHelpers.Num());
-			for (const TSharedRef<FBlueprintNamespaceHelper>& NamespaceHelper : InNamespaceHelpers)
-			{
-				TSharedPtr<IClassViewerFilter> Filter = NamespaceHelper->GetClassViewerFilter();
-				if (Filter.IsValid())
-				{
-					ClassViewerFilters.Add(Filter.ToSharedRef());
-				}
-			}
 		}
 
 		// IClassViewerFilter interface
@@ -257,15 +249,9 @@ namespace BlueprintEditorImpl
 				return true;
 			}
 
-			if(bIsFilterEnabled)
+			if (bIsFilterEnabled && ClassViewerFilter.IsValid())
 			{
-				for (const TSharedRef<IClassViewerFilter>& Filter : ClassViewerFilters)
-				{
-					if (!Filter->IsClassAllowed(InInitOptions, InClass, InFilterFuncs))
-					{
-						return false;
-					}
-				}
+				return ClassViewerFilter->IsClassAllowed(InInitOptions, InClass, InFilterFuncs);
 			}
 
 			return true;
@@ -278,15 +264,9 @@ namespace BlueprintEditorImpl
 				return true;
 			}
 
-			if (bIsFilterEnabled)
+			if (bIsFilterEnabled && ClassViewerFilter.IsValid())
 			{
-				for (const TSharedRef<IClassViewerFilter>& Filter : ClassViewerFilters)
-				{
-					if (!Filter->IsUnloadedClassAllowed(InInitOptions, InBlueprint, InFilterFuncs))
-					{
-						return false;
-					}
-				}
+				return ClassViewerFilter->IsUnloadedClassAllowed(InInitOptions, InBlueprint, InFilterFuncs);
 			}
 
 			return true;
@@ -305,7 +285,7 @@ namespace BlueprintEditorImpl
 				ToggleFilterOption->bEnabled = bIsFilterEnabled;
 				ToggleFilterOption->LabelText = LOCTEXT("ClassViewerNamespaceFilterMenuOptionLabel", "Show Only Imported Types");
 				ToggleFilterOption->ToolTipText = LOCTEXT("ClassViewerNamespaceFilterMenuOptionToolTip", "Don't include non-imported class types.");
-				ToggleFilterOption->OnOptionChanged = FOnClassViewerFilterOptionChanged::CreateSP(this, &FImportedClassViewerFilter::OnFilterOptionChanged);
+				ToggleFilterOption->OnOptionChanged = FOnClassViewerFilterOptionChanged::CreateSP(this, &FImportedClassViewerFilterProxy::OnFilterOptionChanged);
 			}
 
 			OutFilterOptions.Add(ToggleFilterOption.ToSharedRef());
@@ -318,8 +298,8 @@ namespace BlueprintEditorImpl
 		}
 
 	private:
-		/** Class viewer filters (one per edited BP asset). */
-		TArray<TSharedRef<IClassViewerFilter>> ClassViewerFilters;
+		/** Imported namespace class viewer filter. */
+		TSharedPtr<IClassViewerFilter> ClassViewerFilter;
 
 		/** Filter option for the class viewer settings menu. */
 		TSharedPtr<FClassViewerFilterOption> ToggleFilterOption;
@@ -328,25 +308,16 @@ namespace BlueprintEditorImpl
 		bool bIsFilterEnabled;
 	};
 
-	/** Consolidates the pin type selector filters for each edited BP asset into a singular construct */
-	class FImportedPinTypeSelectorFilter : public IPinTypeSelectorFilter, public TSharedFromThis<FImportedPinTypeSelectorFilter>
+	/** Pin type filter proxy for imported namespace type selectors, controlled by a custom filter option */
+	class FImportedPinTypeSelectorFilterProxy : public IPinTypeSelectorFilter, public TSharedFromThis<FImportedPinTypeSelectorFilterProxy>
 	{
 		DECLARE_MULTICAST_DELEGATE(FOnFilterChanged);
 	
 	public:
-		FImportedPinTypeSelectorFilter(const TArray<TSharedRef<FBlueprintNamespaceHelper>>& InNamespaceHelpers)
+		FImportedPinTypeSelectorFilterProxy(TSharedPtr<IPinTypeSelectorFilter> InPinTypeSelectorFilter)
+			: PinTypeSelectorFilter(InPinTypeSelectorFilter)
+			, bIsFilterEnabled(true)
 		{
-			bIsFilterEnabled = true;
-
-			PinTypeSelectorFilters.Reserve(InNamespaceHelpers.Num());
-			for (const TSharedRef<FBlueprintNamespaceHelper>& NamespaceHelper : InNamespaceHelpers)
-			{
-				TSharedPtr<IPinTypeSelectorFilter> Filter = NamespaceHelper->GetPinTypeSelectorFilter();
-				if (Filter.IsValid())
-				{
-					PinTypeSelectorFilters.Add(Filter.ToSharedRef());
-				}
-			}
 		}
 
 		// IPinTypeSelectorFilter interface
@@ -357,15 +328,9 @@ namespace BlueprintEditorImpl
 				return true;
 			}
 
-			if (bIsFilterEnabled)
+			if (bIsFilterEnabled && PinTypeSelectorFilter.IsValid())
 			{
-				for (const TSharedRef<IPinTypeSelectorFilter>& Filter : PinTypeSelectorFilters)
-				{
-					if (!Filter->ShouldShowPinTypeTreeItem(InItem))
-					{
-						return false;
-					}
-				}
+				return PinTypeSelectorFilter->ShouldShowPinTypeTreeItem(InItem);
 			}
 
 			return true;
@@ -391,8 +356,8 @@ namespace BlueprintEditorImpl
 			if (!FilterOptionsWidget.IsValid())
 			{
 				SAssignNew(FilterOptionsWidget, SCheckBox)
-					.IsChecked(this, &FImportedPinTypeSelectorFilter::IsFilterToggleChecked)
-					.OnCheckStateChanged(this, &FImportedPinTypeSelectorFilter::OnToggleFilter)
+					.IsChecked(this, &FImportedPinTypeSelectorFilterProxy::IsFilterToggleChecked)
+					.OnCheckStateChanged(this, &FImportedPinTypeSelectorFilterProxy::OnToggleFilter)
 					[
 						SNew(STextBlock)
 						.Text(LOCTEXT("PinTypeNamespaceFilterToggleOptionLabel", "Hide Non-Imported Types"))
@@ -417,8 +382,8 @@ namespace BlueprintEditorImpl
 		}
 
 	private:
-		/** Pin type filters (one per edited BP asset). */
-		TArray<TSharedRef<IPinTypeSelectorFilter>> PinTypeSelectorFilters;
+		/** Imported namespace pin type selector filter. */
+		TSharedPtr<IPinTypeSelectorFilter> PinTypeSelectorFilter;
 
 		/** Cached filter options widget. */
 		TSharedPtr<SWidget> FilterOptionsWidget;
@@ -1255,6 +1220,20 @@ void FBlueprintEditor::OnComponentDoubleClicked(TSharedPtr<FSubobjectEditorTreeN
 	}
 }
 
+void FBlueprintEditor::OnComponentAddedToBlueprint(const FSubobjectData& NewSubobjectData)
+{
+	const UObject* NewSubobject = NewSubobjectData.GetObject();
+	check(NewSubobject);
+
+	// If necessary, auto-import the namespace associated with the new subobject's class.
+	const bool bShouldAutoImportTypeNamespace = GetDefault<UBlueprintEditorSettings>()->bEnableNamespaceImportingFeatures;
+	if (bShouldAutoImportTypeNamespace)
+	{
+		const FString SubobjectTypeNamespace = FBlueprintNamespaceUtilities::GetObjectNamespace(NewSubobject->GetClass());
+		ImportNamespace(SubobjectTypeNamespace);
+	}
+}
+
 TSharedRef<SWidget> FBlueprintEditor::CreateGraphTitleBarWidget(TSharedRef<FTabInfo> InTabInfo, UEdGraph* InGraph)
 {
 	// Create the title bar widget
@@ -1901,15 +1880,18 @@ void FBlueprintEditor::CommonInitialization(const TArray<UBlueprint*>& InitBluep
 		DocumentManager->RegisterDocumentFactory(GraphEditorFactory);
 	}
 
-	// Create imported namespace type filters for value editing.
-	TArray<TSharedRef<FBlueprintNamespaceHelper>> LocalNamespaceHelpersArray;
-	LocalNamespaceHelpersArray.Reserve(InitBlueprints.Num());
+	// Create a namespace helper to keep track of imports for all BPs being edited.
+	ImportedNamespaceHelper = MakeShared<FBlueprintNamespaceHelper>();
+
+	// Add each Blueprint instance to be edited into the namespace helper's context.
 	for (const UBlueprint* BP : InitBlueprints)
 	{
-		LocalNamespaceHelpersArray.Add(GetOrCreateNamespaceHelperForBlueprint(BP));
+		ImportedNamespaceHelper->AddBlueprint(BP);
 	}
-	ImportedClassViewerFilter = MakeShared<BlueprintEditorImpl::FImportedClassViewerFilter>(LocalNamespaceHelpersArray);
-	ImportedPinTypeSelectorFilter = MakeShared<BlueprintEditorImpl::FImportedPinTypeSelectorFilter>(LocalNamespaceHelpersArray);
+
+	// Create imported namespace type filters for value editing.
+	ImportedClassViewerFilter = MakeShared<BlueprintEditorImpl::FImportedClassViewerFilterProxy>(ImportedNamespaceHelper->GetClassViewerFilter());
+	ImportedPinTypeSelectorFilter = MakeShared<BlueprintEditorImpl::FImportedPinTypeSelectorFilterProxy>(ImportedNamespaceHelper->GetPinTypeSelectorFilter());
 
 	// Make sure we know when tabs become active to update details tab
 	OnActiveTabChangedDelegateHandle = FGlobalTabmanager::Get()->OnActiveTabChanged_Subscribe( FOnActiveTabChanged::FDelegate::CreateRaw(this, &FBlueprintEditor::OnActiveTabChanged) );
@@ -1957,8 +1939,6 @@ void FBlueprintEditor::LoadLibrariesFromAssetRegistry()
 		const FString UserDeveloperPath = FPackageName::FilenameToLongPackageName( FPaths::GameUserDeveloperDir());
 		const FString DeveloperPath = FPackageName::FilenameToLongPackageName( FPaths::GameDevelopersDir() );
 
-		TSharedRef<FBlueprintNamespaceHelper> NamespaceHelper = GetOrCreateNamespaceHelperForBlueprint(BP);
-
 		// Interface blueprints don't show a node context menu anywhere so we can skip library loading
 		if (BP->BlueprintType != BPTYPE_Interface)
 		{
@@ -2000,10 +1980,8 @@ void FBlueprintEditor::LoadLibrariesFromAssetRegistry()
 				{
 					const FString BlueprintPath = AssetEntry.ObjectPath.ToString();
 
-					bool bAllowLoadBP = true;
-
 					// See if this passes the namespace check
-					bAllowLoadBP = bAllowLoadBP && NamespaceHelper->IsImportedAsset(AssetEntry);
+					bool bAllowLoadBP = !ImportedNamespaceHelper.IsValid() || ImportedNamespaceHelper->IsImportedAsset(AssetEntry);
 					
 					// For blueprints inside developers folder, only allow the ones inside current user's developers folder.
 					if (bAllowLoadBP)
@@ -2061,7 +2039,7 @@ void FBlueprintEditor::LoadLibrariesFromAssetRegistry()
 	}
 }
 
-void FBlueprintEditor::ImportNamespace(const FString& InNamespace)
+void FBlueprintEditor::ImportNamespace(const FString& InNamespace, const FImportNamespaceParameters& InParams)
 {
 	// No need to import the global namespace.
 	if (InNamespace.IsEmpty())
@@ -2069,50 +2047,105 @@ void FBlueprintEditor::ImportNamespace(const FString& InNamespace)
 		return;
 	}
 
-	bool bShouldReloadLibraries = false;
+	auto AddNamespaceToImportList = [](UBlueprint* InBlueprint, const FString& InNamespace) -> bool
+	{
+		if (!InBlueprint->ImportedNamespaces.Contains(InNamespace))
+		{
+			InBlueprint->Modify();
+			InBlueprint->ImportedNamespaces.Add(InNamespace);
+
+			return true;
+		}
+
+		return false;
+	};
 
-	// Add the namespace to any cached helper objects.
+	// Update the imported set for all edited objects.
+	bool bWasAdded = false;
 	const TArray<UObject*>& EditingObjs = GetEditingObjects();
 	for (UObject* EditingObj : EditingObjs)
 	{
 		if (UBlueprint* BlueprintObj = Cast<UBlueprint>(EditingObj))
 		{
 			// Add it into the Blueprint's user-facing import set.
-			if (BlueprintObj->BlueprintNamespace != InNamespace && !BlueprintObj->ImportedNamespaces.Contains(InNamespace))
+			bWasAdded |= AddNamespaceToImportList(BlueprintObj, InNamespace);
+			for (const FString& AdditionalNamespace : InParams.AdditionalNamespaces)
 			{
-				BlueprintObj->Modify();
-				BlueprintObj->ImportedNamespaces.Add(InNamespace);
+				bWasAdded |= AddNamespaceToImportList(BlueprintObj, AdditionalNamespace);
 			}
+		}
+	}
 
-			// Add it to the current scope of the Blueprint's editor context.
-			TSharedRef<FBlueprintNamespaceHelper> NamespaceHelper = GetOrCreateNamespaceHelperForBlueprint(BlueprintObj);
-			if (!NamespaceHelper->IsIncludedInNamespaceList(InNamespace))
-			{
-				NamespaceHelper->AddNamespace(InNamespace);
+	auto AddNamespaceToEditorContext = [](TSharedPtr<FBlueprintNamespaceHelper> ImportsHelper, const FString& InNamespace) -> bool
+	{
+		if (!ImportsHelper->IsIncludedInNamespaceList(InNamespace))
+		{
+			ImportsHelper->AddNamespace(InNamespace);
+			return true;
+		}
 
-				// Load additional libraries that may now be in scope.
-				bShouldReloadLibraries = true;
-			}
+		return false;
+	};
+
+	// Add it to the current scope of the Blueprint's editor context. Note that in certain cases, imports may already be associated
+	// with the Blueprint, but not yet associated with the editor context (e.g. - auto-import after setting a Blueprint's namespace;
+	// we won't add it to the Blueprint's import list, but we still want to add to the editor context and do any post-import actions).
+	if (ImportedNamespaceHelper.IsValid())
+	{
+		bWasAdded |= AddNamespaceToEditorContext(ImportedNamespaceHelper, InNamespace);
+		for (const FString& AdditionalNamespace : InParams.AdditionalNamespaces)
+		{
+			bWasAdded |= AddNamespaceToEditorContext(ImportedNamespaceHelper, AdditionalNamespace);
 		}
 	}
 
-	if (bShouldReloadLibraries)
+	if (bWasAdded)
 	{
+		// Load additional libraries that may now be in scope.
 		// @todo_namespaces - Make this more targeted - i.e. get/load only those assets tagged w/ the given namespace
 		LoadLibrariesFromAssetRegistry();
+
+		// Refresh class details on an auto-import if visible, since the list of imports has implicitly changed.
+		if (InParams.bIsAutoImport && IsDetailsPanelEditingGlobalOptions())
+		{
+			RefreshInspector();
+		}
+
+		// If bound, execute the post-import callback.
+		InParams.OnImportCallback.ExecuteIfBound();
 	}
 }
 
-TSharedRef<FBlueprintNamespaceHelper> FBlueprintEditor::GetOrCreateNamespaceHelperForBlueprint(const UBlueprint* InBlueprint)
+void FBlueprintEditor::RemoveNamespace(const FString& InNamespace)
 {
-	TWeakObjectPtr<const UBlueprint> Key = TWeakObjectPtr<const UBlueprint>(InBlueprint);
-	if (TSharedRef<FBlueprintNamespaceHelper>* ValuePtr = CachedNamespaceHelpers.Find(Key))
+	// Cannot remove the global namespace.
+	if (InNamespace.IsEmpty())
+	{
+		return;
+	}
+
+	// Update the imported set for all edited objects.
+	const TArray<UObject*>& EditingObjs = GetEditingObjects();
+	for (UObject* EditingObj : EditingObjs)
 	{
-		return *ValuePtr;
+		if (UBlueprint* BlueprintObj = Cast<UBlueprint>(EditingObj))
+		{
+			if (BlueprintObj->ImportedNamespaces.Contains(InNamespace))
+			{
+				BlueprintObj->Modify();
+				BlueprintObj->ImportedNamespaces.Remove(InNamespace);
+			}
+		}
 	}
 
-	TSharedPtr<FBlueprintNamespaceHelper> NewValue = MakeShared<FBlueprintNamespaceHelper>(InBlueprint);
-	return GetEditingObjects().Contains(InBlueprint) ? CachedNamespaceHelpers.Add(Key, NewValue.ToSharedRef()) : NewValue.ToSharedRef();
+	// Remove it from the current scope of the Blueprint's editor context.
+	if (ImportedNamespaceHelper.IsValid())
+	{
+		if (ImportedNamespaceHelper->IsIncludedInNamespaceList(InNamespace))
+		{
+			ImportedNamespaceHelper->RemoveNamespace(InNamespace);
+		}
+	}
 }
 
 void FBlueprintEditor::RegisterTabSpawners(const TSharedRef<class FTabManager>& InTabManager)
@@ -2779,7 +2812,7 @@ void FBlueprintEditor::CreateSubobjectEditors()
 		.AllowEditing(this, &FBlueprintEditor::InEditingMode)
 		.OnSelectionUpdated(this, &FBlueprintEditor::OnSelectionUpdated)
 		.OnItemDoubleClicked(this, &FBlueprintEditor::OnComponentDoubleClicked)
-		.OnImportNamespaceToEditorContext(this, &FBlueprintEditor::ImportNamespace)
+		.OnNewSubobjectAdded(this, &FBlueprintEditor::OnComponentAddedToBlueprint)
 		.SubobjectClassListFilters(ClassFilters);
 	
 	SubobjectViewport = SAssignNew(SubobjectViewport, SSCSEditorViewport)
@@ -10281,23 +10314,12 @@ void FBlueprintEditor::ClearAllGraphEditorQuickJumps()
 
 bool FBlueprintEditor::IsNonImportedObject(const UObject* InObject) const
 {
-	bool bNotImported = false;
-
-	for (const UObject* EditingObj : GetEditingObjects())
+	if (ImportedNamespaceHelper.IsValid() && !ImportedNamespaceHelper->IsImportedObject(InObject))
 	{
-		if (const UBlueprint* BP = Cast<UBlueprint>(EditingObj))
-		{
-			// Casting away the 'const' here currently because this is a non-const method that can modify the internally-cached set.
-			TSharedRef<FBlueprintNamespaceHelper> NamespaceHelper = const_cast<FBlueprintEditor*>(this)->GetOrCreateNamespaceHelperForBlueprint(BP);
-			if (!NamespaceHelper->IsImportedObject(InObject))
-			{
-				bNotImported = true;
-				break;
-			}
-		}
+		return true;
 	}
 
-	return bNotImported;
+	return false;
 }
 
 void FBlueprintEditor::OnBlueprintProjectSettingsChanged(UObject*, struct FPropertyChangedEvent&)
diff --git a/Engine/Source/Editor/Kismet/Private/BlueprintNamespaceHelper.cpp b/Engine/Source/Editor/Kismet/Private/BlueprintNamespaceHelper.cpp
--- a/Engine/Source/Editor/Kismet/Private/BlueprintNamespaceHelper.cpp
+++ b/Engine/Source/Editor/Kismet/Private/BlueprintNamespaceHelper.cpp
@@ -4,7 +4,6 @@
 #include "Engine/Blueprint.h"
 #include "BlueprintEditor.h"
 #include "BlueprintEditorSettings.h"
-#include "Settings/BlueprintEditorProjectSettings.h"
 #include "Toolkits/ToolkitManager.h"
 #include "EdGraph/EdGraphPin.h"
 #include "EdGraphSchema_K2.h"
@@ -120,46 +119,15 @@ private:
 
 // ---
 
-FBlueprintNamespaceHelper::FBlueprintNamespaceHelper(const UBlueprint* InBlueprint)
+FBlueprintNamespaceHelper::FBlueprintNamespaceHelper()
 {
 	// Instance the path tree used to store/retrieve namespaces.
 	NamespacePathTree = MakeUnique<FBlueprintNamespacePathTree>();
 
-	// Add the default namespace paths implicitly imported by every Blueprint.
-	AddNamespaces(GetDefault<UBlueprintEditorSettings>()->NamespacesToAlwaysInclude);
-	AddNamespaces(GetDefault<UBlueprintEditorProjectSettings>()->NamespacesToAlwaysInclude);
-
-	if (InBlueprint)
-	{
-		// Add the namespace for the given Blueprint.
-		AddNamespace(FBlueprintNamespaceUtilities::GetObjectNamespace(InBlueprint));
-
-		// Also add the namespace for the Blueprint's parent class.
-		AddNamespace(FBlueprintNamespaceUtilities::GetObjectNamespace(InBlueprint->ParentClass));
-
-		// Additional namespaces that are explicitly imported by this Blueprint.
-		AddNamespaces(InBlueprint->ImportedNamespaces);
-
-		// If enabled, also inherit namespaces that are explicitly imported by all ancestor BPs.
-		const bool bAddParentClassNamespaces = CVarBPImportParentClassNamespaces.GetValueOnGameThread();
-		if(bAddParentClassNamespaces)
-		{
-			const UClass* ParentClass = InBlueprint->ParentClass;
-			while (ParentClass)
-			{
-				if (const UBlueprint* ParentClassBlueprint = UBlueprint::GetBlueprintFromClass(ParentClass))
-				{
-					AddNamespaces(ParentClassBlueprint->ImportedNamespaces);
-				}
-				else
-				{
-					break;
-				}
-
-				ParentClass = ParentClass->GetSuperClass();
-			}
-		}
-	}
+	// Add global namespace paths implicitly imported by every Blueprint.
+	TSet<FString> GlobalImports;
+	FBlueprintNamespaceUtilities::GetSharedGlobalImports(GlobalImports);
+	AddNamespaces(GlobalImports);
 
 	// Instance the filters that can be used with type pickers, etc.
 	ClassViewerFilter = MakeShared<FClassViewerNamespaceFilter>(this);
@@ -170,6 +138,22 @@ FBlueprintNamespaceHelper::~FBlueprintNamespaceHelper()
 {
 }
 
+void FBlueprintNamespaceHelper::AddBlueprint(const UBlueprint* InBlueprint)
+{
+	if (!InBlueprint)
+	{
+		return;
+	}
+
+	// Add the default import set for the given Blueprint.
+	TSet<FString> DefaultImports;
+	FBlueprintNamespaceUtilities::GetDefaultImportsForBlueprint(InBlueprint, DefaultImports);
+	AddNamespaces(DefaultImports);
+
+	// Additional namespaces that are explicitly imported by this Blueprint.
+	AddNamespaces(InBlueprint->ImportedNamespaces);
+}
+
 void FBlueprintNamespaceHelper::AddNamespace(const FString& Namespace)
 {
 	if (!Namespace.IsEmpty())
@@ -179,6 +163,15 @@ void FBlueprintNamespaceHelper::AddNamespace(const FString& Namespace)
 	}
 }
 
+void FBlueprintNamespaceHelper::RemoveNamespace(const FString& Namespace)
+{
+	if (!Namespace.IsEmpty())
+	{
+		// Remove the path corresponding to the given namespace identifier.
+		NamespacePathTree->RemovePath(Namespace);
+	}
+}
+
 bool FBlueprintNamespaceHelper::IsIncludedInNamespaceList(const FString& TestNamespace) const
 {
 	// Empty namespace == global namespace
@@ -259,6 +252,7 @@ void FBlueprintNamespaceHelper::RefreshEditorFeatureConsoleFlags()
 
 		InitCVarFlag(CVarBPEnableNamespaceFilteringFeatures.AsVariable(), BlueprintEditorSettings->bEnableNamespaceFilteringFeatures);
 		InitCVarFlag(CVarBPEnableNamespaceImportingFeatures.AsVariable(), BlueprintEditorSettings->bEnableNamespaceImportingFeatures);
+		InitCVarFlag(CVarBPImportParentClassNamespaces.AsVariable(), BlueprintEditorSettings->bInheritImportedNamespacesFromParentBP);
 
 		bIsInitialized = true;
 	}
@@ -276,6 +270,7 @@ void FBlueprintNamespaceHelper::RefreshEditorFeatureConsoleFlags()
 
 		SetCVarFlag(CVarBPEnableNamespaceFilteringFeatures.AsVariable(), BlueprintEditorSettings->bEnableNamespaceFilteringFeatures);
 		SetCVarFlag(CVarBPEnableNamespaceImportingFeatures.AsVariable(), BlueprintEditorSettings->bEnableNamespaceImportingFeatures);
+		SetCVarFlag(CVarBPImportParentClassNamespaces.AsVariable(), BlueprintEditorSettings->bInheritImportedNamespacesFromParentBP);
 	}
 }
 
diff --git a/Engine/Source/Editor/Kismet/Private/BlueprintNamespacePathTree.h b/Engine/Source/Editor/Kismet/Private/BlueprintNamespacePathTree.h
--- a/Engine/Source/Editor/Kismet/Private/BlueprintNamespacePathTree.h
+++ b/Engine/Source/Editor/Kismet/Private/BlueprintNamespacePathTree.h
@@ -12,7 +12,7 @@ struct FBlueprintNamespacePathTree
 {
 public:
 	/** Path tree node structure. */
-	struct FNode
+	struct FNode : TSharedFromThis<FNode>
 	{
 		/** If TRUE, this node marks the end of an explicitly-added path string. Allows for "wildcard" paths which are inclusive of all subtrees. */
 		bool bIsAddedPath = false;
@@ -101,6 +101,76 @@ public:
 		Node->bIsAddedPath = true;
 	}
 
+	/**
+	 * Removes the given namespace identifier string as an explicitly-added path.
+	 *
+	 * @param InPath	A Blueprint namespace path identifier string (e.g. "X.Y.Z").
+	 */
+	void RemovePath(const FString& InPath)
+	{
+		TArray<FString> PathSegments;
+		InPath.ParseIntoArray(PathSegments, TEXT("."));
+
+		using FPathTraceNode = TTuple<FName, TSharedPtr<FNode>>;
+		TArray<FPathTraceNode> PathTrace;
+		PathTrace.Reserve(PathSegments.Num() + 1);
+
+		TSharedPtr<FNode> Node = GetRootNode();
+		PathTrace.Push(MakeTuple(NAME_None, Node));
+		for (const FString& PathSegment : PathSegments)
+		{
+			FName SubPathName = FName(*PathSegment);
+			if (const TSharedPtr<FNode>* NodePtr = Node->Children.Find(SubPathName))
+			{
+				Node = *NodePtr;
+				PathTrace.Push(MakeTuple(SubPathName, Node));
+			}
+		}
+
+		FName LastKey = NAME_None;
+		TSharedPtr<FNode> LastNode;
+		while (PathTrace.Num() > 0)
+		{
+			const FPathTraceNode& PathTraceNode = PathTrace.Pop();
+
+			TSharedPtr<FNode> CurrNode = PathTraceNode.Get<1>();
+			check(CurrNode.IsValid());
+			
+			// On the first pass, LastNode will not be valid, so we clear the added flag on the leaf node. On the
+			// next pass, the current node equates to the parent of the leaf node, and we remove the leaf node only
+			// if it has an empty subtree. All subsequent passes will iterate from there back up to the root node,
+			// where at each level we examine the last child node to see if it can be removed (no children and
+			// not explicitly added; i.e. - not a wildcard path (e.g. X.Y.*)). This means that we will not be left
+			// with a subtree that contains only non-explicitly added paths (i.e. we're also minimizing the tree).
+			//
+			// Example: Remove path "X.Y.Z" from the following tree:
+			//
+			//                         Root        <-- The root node is never removed.
+			//                          |
+			//                          +- X        <-- Node "X" will not be removed even if "Y" is removed, b/c its subtree is still not empty (i.e. "X.W" would still be a valid path). 
+			//                             |
+			//                             +- Y        <-- Node "Y" will be removed (*unless* "X.Y" was explicitly added) only if it's not left with any children after node "Z" is removed.
+			//                             |  |
+			//                             |  +- Z        <-- Node "Z" will not be removed if its subtree is not empty. Since this is the leaf node, we clear the "added" flag and shift up.
+			//                             |     |
+			//                             |     +- T
+			//                             +- W
+			//
+			if (!LastNode.IsValid())
+			{
+				CurrNode->bIsAddedPath = false;
+			}
+			else if(!LastNode->bIsAddedPath && LastNode->Children.Num() == 0)
+			{
+				TSharedPtr<FNode> RemovedNode = CurrNode->Children.FindAndRemoveChecked(LastKey);
+				check(RemovedNode == LastNode);
+			}
+
+			LastKey = PathTraceNode.Get<0>();
+			LastNode = CurrNode;
+		}
+	}
+
 	/** Path node visitor function signature.
 	 * 
 	 * @param CurrentPath	Current path (represented as a stack of names).
diff --git a/Engine/Source/Editor/Kismet/Private/BlueprintNamespaceUtilities.cpp b/Engine/Source/Editor/Kismet/Private/BlueprintNamespaceUtilities.cpp
--- a/Engine/Source/Editor/Kismet/Private/BlueprintNamespaceUtilities.cpp
+++ b/Engine/Source/Editor/Kismet/Private/BlueprintNamespaceUtilities.cpp
@@ -8,6 +8,8 @@
 #include "Misc/AssertionMacros.h"
 #include "BlueprintEditorModule.h"
 #include "Toolkits/ToolkitManager.h"
+#include "BlueprintEditorSettings.h"
+#include "Settings/BlueprintEditorProjectSettings.h"
 
 namespace UE::Editor::Kismet::Private
 {
@@ -153,7 +155,7 @@ FString FBlueprintNamespaceUtilities::GetObjectNamespace(const FSoftObjectPath&
 	return GetAssetNamespace(AssetData);
 }
 
-void FBlueprintNamespaceUtilities::GetPropertyValueNamespaces(const UStruct* InStruct, const FProperty* InProperty, const void* InContainer, TSet<FString>& OutNamespaces)
+void FBlueprintNamespaceUtilities::GetPropertyValueNamespaces(const UStruct* InStruct, const FProperty* InProperty, const void* InContainer, TArray<FString>& OutNamespaces)
 {
 	if (!InStruct || !InProperty || !InContainer)
 	{
@@ -214,7 +216,7 @@ void FBlueprintNamespaceUtilities::GetPropertyValueNamespaces(const UStruct* InS
 			if (ObjectPath.IsValid())
 			{
 				FString Namespace = GetObjectNamespace(ObjectPath);
-				OutNamespaces.Add(Namespace);
+				OutNamespaces.AddUnique(Namespace);
 			}
 		}
 		else if (const FObjectPropertyBase* ObjectProperty = CastField<FObjectPropertyBase>(InProperty))
@@ -223,7 +225,58 @@ void FBlueprintNamespaceUtilities::GetPropertyValueNamespaces(const UStruct* InS
 			if (ObjectValue)
 			{
 				FString Namespace = GetObjectNamespace(ObjectValue);
-				OutNamespaces.Add(Namespace);
+				OutNamespaces.AddUnique(Namespace);
+			}
+		}
+	}
+}
+
+void FBlueprintNamespaceUtilities::GetSharedGlobalImports(TSet<FString>& OutNamespaces)
+{
+	// Local editor imports.
+	OutNamespaces.Append(GetDefault<UBlueprintEditorSettings>()->NamespacesToAlwaysInclude);
+
+	// Project-wide imports.
+	OutNamespaces.Append(GetDefault<UBlueprintEditorProjectSettings>()->NamespacesToAlwaysInclude);
+
+	// Exclude the global namespace (empty) if it was included; this is implied.
+	OutNamespaces.Remove(FString());
+}
+
+void FBlueprintNamespaceUtilities::GetDefaultImportsForBlueprint(const UBlueprint* InBlueprint, TSet<FString>& OutNamespaces)
+{
+	if (!InBlueprint)
+	{
+		return;
+	}
+
+	// Blueprint namespace (if set).
+	FString BlueprintNamespace = GetObjectNamespace(InBlueprint);
+	if (!BlueprintNamespace.IsEmpty())
+	{
+		OutNamespaces.Add(BlueprintNamespace);
+	}
+
+	const bool bAddParentClassImportedNamespaces = GetDefault<UBlueprintEditorSettings>()->bInheritImportedNamespacesFromParentBP;
+
+	// Inherited Blueprint namespaces (if set).
+	TArray<UBlueprint*> ParentBPStack;
+	UBlueprint::GetBlueprintHierarchyFromClass(InBlueprint->ParentClass, ParentBPStack);
+	while (ParentBPStack.Num() > 0)
+	{
+		if (const UBlueprint* ParentBP = ParentBPStack.Pop())
+		{
+			// Parent Blueprint's namespace (if set).
+			FString ParentBPNamespace = GetObjectNamespace(ParentBP);
+			if (!ParentBPNamespace.IsEmpty())
+			{
+				OutNamespaces.Add(ParentBPNamespace);
+			}
+
+			// If enabled, also include namespaces that are explicitly imported by all ancestor BPs.
+			if (bAddParentClassImportedNamespaces)
+			{
+				OutNamespaces.Append(ParentBP->ImportedNamespaces);
 			}
 		}
 	}
diff --git a/Engine/Source/Editor/Kismet/Private/SBlueprintActionMenu.cpp b/Engine/Source/Editor/Kismet/Private/SBlueprintActionMenu.cpp
--- a/Engine/Source/Editor/Kismet/Private/SBlueprintActionMenu.cpp
+++ b/Engine/Source/Editor/Kismet/Private/SBlueprintActionMenu.cpp
@@ -622,37 +622,20 @@ void SBlueprintActionMenu::OnFilterImportNamespaceList(TArray<FString>& InOutNam
 
 void SBlueprintActionMenu::OnNamespaceSelectedForImport(const FString& InNamespace)
 {
-	bool bWasAdded = false;
-
-	FBlueprintActionContext MenuContext;
-	ConstructActionContext(MenuContext);
-
-	// Add to the blueprint's list of imports.
-	if (!InNamespace.IsEmpty())
+	TSharedPtr<FBlueprintEditor> BlueprintEditorPtr = EditorPtr.Pin();
+	if (BlueprintEditorPtr.IsValid())
 	{
-		for (UBlueprint* Blueprint : MenuContext.Blueprints)
+		FBlueprintEditor::FImportNamespaceParameters Params;
+		Params.OnImportCallback = FSimpleDelegate::CreateLambda([GraphActionMenu = this->GraphActionMenu]()
 		{
-			if (FBlueprintEditorUtils::AddNamespaceToImportList(Blueprint, InNamespace))
-			{
-				bWasAdded = true;
-			}
-		}
-	}
-
-	if (bWasAdded)
-	{
-		// Import the namespace into the current editor context. This may load additional type assets.
-		TSharedPtr<FBlueprintEditor> BlueprintEditorPtr = EditorPtr.Pin();
-		if (BlueprintEditorPtr.IsValid())
-		{
-			BlueprintEditorPtr->ImportNamespace(InNamespace);
-			BlueprintEditorPtr->RefreshInspector();
-		}
-
-		// Now that additional types have been loaded/imported, update the menu to include any additional action(s).
-		const bool bPreserveExpansion = true;
-		const bool bHandleOnSelectionEvent = false;
-		GraphActionMenu->RefreshAllActions(bPreserveExpansion, bHandleOnSelectionEvent);
+			// Now that additional types have been loaded/imported, update the menu to include any additional action(s).
+			const bool bPreserveExpansion = true;
+			const bool bHandleOnSelectionEvent = false;
+			GraphActionMenu->RefreshAllActions(bPreserveExpansion, bHandleOnSelectionEvent);
+		});
+
+		// Auto-import the namespace into the current editor context. This may load additional type assets.
+		BlueprintEditorPtr->ImportNamespace(InNamespace, Params);
 	}
 }
 
diff --git a/Engine/Source/Editor/Kismet/Private/SBlueprintNamespaceEntry.cpp b/Engine/Source/Editor/Kismet/Private/SBlueprintNamespaceEntry.cpp
--- a/Engine/Source/Editor/Kismet/Private/SBlueprintNamespaceEntry.cpp
+++ b/Engine/Source/Editor/Kismet/Private/SBlueprintNamespaceEntry.cpp
@@ -30,7 +30,7 @@ void SBlueprintNamespaceEntry::Construct(const FArguments& InArgs)
 			.Font(InArgs._Font)
 			.ForegroundColor(FSlateColor::UseForeground())
 			.Visibility(InArgs._AllowTextEntry ? EVisibility::Visible : EVisibility::Collapsed)
-			.Text(FText::FromString(CurrentNamespace))
+			.Text(this, &SBlueprintNamespaceEntry::GetCurrentNamespaceText)
 			.OnTextChanged(this, &SBlueprintNamespaceEntry::OnTextChanged)
 			.OnTextCommitted(this, &SBlueprintNamespaceEntry::OnTextCommitted)
 			.OnShowingSuggestions(this, &SBlueprintNamespaceEntry::OnShowingSuggestions)
@@ -265,6 +265,11 @@ void SBlueprintNamespaceEntry::OnNamespaceListSelectionChanged(TSharedPtr<FStrin
 	}
 }
 
+FText SBlueprintNamespaceEntry::GetCurrentNamespaceText() const
+{
+	return FText::FromString(CurrentNamespace);
+}
+
 void SBlueprintNamespaceEntry::PopulateNamespaceList()
 {
 	// Clear the current list.
diff --git a/Engine/Source/Editor/Kismet/Private/SBlueprintPalette.cpp b/Engine/Source/Editor/Kismet/Private/SBlueprintPalette.cpp
--- a/Engine/Source/Editor/Kismet/Private/SBlueprintPalette.cpp
+++ b/Engine/Source/Editor/Kismet/Private/SBlueprintPalette.cpp
@@ -842,7 +842,7 @@ private:
 				if (PinSubCategoryObject && BlueprintEditor.IsValid() && BlueprintEditor->IsNonImportedObject(PinSubCategoryObject))
 				{
 					FString Namespace = FBlueprintNamespaceUtilities::GetObjectNamespace(PinSubCategoryObject);
-					if (!Namespace.IsEmpty())
+					if (!Namespace.IsEmpty() && BlueprintEditor.IsValid())
 					{
 						BlueprintEditor->ImportNamespace(Namespace);
 					}
diff --git a/Engine/Source/Editor/Kismet/Private/SKismetInspector.cpp b/Engine/Source/Editor/Kismet/Private/SKismetInspector.cpp
--- a/Engine/Source/Editor/Kismet/Private/SKismetInspector.cpp
+++ b/Engine/Source/Editor/Kismet/Private/SKismetInspector.cpp
@@ -1017,7 +1017,7 @@ void SKismetInspector::ImportNamespacesForPropertyValue(const FProperty* InPrope
 	}
 
 	// Gather all namespace identifier strings associated with the property's value for each edited object.
-	TSet<FString> AssociatedNamespaces;
+	TArray<FString> AssociatedNamespaces;
 	for (const TWeakObjectPtr<UObject>& SelectedObjectPtr : SelectedObjects)
 	{
 		if (const UObject* SelectedObject = SelectedObjectPtr.Get())
@@ -1043,10 +1043,9 @@ void SKismetInspector::ImportNamespacesForPropertyValue(const FProperty* InPrope
 		TSharedPtr<FBlueprintEditor> BlueprintEditor = BlueprintEditorPtr.Pin();
 		if (BlueprintEditor.IsValid())
 		{
-			for (const FString& AssociatedNamespace : AssociatedNamespaces)
-			{
-				BlueprintEditor->ImportNamespace(AssociatedNamespace);
-			}
+			FBlueprintEditor::FImportNamespaceParameters Params;
+			Params.AdditionalNamespaces = MakeArrayView(AssociatedNamespaces).RightChop(1);
+			BlueprintEditor->ImportNamespace(AssociatedNamespaces[0], Params);
 		}
 	}
 }
diff --git a/Engine/Source/Editor/Kismet/Public/BlueprintEditor.h b/Engine/Source/Editor/Kismet/Public/BlueprintEditor.h
--- a/Engine/Source/Editor/Kismet/Public/BlueprintEditor.h
+++ b/Engine/Source/Editor/Kismet/Public/BlueprintEditor.h
@@ -55,6 +55,7 @@ class SSubobjectEditor;
 class FSubobjectEditorTreeNode;
 class UBlueprintEditorSettings;
 class UBlueprintEditorProjectSettings;
+struct FSubobjectData;
 
 /* Enums to use when grouping the blueprint members in the list panel. The order here will determine the order in the list */
 namespace NodeSectionID
@@ -521,6 +522,9 @@ public:
 
 	/** Delegate invoked when an item is double clicked in the subobject editor widget */
 	virtual void OnComponentDoubleClicked(TSharedPtr<class FSubobjectEditorTreeNode> Node);
+
+	/** Delegate invoked when a new subobject is added in the subobject editor widget */
+	virtual void OnComponentAddedToBlueprint(const FSubobjectData& NewSubobjectData);
 	
 	/** Delegate invoked when the selection is changed in the subobject editor widget */
 	UE_DEPRECATED(5.0, "OnSelectionUpdated(const TArray<TSharedPtr<class FSCSEditorTreeNode>>&) has been deprecated. Use OnSelectionUpdated(const TArray<TSharedPtr<class FSubobjectEditorTreeNode>>&) instead.")
@@ -714,11 +718,31 @@ public:
 	/** Gets the default schema for this editor */
 	TSubclassOf<UEdGraphSchema> GetDefaultSchema() const { return GetDefaultSchemaClass(); }
 
+	/** Extended parameters for the ImportNamespace() method */
+	struct FImportNamespaceParameters
+	{
+		/** Whether this is an automatic or explicit (i.e. user-initiated) action. */
+		bool bIsAutoImport;
+
+		/** Callback to use for any post-import actions. Will not be invoked if nothing is imported. */
+		FSimpleDelegate OnImportCallback;
+
+		/** Additional namespaces to be imported as part of the single transaction. */
+		TArray<FString> AdditionalNamespaces;
+
+		/** Default ctor (for initialization). */
+		FImportNamespaceParameters()
+		{
+			// Treat as auto-import by default (implies that the editor will auto-refresh the details view).
+			bIsAutoImport = true;
+		}
+	};
+	
 	/** Imports the given namespace into the editor. This may trigger a load event for additional macro and/or function library assets if not already loaded. */
-	void ImportNamespace(const FString& InNamespace);
+	void ImportNamespace(const FString& InNamespace, const FImportNamespaceParameters& InParams = FImportNamespaceParameters());
 
-	/** Returns an instanced namespace helper utility object that corresponds to the given Blueprint. */
-	TSharedRef<FBlueprintNamespaceHelper> GetOrCreateNamespaceHelperForBlueprint(const UBlueprint* InBlueprint);
+	/** Removes the given namespace from the editor's current import context. However, this will NOT unload any associated macro and/or function library assets. */
+	void RemoveNamespace(const FString& InNamespace);
 
 	/** Selects a local variable to load in the details panel. */
 	virtual bool SelectLocalVariable(const UEdGraph* Graph, const FName& VariableName) { return false; }
@@ -1363,8 +1387,8 @@ protected:
 	/** The toolbar builder class */
 	TSharedPtr<class FBlueprintEditorToolbar> Toolbar;
 
-	/** Cached set of instanced namespace helper objects */
-	TMap<TWeakObjectPtr<const UBlueprint>, TSharedRef<FBlueprintNamespaceHelper>> CachedNamespaceHelpers;
+	/** Imported namespace helper utility object */
+	TSharedPtr<FBlueprintNamespaceHelper> ImportedNamespaceHelper;
 
 	/** Filter used to restrict class viewer widgets in the editor context to imported namespaces only */
 	TSharedPtr<class IClassViewerFilter> ImportedClassViewerFilter;
diff --git a/Engine/Source/Editor/Kismet/Public/BlueprintNamespaceHelper.h b/Engine/Source/Editor/Kismet/Public/BlueprintNamespaceHelper.h
--- a/Engine/Source/Editor/Kismet/Public/BlueprintNamespaceHelper.h
+++ b/Engine/Source/Editor/Kismet/Public/BlueprintNamespaceHelper.h
@@ -6,6 +6,7 @@
 
 class IClassViewerFilter;
 class IPinTypeSelectorFilter;
+class UBlueprint;
 struct FAssetData;
 struct FBlueprintNamespacePathTree;
 
@@ -15,16 +16,19 @@ struct FBlueprintNamespacePathTree;
 class KISMET_API FBlueprintNamespaceHelper
 {
 public:
-	/**
-	 * Default constructor.
-	 * 
-	 * @param InBlueprint	The Blueprint object being edited.
-	 */
-	explicit FBlueprintNamespaceHelper(const class UBlueprint* InBlueprint);
+	/** Default constructor. */
+	FBlueprintNamespaceHelper();
 
 	/** Note: We explicitly declare/implement the dtor so that forward-declared types can be destroyed. */
 	virtual ~FBlueprintNamespaceHelper();
 
+	/**
+	 * Add a new Blueprint into the current helper context.
+	 * 
+	 * @param InBlueprint	A Blueprint object being edited.
+	 */
+	void AddBlueprint(const UBlueprint* InBlueprint);
+
 	/**
 	 * Add a namespace identifier into the Blueprint editor context that's managed by this instance.
 	 *
@@ -46,6 +50,13 @@ public:
 		}
 	}
 
+	/**
+	 * Remove a namespace identifier from the Blueprint editor context that's managed by this instance.
+	 * 
+	 * @param Namespace		The namespace identifier string to remove. Should resemble "X.Y.Z" as the format.
+	 */
+	void RemoveNamespace(const FString& Namespace);
+
 	/**
 	 * Tests a namespace identifier to see if it's been imported. 
 	 *
diff --git a/Engine/Source/Editor/Kismet/Public/BlueprintNamespaceUtilities.h b/Engine/Source/Editor/Kismet/Public/BlueprintNamespaceUtilities.h
--- a/Engine/Source/Editor/Kismet/Public/BlueprintNamespaceUtilities.h
+++ b/Engine/Source/Editor/Kismet/Public/BlueprintNamespaceUtilities.h
@@ -52,7 +52,22 @@ public:
 	 * @param InContainer	The source address of the struct/object containing the property's value.
 	 * @param OutNamespaces	Zero or more unique namespace identifier(s) referenced by the property value. An entry with an empty string equates to the default global namespace.
 	 */
-	static void GetPropertyValueNamespaces(const UStruct* InStruct, const FProperty* InProperty, const void* InContainer, TSet<FString>& OutNamespaces);
+	static void GetPropertyValueNamespaces(const UStruct* InStruct, const FProperty* InProperty, const void* InContainer, TArray<FString>& OutNamespaces);
+
+	/**
+	 * Gathers the set of global namespaces that are implicitly imported by all Blueprint assets.
+	 * 
+	 * @param OutNamespaces	Zero or more unique namespace identifier(s) representing the shared global namespace set.
+	 */
+	static void GetSharedGlobalImports(TSet<FString>& OutNamespaces);
+
+	/**
+	 * Gathers the set of default namespaces that are implicitly imported by a given Blueprint asset.
+	 *
+	 * @param InBlueprint	A reference to a loaded Blueprint asset.
+	 * @param OutNamespaces	Zero or more unique namespace identifier(s) representing the Blueprint asset's default namespace set.
+	 */
+	static void GetDefaultImportsForBlueprint(const UBlueprint* InBlueprint, TSet<FString>& OutNamespaces);
 
 	/**
 	 * Sets the default Blueprint namespace type that objects/assets should use when not explicitly assigned.
diff --git a/Engine/Source/Editor/Kismet/Public/SBlueprintNamespaceEntry.h b/Engine/Source/Editor/Kismet/Public/SBlueprintNamespaceEntry.h
--- a/Engine/Source/Editor/Kismet/Public/SBlueprintNamespaceEntry.h
+++ b/Engine/Source/Editor/Kismet/Public/SBlueprintNamespaceEntry.h
@@ -58,6 +58,7 @@ protected:
 	TSharedRef<ITableRow> OnGenerateRowForNamespaceList(TSharedPtr<FString> Item, const TSharedRef<STableViewBase>& OwnerTable);
 	void OnNamespaceListFilterTextChanged(const FText& InText);
 	void OnNamespaceListSelectionChanged(TSharedPtr<FString> Item, ESelectInfo::Type SelectInfo);
+	FText GetCurrentNamespaceText() const;
 
 	void PopulateNamespaceList();
 	void SelectNamespace(const FString& InNamespace);
diff --git a/Engine/Source/Editor/SubobjectEditor/Private/SSubobjectBlueprintEditor.cpp b/Engine/Source/Editor/SubobjectEditor/Private/SSubobjectBlueprintEditor.cpp
--- a/Engine/Source/Editor/SubobjectEditor/Private/SSubobjectBlueprintEditor.cpp
+++ b/Engine/Source/Editor/SubobjectEditor/Private/SSubobjectBlueprintEditor.cpp
@@ -27,8 +27,6 @@
 #include "Kismet2/ChildActorComponentEditorUtils.h"
 #include "ObjectTools.h"						// ThumbnailTools::CacheEmptyThumbnail
 #include "K2Node_ComponentBoundEvent.h"
-#include "BlueprintEditorSettings.h"
-#include "BlueprintNamespaceUtilities.h"
 
 #define LOCTEXT_NAMESPACE "SSubobjectBlueprintEditor"
 
@@ -46,7 +44,7 @@ void SSubobjectBlueprintEditor::Construct(const FArguments& InArgs)
 	OnSelectionUpdated = InArgs._OnSelectionUpdated;
 	OnItemDoubleClicked = InArgs._OnItemDoubleClicked;
 	OnHighlightPropertyInDetailsView = InArgs._OnHighlightPropertyInDetailsView;
-	OnImportNamespaceToEditorContext = InArgs._OnImportNamespaceToEditorContext;
+	OnNewSubobjectAdded = InArgs._OnNewSubobjectAdded;
 	AllowEditing = InArgs._AllowEditing;
 	HideComponentClassCombo = InArgs._HideComponentClassCombo;
 	bAllowTreeUpdates = true;
@@ -518,12 +516,6 @@ TSharedPtr<SWidget> SSubobjectBlueprintEditor::BuildSceneRootDropActionMenu(FSub
 
 FSubobjectDataHandle SSubobjectBlueprintEditor::AddNewSubobject(const FSubobjectDataHandle& ParentHandle, UClass* NewClass, UObject* AssetOverride, FText& OutFailReason, TUniquePtr<FScopedTransaction> InOngoingTransaction)
 {
-	if(GetDefault<UBlueprintEditorSettings>()->bEnableNamespaceImportingFeatures)
-	{
-		// If bound, invoke the delegate to import the namespace associated with the given class.
-		OnImportNamespaceToEditorContext.ExecuteIfBound(FBlueprintNamespaceUtilities::GetObjectNamespace(NewClass));
-	}
-
 	FAddNewSubobjectParams Params;
     Params.ParentHandle = ParentHandle;
     Params.NewClass = NewClass;
@@ -535,6 +527,14 @@ FSubobjectDataHandle SSubobjectBlueprintEditor::AddNewSubobject(const FSubobject
     check(System);
 	
 	FSubobjectDataHandle NewHandle = System->AddNewSubobject(Params,OutFailReason);
+
+	if (NewHandle.IsValid())
+	{
+		if (const FSubobjectData* NewSubobjectData = NewHandle.GetData())
+		{
+			OnNewSubobjectAdded.ExecuteIfBound(*NewSubobjectData);
+		}
+	}
 	
     return NewHandle;	
 }
diff --git a/Engine/Source/Editor/SubobjectEditor/Public/SSubobjectBlueprintEditor.h b/Engine/Source/Editor/SubobjectEditor/Public/SSubobjectBlueprintEditor.h
--- a/Engine/Source/Editor/SubobjectEditor/Public/SSubobjectBlueprintEditor.h
+++ b/Engine/Source/Editor/SubobjectEditor/Public/SSubobjectBlueprintEditor.h
@@ -13,7 +13,7 @@ class SUBOBJECTEDITOR_API SSubobjectBlueprintEditor final : public SSubobjectEdi
 {
 public:
 	DECLARE_DELEGATE_OneParam(FOnHighlightPropertyInDetailsView, const class FPropertyPath&);
-	DECLARE_DELEGATE_OneParam(FOnImportNamespaceToEditorContext, const FString& InNamespace);
+	DECLARE_DELEGATE_OneParam(FOnNewSubobjectAdded, const FSubobjectData&);
 
 private:	
 	SLATE_BEGIN_ARGS(SSubobjectBlueprintEditor)
@@ -23,7 +23,7 @@ private:
 		, _HideComponentClassCombo(false)
         , _OnSelectionUpdated()
 		, _OnHighlightPropertyInDetailsView()
-		, _OnImportNamespaceToEditorContext()
+		, _OnNewSubobjectAdded()
 		{}
 
 		SLATE_ATTRIBUTE(UObject*, ObjectContext)
@@ -33,7 +33,7 @@ private:
 	    SLATE_EVENT(FOnSelectionUpdated, OnSelectionUpdated)
 	    SLATE_EVENT(FOnItemDoubleClicked, OnItemDoubleClicked)
 		SLATE_EVENT(FOnHighlightPropertyInDetailsView, OnHighlightPropertyInDetailsView)
-		SLATE_EVENT(FOnImportNamespaceToEditorContext, OnImportNamespaceToEditorContext)
+		SLATE_EVENT(FOnNewSubobjectAdded, OnNewSubobjectAdded)
 		SLATE_ARGUMENT(TArray<TSharedRef<IClassViewerFilter>>, SubobjectClassListFilters)
 	    
 	SLATE_END_ARGS()
@@ -74,8 +74,8 @@ public:
 	/** Delegate to invoke when the given property should be highlighted in the details view (e.g. diff). */
 	FOnHighlightPropertyInDetailsView OnHighlightPropertyInDetailsView;
 
-	/** Delegate to invoke when a component class is selected in order to ensure the namespace is imported into the current editor context. */
-	FOnImportNamespaceToEditorContext OnImportNamespaceToEditorContext;
+	/** Delegate to invoke when a subobject type is chosen to add a new subobject entry. */
+	FOnNewSubobjectAdded OnNewSubobjectAdded;
 	
 	/**
 	* Fills out an events section in ui.
diff --git a/Engine/Source/Editor/UnrealEd/Private/Kismet2/BlueprintEditorUtils.cpp b/Engine/Source/Editor/UnrealEd/Private/Kismet2/BlueprintEditorUtils.cpp
--- a/Engine/Source/Editor/UnrealEd/Private/Kismet2/BlueprintEditorUtils.cpp
+++ b/Engine/Source/Editor/UnrealEd/Private/Kismet2/BlueprintEditorUtils.cpp
@@ -7858,32 +7858,6 @@ FName FBlueprintEditorUtils::FindUniqueCustomEventName(const UBlueprint* Bluepri
 	return FindUniqueKismetName(Blueprint, LOCTEXT("DefaultCustomEventName", "CustomEvent").ToString());
 }
 
-//////////////////////////////////////////////////////////////////////////
-// Scoping
-
-bool FBlueprintEditorUtils::AddNamespaceToImportList(UBlueprint* Blueprint, const FString& Namespace)
-{
-	check(!Namespace.IsEmpty());
-
-	if (Blueprint->ImportedNamespaces.Contains(Namespace))
-	{
-		ShowNotification(
-			FText::Format(
-				LOCTEXT("NamespaceAlreadyImportedFmt", "Blueprint '{0}' already imports the namespace '{1}'"),
-				FText::FromString(Blueprint->GetName()),
-				FText::FromString(Namespace)
-			),
-			EMessageSeverity::Warning
-		);
-		return false;
-	}
-
-	Blueprint->ImportedNamespaces.Add(Namespace);
-	FBlueprintEditorUtils::MarkBlueprintAsModified(Blueprint);
-
-	return true;
-}
-
 //////////////////////////////////////////////////////////////////////////
 // Timeline
 
diff --git a/Engine/Source/Editor/UnrealEd/Public/Kismet2/BlueprintEditorUtils.h b/Engine/Source/Editor/UnrealEd/Public/Kismet2/BlueprintEditorUtils.h
--- a/Engine/Source/Editor/UnrealEd/Public/Kismet2/BlueprintEditorUtils.h
+++ b/Engine/Source/Editor/UnrealEd/Public/Kismet2/BlueprintEditorUtils.h
@@ -1522,12 +1522,6 @@ public:
 	/** Finds a unique and valid name for a custom event */
 	static FName FindUniqueCustomEventName(const UBlueprint* Blueprint);
 
-	//////////////////////////////////////////////////////////////////////////
-	// Scoping
-
-	/** Add a new namespace to the blueprint's import list. Will return false if already imported. */
-	static bool AddNamespaceToImportList(UBlueprint* Blueprint, const FString& Namespace);
-
 	//////////////////////////////////////////////////////////////////////////
 	// Timeline
 
