Fix pin tooltip not dismissing when pin destroyed if tooltip hovered.

diff --git a/Engine/Source/Editor/GraphEditor/Private/SGraphPin.cpp b/Engine/Source/Editor/GraphEditor/Private/SGraphPin.cpp
--- a/Engine/Source/Editor/GraphEditor/Private/SGraphPin.cpp
+++ b/Engine/Source/Editor/GraphEditor/Private/SGraphPin.cpp
@@ -151,7 +151,8 @@ SGraphPin::~SGraphPin()
 {
 	if (ValueInspectorTooltip.IsValid())
 	{
-		ValueInspectorTooltip.Pin()->TryDismissTooltip();
+		bool bForceDismiss = true;
+		ValueInspectorTooltip.Pin()->TryDismissTooltip(bForceDismiss);
 	}
 }
 
diff --git a/Engine/Source/Editor/Kismet/Private/SPinValueInspector.cpp b/Engine/Source/Editor/Kismet/Private/SPinValueInspector.cpp
--- a/Engine/Source/Editor/Kismet/Private/SPinValueInspector.cpp
+++ b/Engine/Source/Editor/Kismet/Private/SPinValueInspector.cpp
@@ -272,12 +272,12 @@ void FPinValueInspectorTooltip::MoveTooltip(const FVector2D& InNewLocation)
 	}
 }
 
-void FPinValueInspectorTooltip::TryDismissTooltip()
+void FPinValueInspectorTooltip::TryDismissTooltip(bool bForceDismiss)
 {
 	if (Instance.IsValid() && ValueInspectorWidget.IsValid())
 	{
 		// if we can't inspect the pin, force it to close
-		if (!FKismetDebugUtilities::CanInspectPinValue(ValueInspectorWidget->PinRef.Get()) || TooltipCanClose())
+		if (!FKismetDebugUtilities::CanInspectPinValue(ValueInspectorWidget->PinRef.Get()) || TooltipCanClose() || bForceDismiss)
 		{
 			DismissTooltip();
 		}
diff --git a/Engine/Source/Editor/Kismet/Public/SPinValueInspector.h b/Engine/Source/Editor/Kismet/Public/SPinValueInspector.h
--- a/Engine/Source/Editor/Kismet/Public/SPinValueInspector.h
+++ b/Engine/Source/Editor/Kismet/Public/SPinValueInspector.h
@@ -69,8 +69,10 @@ public:
 	/** 
 	 * Dismisses the current tooltip, if it is not currently hovered
 	 * NOTE: Pin a shared ptr before calling this function 
+	 * 
+	 * @param bForceDismiss	true to dismiss regardless of hover & tooltip ownership 
 	 */
-	void TryDismissTooltip();
+	void TryDismissTooltip(bool bForceDismiss = false);
 
 private:
 	/** Dismisses the current tooltip (internal implementation) */
