Non-unity compile fix.

diff --git a/Engine/Source/Editor/Kismet/Public/BlueprintNamespaceUtilities.h b/Engine/Source/Editor/Kismet/Public/BlueprintNamespaceUtilities.h
--- a/Engine/Source/Editor/Kismet/Public/BlueprintNamespaceUtilities.h
+++ b/Engine/Source/Editor/Kismet/Public/BlueprintNamespaceUtilities.h
@@ -5,6 +5,7 @@
 #include "CoreMinimal.h"
 
 struct FAssetData;
+class UBlueprint;
 
 /** Default namespace type for objects/assets if one is not explicitly assigned. */
 enum class EDefaultBlueprintNamespaceType
