Customized Blueprint editor and project settings for global namespace imports entry.

Change summary:
- Moved base "managed list" customization out into its own class so that it can be shared with editor/project settings customization.
- Refactored Blueprint global options (i.e. Class Settings) details customization to utilize the new "managed list" base class.
- Removed filter text as a customizable managed list display option and redefined it to be more consistent with property editor nodes.
- Fixed the "Imports" display in Class Settings to no longer list global imports in the non-default section.
- Added a new details customization for UBlueprintEditorSettings.
- Added a customized editor for the 'NamespacesToAlwaysInclude' property for both local user settings and shared project settings.
- Moved the "Add" button for imports and interfaces to be on the same line as the header row (eliminates some wasted space in the details view).
- Fixed up global import property display names and tooltips to be more descriptive.
- Fixed an issue that prevented reselection of the previous namespace setting in the Blueprint namespace entry combo box after clearing it.
- Fixed the incorrect vertical alignment of import/interface item rows.

diff --git a/Engine/Source/Editor/BlueprintGraph/Private/BlueprintEditorSettings.cpp b/Engine/Source/Editor/BlueprintGraph/Private/BlueprintEditorSettings.cpp
--- a/Engine/Source/Editor/BlueprintGraph/Private/BlueprintEditorSettings.cpp
+++ b/Engine/Source/Editor/BlueprintGraph/Private/BlueprintEditorSettings.cpp
@@ -6,10 +6,10 @@
 #include "Settings/EditorExperimentalSettings.h"
 #include "BlueprintActionDatabase.h"
 #include "FindInBlueprintManager.h"
-#include "Subsystems/AssetEditorSubsystem.h"
-#include "Editor.h"
 #include "BlueprintTypePromotion.h"
+#include "Modules/ModuleManager.h"
 #include "AssetRegistry/AssetRegistryModule.h"
+#include "BlueprintEditorModule.h"
 #include "BlueprintNamespaceHelper.h"
 #include "BlueprintNamespaceUtilities.h"
 
@@ -155,6 +155,15 @@ void UBlueprintEditorSettings::PostEditChangeProperty(FPropertyChangedEvent& Pro
 		// Refresh the Blueprint editor UI environment to match current settings.
 		FBlueprintNamespaceUtilities::RefreshBlueprintEditorFeatures();
 	}
+	else if (PropertyName == GET_MEMBER_NAME_CHECKED(UBlueprintEditorSettings, NamespacesToAlwaysInclude))
+	{
+		// Close any open Blueprint editor windows so that we have a chance to reload them with the updated import set.
+		FBlueprintEditorModule& BlueprintEditorModule = FModuleManager::LoadModuleChecked<FBlueprintEditorModule>("Kismet");
+		for (const TSharedRef<IBlueprintEditor>& BlueprintEditor : BlueprintEditorModule.GetBlueprintEditors())
+		{
+			BlueprintEditor->CloseWindow();
+		}
+	}
 
 	Super::PostEditChangeProperty(PropertyChangedEvent);
 }
\ No newline at end of file
diff --git a/Engine/Source/Editor/BlueprintGraph/Public/BlueprintEditorSettings.h b/Engine/Source/Editor/BlueprintGraph/Public/BlueprintEditorSettings.h
--- a/Engine/Source/Editor/BlueprintGraph/Public/BlueprintEditorSettings.h
+++ b/Engine/Source/Editor/BlueprintGraph/Public/BlueprintEditorSettings.h
@@ -180,8 +180,8 @@ public:
 	UPROPERTY(Transient)
 	bool bInheritImportedNamespacesFromParentBP;
 
-	// The list of namespaces to always expose in any Blueprint (local per-user)
-	UPROPERTY(EditAnywhere, config, Category = Experimental, meta = (EditCondition = "bEnableNamespaceEditorFeatures"))
+	// A list of namespace identifiers that the Blueprint editor should always import by default. Requires Blueprint namespace features to be enabled and only applies to the current local user. Editing this list will also cause any visible Blueprint editor windows to be closed.
+	UPROPERTY(EditAnywhere, config, Category = Experimental, meta = (EditCondition = "bEnableNamespaceEditorFeatures", DisplayName = "Global Namespace Imports (Local User Only)"))
 	TArray<FString> NamespacesToAlwaysInclude;
 
 	/** If enabled, then placed cast nodes will default to their "pure" form (meaning: without execution pins). */
diff --git a/Engine/Source/Editor/Kismet/Private/BlueprintDetailsCustomization.cpp b/Engine/Source/Editor/Kismet/Private/BlueprintDetailsCustomization.cpp
--- a/Engine/Source/Editor/Kismet/Private/BlueprintDetailsCustomization.cpp
+++ b/Engine/Source/Editor/Kismet/Private/BlueprintDetailsCustomization.cpp
@@ -5542,12 +5542,12 @@ FReply FBlueprintGraphActionDetails::OnAddNewOutputClicked()
 }
 
 
-FBlueprintManagedListDetails::FBlueprintManagedListDetails(TWeakPtr<class FBlueprintGlobalOptionsDetails> InGlobalOptionsDetailsPtr)
-: GlobalOptionsDetailsPtr(InGlobalOptionsDetailsPtr)
+FBlueprintGlobalOptionsManagedListDetails::FBlueprintGlobalOptionsManagedListDetails(TWeakPtr<class FBlueprintGlobalOptionsDetails> InGlobalOptionsDetailsPtr)
+	: GlobalOptionsDetailsPtr(InGlobalOptionsDetailsPtr)
 {
 }
 
-UBlueprint* FBlueprintManagedListDetails::GetBlueprintObjectChecked() const
+UBlueprint* FBlueprintGlobalOptionsManagedListDetails::GetBlueprintObjectChecked() const
 {
 	UBlueprint* BlueprintObject = nullptr;
 
@@ -5561,7 +5561,7 @@ UBlueprint* FBlueprintManagedListDetails::GetBlueprintObjectChecked() const
 	return BlueprintObject;
 }
 
-TSharedPtr<FBlueprintEditor> FBlueprintManagedListDetails::GetPinnedBlueprintEditorPtr() const
+TSharedPtr<FBlueprintEditor> FBlueprintGlobalOptionsManagedListDetails::GetPinnedBlueprintEditorPtr() const
 {
 	TSharedPtr<FBlueprintGlobalOptionsDetails> PinnedGlobalOptionsDetailsPtr = GlobalOptionsDetailsPtr.Pin();
 	if (PinnedGlobalOptionsDetailsPtr.IsValid())
@@ -5572,106 +5572,7 @@ TSharedPtr<FBlueprintEditor> FBlueprintManagedListDetails::GetPinnedBlueprintEdi
 	return TSharedPtr<FBlueprintEditor>();
 }
 
-void FBlueprintManagedListDetails::GenerateHeaderRowContent(FDetailWidgetRow& HeaderRow)
-{
-	HeaderRow
-	[
-		SNew(STextBlock)
-		.Text(DisplayOptions.TitleText)
-		.Font(IDetailLayoutBuilder::GetDetailFont())
-	];
-}
-
-BEGIN_SLATE_FUNCTION_BUILD_OPTIMIZATION
-void FBlueprintManagedListDetails::GenerateChildContent(IDetailChildrenBuilder& ChildrenBuilder)
-{
-	TArray<FManagedListItem> Items;
-	GetManagedListItems(Items);
-
-	if (Items.Num() > 0)
-	{
-		for (const FManagedListItem& Item : Items)
-		{
-			TSharedPtr<SHorizontalBox> Box;
-			ChildrenBuilder.AddCustomRow(DisplayOptions.ItemRowFilterText)
-			[
-				SAssignNew(Box, SHorizontalBox)
-				+ SHorizontalBox::Slot()
-				[
-					SNew(STextBlock)
-					.Text(Item.DisplayName)
-					.Font(IDetailLayoutBuilder::GetDetailFont())
-				]
-			];
-
-			if (Item.AssetPtr.IsValid())
-			{
-				TSharedRef<SWidget> BrowseButton = PropertyCustomizationHelpers::MakeBrowseButton(FSimpleDelegate::CreateLambda([this, AssetPtr = Item.AssetPtr]() -> void
-				{
-					if (AssetPtr.IsValid())
-					{
-						GEditor->GetEditorSubsystem<UAssetEditorSubsystem>()->OpenEditorForAsset(AssetPtr.Get());
-					}
-				}));
-
-				BrowseButton->SetToolTipText(DisplayOptions.BrowseButtonToolTipText);
-
-				Box->AddSlot()
-				.AutoWidth()
-				.Padding(2.0f, 0.0f)
-				[
-					BrowseButton
-				];
-			}
-
-			if (Item.bIsRemovable)
-			{
-				TSharedRef<SWidget> RemoveButton = PropertyCustomizationHelpers::MakeClearButton(FSimpleDelegate::CreateLambda([this, Item]() -> void
-				{
-					OnRemoveItem(Item);
-				}));
-
-				RemoveButton->SetToolTipText(DisplayOptions.RemoveButtonToolTipText);
-
-				Box->AddSlot()
-				.AutoWidth()
-				[
-					RemoveButton
-				];
-			}
-		}
-	}
-	else
-	{
-		ChildrenBuilder.AddCustomRow(DisplayOptions.ItemRowFilterText)
-			[
-			SNew(STextBlock)
-			.Text(DisplayOptions.NoItemsLabelText)
-			.Font(IDetailLayoutBuilder::GetDetailFontItalic())
-		];
-	}
-
-	TSharedPtr<SWidget> AddItemRowWidget = MakeAddItemRowWidget();
-	if (AddItemRowWidget.IsValid())
-	{
-		ChildrenBuilder.AddCustomRow(DisplayOptions.AddItemRowFilterText)
-		[
-			SNew(SBox)
-			.HAlign(HAlign_Right)
-			[
-				AddItemRowWidget.ToSharedRef()
-			]
-		];
-	}
-}
-END_SLATE_FUNCTION_BUILD_OPTIMIZATION
-
-void FBlueprintManagedListDetails::RegenerateChildContent()
-{
-	RegenerateChildrenDelegate.ExecuteIfBound();
-}
-
-void FBlueprintManagedListDetails::OnRefreshInDetailsView()
+void FBlueprintGlobalOptionsManagedListDetails::OnRefreshInDetailsView()
 {
 	TSharedPtr<FBlueprintEditor> BlueprintEditorPtr = GetPinnedBlueprintEditorPtr();
 	if (BlueprintEditorPtr.IsValid())
@@ -5687,18 +5588,16 @@ void FBlueprintManagedListDetails::OnRefreshInDetailsView()
 
 
 FBlueprintImportsLayout::FBlueprintImportsLayout(TWeakPtr<class FBlueprintGlobalOptionsDetails> InGlobalOptionsDetails, bool bInShowDefaultImports)
-	:FBlueprintManagedListDetails(InGlobalOptionsDetails)
-	,bShouldShowDefaultImports(bInShowDefaultImports)
+	: FBlueprintGlobalOptionsManagedListDetails(InGlobalOptionsDetails)
+	, bShouldShowDefaultImports(bInShowDefaultImports)
 {
 	DisplayOptions.TitleText = bShouldShowDefaultImports ?
 		LOCTEXT("BlueprintDefaultNamespaceTitle", "Default Namespaces") :
 		LOCTEXT("BlueprintImportedNamespaceTitle", "Imported Namespaces");
 	DisplayOptions.NoItemsLabelText = LOCTEXT("NoBlueprintImports", "No Imports");
-	DisplayOptions.ItemRowFilterText = LOCTEXT("BlueprintImportsValue", "Imports Value");
-	DisplayOptions.AddItemRowFilterText = LOCTEXT("BlueprintAddImport", "Add Import");
 }
 
-TSharedPtr<SWidget> FBlueprintImportsLayout::MakeAddItemRowWidget()
+TSharedPtr<SWidget> FBlueprintImportsLayout::MakeAddItemWidget()
 {
 	if (bShouldShowDefaultImports)
 	{
@@ -5734,18 +5633,14 @@ void FBlueprintImportsLayout::GetManagedListItems(TArray<FManagedListItem>& OutL
 	TSet<FString> NamespaceItems;
 	const UBlueprint* Blueprint = GetBlueprintObjectChecked();
 
-	if (bShouldShowDefaultImports)
-	{
-		FBlueprintNamespaceUtilities::GetSharedGlobalImports(NamespaceItems);
-		FBlueprintNamespaceUtilities::GetDefaultImportsForBlueprint(Blueprint, NamespaceItems);
-	}
-	else
-	{
-		// Blueprint imports (removable).
-		NamespaceItems.Append(Blueprint->ImportedNamespaces);
+	// Default imports (non-removable). These include anything from the shared global set, as well as any namespaces assigned to the Blueprint hierarchy.
+	FBlueprintNamespaceUtilities::GetSharedGlobalImports(NamespaceItems);
+	FBlueprintNamespaceUtilities::GetDefaultImportsForBlueprint(Blueprint, NamespaceItems);
 
-		// A Blueprint can explicitly import and also be assigned to the same namespace.
-		NamespaceItems.Remove(Blueprint->BlueprintNamespace);
+	if(!bShouldShowDefaultImports)
+	{
+		// Blueprint imports (removable). A Blueprint may explicitly import a namespace that's also in the default set, but we exclude those here so they can't be removed.
+		NamespaceItems = Blueprint->ImportedNamespaces.Difference(NamespaceItems);
 	}
 
 	AddNamespaceItemsToOutputList(NamespaceItems, !bShouldShowDefaultImports);
@@ -5795,19 +5690,17 @@ void FBlueprintImportsLayout::OnFilterNamespaceList(TArray<FString>& InOutNamesp
 
 
 FBlueprintInterfaceLayout::FBlueprintInterfaceLayout(TWeakPtr<class FBlueprintGlobalOptionsDetails> InGlobalOptionsDetails, bool bInShowsInheritedInterfaces)
-	:FBlueprintManagedListDetails(InGlobalOptionsDetails)
-	,bShowsInheritedInterfaces(bInShowsInheritedInterfaces)
+	: FBlueprintGlobalOptionsManagedListDetails(InGlobalOptionsDetails)
+	, bShowsInheritedInterfaces(bInShowsInheritedInterfaces)
 {
 	DisplayOptions.TitleText = bShowsInheritedInterfaces ?
 		LOCTEXT("BlueprintInheritedInterfaceTitle", "Inherited Interfaces") :
 		LOCTEXT("BlueprintImplementedInterfaceTitle", "Implemented Interfaces");
 	DisplayOptions.NoItemsLabelText = LOCTEXT("NoBlueprintInterface", "No Interfaces");
-	DisplayOptions.ItemRowFilterText = LOCTEXT("BlueprintInterfaceValue", "Interface Value");
-	DisplayOptions.AddItemRowFilterText = LOCTEXT("BlueprintAddInterface", "Add Interface");
 	DisplayOptions.BrowseButtonToolTipText = LOCTEXT("BlueprintInterfaceBrowseTooltip", "Opens this interface");
 }
 
-TSharedPtr<SWidget> FBlueprintInterfaceLayout::MakeAddItemRowWidget()
+TSharedPtr<SWidget> FBlueprintInterfaceLayout::MakeAddItemWidget()
 {
 	if (!bShowsInheritedInterfaces)
 	{
diff --git a/Engine/Source/Editor/Kismet/Private/BlueprintDetailsCustomization.h b/Engine/Source/Editor/Kismet/Private/BlueprintDetailsCustomization.h
--- a/Engine/Source/Editor/Kismet/Private/BlueprintDetailsCustomization.h
+++ b/Engine/Source/Editor/Kismet/Private/BlueprintDetailsCustomization.h
@@ -17,6 +17,7 @@
 #include "SMyBlueprint.h"
 #include "SGraphPin.h"
 #include "UObject/WeakFieldPtr.h"
+#include "BlueprintManagedListDetails.h"
 
 class Error;
 class FBlueprintGlobalOptionsDetails;
@@ -682,91 +683,41 @@ private:
 	TArray<TSharedPtr<IDetailCustomization>> ExternalDetailCustomizations;
 };
 
-/** Blueprint managed list details */
-class FBlueprintManagedListDetails : public IDetailCustomNodeBuilder
+/** Blueprint global options - managed list details customization base */
+class FBlueprintGlobalOptionsManagedListDetails : public FBlueprintManagedListDetails
 {
 public:
 	/** Constructor method */
-	FBlueprintManagedListDetails(TWeakPtr<class FBlueprintGlobalOptionsDetails> InGlobalOptionsDetailsPtr);
-	
-protected:
-	/** List item node type */
-	struct FManagedListItem
-	{
-		FString ItemName;
-		FText DisplayName;
-		TWeakObjectPtr<> AssetPtr;
-		uint8 bIsRemovable : 1;
-
-		FManagedListItem()
-			:bIsRemovable(false)
-		{
-		}
-	};
-
-	/** Customizable display options */
-	struct FDisplayOptions
-	{
-		FText TitleText;
-		FText NoItemsLabelText;
-		FText ItemRowFilterText;
-		FText AddItemRowFilterText;
-		FText BrowseButtonToolTipText;
-		FText RemoveButtonToolTipText;
-	};
-
-	/** Mutable display options */
-	FDisplayOptions DisplayOptions;
+	FBlueprintGlobalOptionsManagedListDetails(TWeakPtr<class FBlueprintGlobalOptionsDetails> InGlobalOptionsDetailsPtr);
 
+protected:
 	/** Access to external resources */
 	UBlueprint* GetBlueprintObjectChecked() const;
 	TSharedPtr<FBlueprintEditor> GetPinnedBlueprintEditorPtr() const;
 
-	/** IDetailCustomNodeBuilder interface */
-	virtual void GenerateHeaderRowContent(FDetailWidgetRow& HeaderRow) override;
-	virtual void GenerateChildContent(IDetailChildrenBuilder& ChildrenBuilder) override;
-	virtual FName GetName() const override { return NAME_None; }
-	virtual bool InitiallyCollapsed() const override { return false; }
-	virtual void Tick(float DeltaTime) override {}
-	virtual bool RequiresTick() const override { return false; }
-	virtual void SetOnRebuildChildren(FSimpleDelegate InOnRebuildChildren) override { RegenerateChildrenDelegate = InOnRebuildChildren; }
-	/** END IDetailCustomNodeBuilder interface */
-
-	/** Overridable interface methods */
-	virtual TSharedPtr<SWidget> MakeAddItemRowWidget() { return nullptr; }
-	virtual void GetManagedListItems(TArray<FManagedListItem>& OutListItems) const {}
-	virtual void OnRemoveItem(const FManagedListItem& Item) {}
-
-	/** Helper function to regenerate the details customization */
-	void RegenerateChildContent();
-
 	/** Helper function to set the Blueprint back into the KismetInspector's details view */
 	void OnRefreshInDetailsView();
 
 private:
 	/** The parent graph action details customization */
 	TWeakPtr<class FBlueprintGlobalOptionsDetails> GlobalOptionsDetailsPtr;
-
-	/** A delegate to regenerate this list of children */
-	FSimpleDelegate RegenerateChildrenDelegate;
 };
 
 /** Blueprint Imports List Details */
-class FBlueprintImportsLayout : public FBlueprintManagedListDetails, public TSharedFromThis<FBlueprintImportsLayout>
+class FBlueprintImportsLayout : public FBlueprintGlobalOptionsManagedListDetails, public TSharedFromThis<FBlueprintImportsLayout>
 {
 public:
 	FBlueprintImportsLayout(TWeakPtr<class FBlueprintGlobalOptionsDetails> InGlobalOptionsDetails, bool bInShowDefaultImports);
 
 protected:
 	/** FBlueprintManagedListDetails interface*/
-	virtual TSharedPtr<SWidget> MakeAddItemRowWidget() override;
+	virtual TSharedPtr<SWidget> MakeAddItemWidget() override;
 	virtual void GetManagedListItems(TArray<FManagedListItem>& OutListItems) const override;
 	virtual void OnRemoveItem(const FManagedListItem& Item) override;
 	/** END FBlueprintManagedListDetails interface */
 
 	void OnNamespaceSelected(const FString& InNamespace);
 	void OnFilterNamespaceList(TArray<FString>& InOutNamespaceList);
-	void HandleImportEntryTextCommitted(const FText& NewLabel, ETextCommit::Type CommitType);
 
 private:
 	/** Whether we should show default (i.e. global) imports versus custom (i.e. local) imports */
@@ -774,14 +725,14 @@ private:
 };
 
 /** Blueprint Interface List Details */
-class FBlueprintInterfaceLayout : public FBlueprintManagedListDetails, public TSharedFromThis<FBlueprintInterfaceLayout>
+class FBlueprintInterfaceLayout : public FBlueprintGlobalOptionsManagedListDetails, public TSharedFromThis<FBlueprintInterfaceLayout>
 {
 public:
 	FBlueprintInterfaceLayout(TWeakPtr<class FBlueprintGlobalOptionsDetails> InGlobalOptionsDetails, bool bInShowsInheritedInterfaces);
 
 protected:
 	/** FBlueprintManagedListDetails interface */
-	virtual TSharedPtr<SWidget> MakeAddItemRowWidget() override;
+	virtual TSharedPtr<SWidget> MakeAddItemWidget() override;
 	virtual void GetManagedListItems(TArray<FManagedListItem>& OutListItems) const override;
 	virtual void OnRemoveItem(const FManagedListItem& Item) override;
 	/** END FBlueprintManagedListDetails interface */
@@ -794,11 +745,11 @@ private:
 	void OnClassPicked(UClass* PickedClass);
 
 private:
-	/** Whether we show inherited interfaces versus implemented interfaces */
-	bool bShowsInheritedInterfaces;
-
 	/** The add interface combo button */
 	TSharedPtr<SComboButton> AddInterfaceComboButton;
+
+	/** Whether we show inherited interfaces versus implemented interfaces */
+	bool bShowsInheritedInterfaces;
 };
 
 /** Details customization for Blueprint settings */
diff --git a/Engine/Source/Editor/Kismet/Private/BlueprintManagedListDetails.cpp b/Engine/Source/Editor/Kismet/Private/BlueprintManagedListDetails.cpp
new file mode 100644
--- /dev/null
+++ b/Engine/Source/Editor/Kismet/Private/BlueprintManagedListDetails.cpp
@@ -0,0 +1,148 @@
+#include "BlueprintManagedListDetails.h"
+#include "DetailLayoutBuilder.h"
+#include "DetailWidgetRow.h"
+#include "IDetailChildrenBuilder.h"
+#include "PropertyCustomizationHelpers.h"
+#include "SlateOptMacros.h"
+#include "Widgets/SBoxPanel.h"
+#include "Widgets/SNullWidget.h"
+#include "Widgets/Layout/SBox.h"
+#include "Widgets/Text/STextBlock.h"
+#include "Subsystems/AssetEditorSubsystem.h"
+#include "Editor.h"
+
+void FBlueprintManagedListDetails::GenerateHeaderRowContent(FDetailWidgetRow& HeaderRow)
+{
+	TSharedPtr<SWidget> AddItemWidget = MakeAddItemWidget();
+	if (!AddItemWidget.IsValid())
+	{
+		AddItemWidget = SNullWidget::NullWidget;
+	}
+
+	HeaderRow
+	.NameContent()
+	[
+		SNew(STextBlock)
+		.Text(DisplayOptions.TitleText)
+		.Font(IDetailLayoutBuilder::GetDetailFont())
+		.ToolTipText(DisplayOptions.TitleTooltipText)
+	]
+	.ValueContent()
+	[
+		SNew(SHorizontalBox)
+		+SHorizontalBox::Slot()
+		.HAlign(HAlign_Left)
+		.AutoWidth()
+		[
+			AddItemWidget.ToSharedRef()
+		]
+	]
+	.FilterString(
+		DisplayOptions.TitleText
+	)
+	.EditCondition(
+		DisplayOptions.EditCondition,
+		FOnBooleanValueChanged()
+	);
+}
+
+BEGIN_SLATE_FUNCTION_BUILD_OPTIMIZATION
+void FBlueprintManagedListDetails::GenerateChildContent(IDetailChildrenBuilder& ChildrenBuilder)
+{
+	TArray<FManagedListItem> Items;
+	GetManagedListItems(Items);
+
+	if (Items.Num() > 0)
+	{
+		for (const FManagedListItem& Item : Items)
+		{
+			TSharedPtr<SHorizontalBox> Box;
+			ChildrenBuilder.AddCustomRow(FText::GetEmpty())
+			[
+				SAssignNew(Box, SHorizontalBox)
+				+ SHorizontalBox::Slot()
+				.FillWidth(1.0f)
+				.VAlign(VAlign_Center)
+				[
+					SNew(STextBlock)
+					.Text(Item.DisplayName)
+					.Font(IDetailLayoutBuilder::GetDetailFont())
+				]
+			]
+			.FilterString(
+				Item.DisplayName
+			)
+			.EditCondition(
+				DisplayOptions.EditCondition,
+				FOnBooleanValueChanged()
+			);
+
+			if (Item.AssetPtr.IsValid())
+			{
+				TSharedRef<SWidget> BrowseButton = PropertyCustomizationHelpers::MakeBrowseButton(FSimpleDelegate::CreateLambda([this, AssetPtr = Item.AssetPtr]() -> void
+				{
+					if (AssetPtr.IsValid())
+					{
+						GEditor->GetEditorSubsystem<UAssetEditorSubsystem>()->OpenEditorForAsset(AssetPtr.Get());
+					}
+				}));
+
+				BrowseButton->SetToolTipText(DisplayOptions.BrowseButtonToolTipText);
+
+				Box->AddSlot()
+				.AutoWidth()
+				.HAlign(HAlign_Right)
+				.VAlign(VAlign_Center)
+				.Padding(2.0f, 0.0f)
+				[
+					BrowseButton
+				];
+			}
+
+			if (Item.bIsRemovable)
+			{
+				TSharedRef<SWidget> RemoveButton = PropertyCustomizationHelpers::MakeClearButton(FSimpleDelegate::CreateLambda([this, Item]() -> void
+				{
+					OnRemoveItem(Item);
+				}));
+
+				RemoveButton->SetToolTipText(DisplayOptions.RemoveButtonToolTipText);
+
+				Box->AddSlot()
+				.AutoWidth()
+				.HAlign(HAlign_Right)
+				.VAlign(VAlign_Center)
+				.Padding(2.0f, 0.0f)
+				[
+					RemoveButton
+				];
+			}
+		}
+	}
+	else
+	{
+		ChildrenBuilder.AddCustomRow(FText::GetEmpty())
+		[
+			SNew(SBox)
+			.VAlign(VAlign_Center)
+			[
+				SNew(STextBlock)
+				.Text(DisplayOptions.NoItemsLabelText)
+				.Font(IDetailLayoutBuilder::GetDetailFontItalic())
+			]
+		]
+		.FilterString(
+			DisplayOptions.NoItemsLabelText
+		)
+		.EditCondition(
+			DisplayOptions.EditCondition,
+			FOnBooleanValueChanged()
+		);
+	}
+}
+END_SLATE_FUNCTION_BUILD_OPTIMIZATION
+
+void FBlueprintManagedListDetails::RegenerateChildContent()
+{
+	RegenerateChildrenDelegate.ExecuteIfBound();
+}
\ No newline at end of file
diff --git a/Engine/Source/Editor/Kismet/Private/SBlueprintNamespaceEntry.cpp b/Engine/Source/Editor/Kismet/Private/SBlueprintNamespaceEntry.cpp
--- a/Engine/Source/Editor/Kismet/Private/SBlueprintNamespaceEntry.cpp
+++ b/Engine/Source/Editor/Kismet/Private/SBlueprintNamespaceEntry.cpp
@@ -150,13 +150,13 @@ TSharedRef<SWidget> SBlueprintNamespaceEntry::OnGetNamespaceListMenuContent()
 
 	// Construct the list view widget that we'll use for the menu content.
 	SAssignNew(ListView, SListView<TSharedPtr<FString>>)
-		.SelectionMode(ESelectionMode::SingleToggle)
+		.SelectionMode(ESelectionMode::Single)
 		.ListItemsSource(&ListItems)
 		.OnGenerateRow(this, &SBlueprintNamespaceEntry::OnGenerateRowForNamespaceList)
 		.OnSelectionChanged(this, &SBlueprintNamespaceEntry::OnNamespaceListSelectionChanged);
 
-	// If the current namespace is non-empty, look for a matching item in the set.
-	if (!CurrentNamespace.IsEmpty())
+	// If we are allowing manual entry and the current namespace is non-empty, look for a matching item in the set.
+	if (!CurrentNamespace.IsEmpty() && TextBox.IsValid() && TextBox->GetVisibility().IsVisible())
 	{
 		const TSharedPtr<FString>* CurrentItemPtr = ListItems.FindByPredicate([&CurrentNamespace = this->CurrentNamespace](const TSharedPtr<FString>& Item)
 		{
diff --git a/Engine/Source/Editor/Kismet/Public/BlueprintManagedListDetails.h b/Engine/Source/Editor/Kismet/Public/BlueprintManagedListDetails.h
new file mode 100644
--- /dev/null
+++ b/Engine/Source/Editor/Kismet/Public/BlueprintManagedListDetails.h
@@ -0,0 +1,65 @@
+// Copyright Epic Games, Inc. All Rights Reserved.
+
+#pragma once
+
+#include "CoreMinimal.h"
+#include "Misc/Attribute.h"
+#include "UObject/WeakObjectPtr.h"
+#include "IDetailCustomNodeBuilder.h"
+
+class SWidget;
+
+/** Blueprint managed list details */
+class KISMET_API FBlueprintManagedListDetails : public IDetailCustomNodeBuilder
+{
+protected:
+	/** List item node type */
+	struct FManagedListItem
+	{
+		FString ItemName;
+		FText DisplayName;
+		TWeakObjectPtr<> AssetPtr;
+		uint8 bIsRemovable : 1;
+
+		FManagedListItem()
+			:bIsRemovable(false)
+		{
+		}
+	};
+
+	/** Customizable display options */
+	struct FDisplayOptions
+	{
+		FText TitleText;
+		FText TitleTooltipText;
+		FText NoItemsLabelText;
+		FText BrowseButtonToolTipText;
+		FText RemoveButtonToolTipText;
+		TAttribute<bool> EditCondition;
+	};
+
+	/** Mutable display options */
+	FDisplayOptions DisplayOptions;
+
+	/** IDetailCustomNodeBuilder interface */
+	virtual void GenerateHeaderRowContent(FDetailWidgetRow& HeaderRow) override;
+	virtual void GenerateChildContent(IDetailChildrenBuilder& ChildrenBuilder) override;
+	virtual FName GetName() const override { return NAME_None; }
+	virtual bool InitiallyCollapsed() const override { return false; }
+	virtual void Tick(float DeltaTime) override {}
+	virtual bool RequiresTick() const override { return false; }
+	virtual void SetOnRebuildChildren(FSimpleDelegate InOnRebuildChildren) override { RegenerateChildrenDelegate = InOnRebuildChildren; }
+	/** END IDetailCustomNodeBuilder interface */
+
+	/** Overridable interface methods */
+	virtual TSharedPtr<SWidget> MakeAddItemWidget() { return nullptr; }
+	virtual void GetManagedListItems(TArray<FManagedListItem>& OutListItems) const {}
+	virtual void OnRemoveItem(const FManagedListItem& Item) {}
+
+	/** Helper function to regenerate the details customization */
+	void RegenerateChildContent();
+
+private:
+	/** A delegate to regenerate this list of children */
+	FSimpleDelegate RegenerateChildrenDelegate;
+};
\ No newline at end of file
diff --git a/Engine/Source/Editor/UnrealEd/Private/Settings/BlueprintEditorProjectSettings.cpp b/Engine/Source/Editor/UnrealEd/Private/Settings/BlueprintEditorProjectSettings.cpp
--- a/Engine/Source/Editor/UnrealEd/Private/Settings/BlueprintEditorProjectSettings.cpp
+++ b/Engine/Source/Editor/UnrealEd/Private/Settings/BlueprintEditorProjectSettings.cpp
@@ -1,9 +1,9 @@
 // Copyright Epic Games, Inc. All Rights Reserved.
 
 #include "Settings/BlueprintEditorProjectSettings.h"
-#include "UObject/UnrealType.h"
-#include "Toolkits/ToolkitManager.h"
-#include "BlueprintEditor.h"
+#include "Modules/ModuleManager.h"
+#include "BlueprintEditorModule.h"
+#include "Editor.h"
 
 /* UBlueprintEditorProjectSettings */
 
@@ -20,36 +20,29 @@ void UBlueprintEditorProjectSettings::PostEditChangeProperty(struct FPropertyCha
 	const FName Name = PropertyChangedEvent.Property ? PropertyChangedEvent.Property->GetFName() : NAME_None;
 	if (Name == GET_MEMBER_NAME_CHECKED(UBlueprintEditorProjectSettings, bEnableChildActorExpansionInTreeView))
 	{
-		if (!GEditor)
-		{
-			return;
-		}
-
 		// Find open blueprint editors and refresh them
-		if (UAssetEditorSubsystem* AssetEditorSubsystem = GEditor->GetEditorSubsystem<UAssetEditorSubsystem>())
+		FBlueprintEditorModule& BlueprintEditorModule = FModuleManager::LoadModuleChecked<FBlueprintEditorModule>("Kismet");
+		for (const TSharedRef<IBlueprintEditor>& BlueprintEditor : BlueprintEditorModule.GetBlueprintEditors())
 		{
-			TArray<UObject*> EditedAssets = AssetEditorSubsystem->GetAllEditedAssets();
-			for (UObject* Asset : EditedAssets)
-			{
-				if (Asset && Asset->IsA<UBlueprint>())
-				{
-					TSharedPtr<IToolkit> AssetEditorPtr = FToolkitManager::Get().FindEditorForAsset(Asset);
-					if (AssetEditorPtr.IsValid() && AssetEditorPtr->IsBlueprintEditor())
-					{
-						TSharedPtr<IBlueprintEditor> BlueprintEditorPtr = StaticCastSharedPtr<IBlueprintEditor>(AssetEditorPtr);
-						BlueprintEditorPtr->RefreshEditors();
-					}
-				}
-			}
+			BlueprintEditor->RefreshEditors();
 		}
 
 		// Deselect actors so we are forced to clear the current tree view
 		// @todo - Figure out how to update the tree view directly instead?
-		if (GEditor->GetSelectedActorCount() > 0)
+		if (GEditor && GEditor->GetSelectedActorCount() > 0)
 		{
 			const bool bNoteSelectionChange = true;
 			const bool bDeselectBSPSurfaces = true;
 			GEditor->SelectNone(bNoteSelectionChange, bDeselectBSPSurfaces);
 		}
 	}
+	else if (Name == GET_MEMBER_NAME_CHECKED(UBlueprintEditorProjectSettings, NamespacesToAlwaysInclude))
+	{
+		// Close any open Blueprint editor windows so that we have a chance to reload them with the updated import set.
+		FBlueprintEditorModule& BlueprintEditorModule = FModuleManager::LoadModuleChecked<FBlueprintEditorModule>("Kismet");
+		for (const TSharedRef<IBlueprintEditor>& BlueprintEditor : BlueprintEditorModule.GetBlueprintEditors())
+		{
+			BlueprintEditor->CloseWindow();
+		}
+	}
 }
diff --git a/Engine/Source/Editor/UnrealEd/Private/Settings/BlueprintEditorProjectSettingsCustomization.cpp b/Engine/Source/Editor/UnrealEd/Private/Settings/BlueprintEditorProjectSettingsCustomization.cpp
--- a/Engine/Source/Editor/UnrealEd/Private/Settings/BlueprintEditorProjectSettingsCustomization.cpp
+++ b/Engine/Source/Editor/UnrealEd/Private/Settings/BlueprintEditorProjectSettingsCustomization.cpp
@@ -6,33 +6,122 @@
 #include "PropertyHandle.h"
 #include "DetailLayoutBuilder.h"
 #include "IDetailPropertyRow.h"
+#include "DetailWidgetRow.h"
+#include "PropertyCustomizationHelpers.h"
+#include "BlueprintManagedListDetails.h"
+#include "SBlueprintNamespaceEntry.h"
+#include "ScopedTransaction.h"
 
 #define LOCTEXT_NAMESPACE "FBlueprintEditorProjectSettingsCustomization"
 
-TSharedRef<IDetailCustomization> FBlueprintEditorProjectSettingsCustomization::MakeInstance()
+namespace UE::Editor::BlueprintEditorProjectSettingsCustomization::Private
 {
-	return MakeShared<FBlueprintEditorProjectSettingsCustomization>();
+	class FBlueprintGlobalProjectImportsLayout : public FBlueprintManagedListDetails, public TSharedFromThis<FBlueprintGlobalProjectImportsLayout>
+	{
+	public:
+		FBlueprintGlobalProjectImportsLayout(TSharedRef<IPropertyHandle> InPropertyHandle)
+			: FBlueprintManagedListDetails()
+			, PropertyHandle(InPropertyHandle)
+		{
+			DisplayOptions.TitleText = PropertyHandle->GetPropertyDisplayName();
+			DisplayOptions.TitleTooltipText = PropertyHandle->GetToolTipText();
+
+			DisplayOptions.NoItemsLabelText = LOCTEXT("NoGlobalImports", "None");
+
+			// Add a custom edit condition to link it to the editor's namespace feature toggle flag (for consistency w/ the editor-specific set).
+			DisplayOptions.EditCondition = TAttribute<bool>::CreateLambda([]()
+			{
+				return GetDefault<UBlueprintEditorSettings>()->bEnableNamespaceEditorFeatures;
+			});
+		}
+
+	protected:
+		/** FBlueprintManagedListDetails interface*/
+		virtual TSharedPtr<SWidget> MakeAddItemWidget() override
+		{
+			return SNew(SBlueprintNamespaceEntry)
+				.AllowTextEntry(false)
+				.OnNamespaceSelected(this, &FBlueprintGlobalProjectImportsLayout::OnNamespaceSelected)
+				.OnFilterNamespaceList(this, &FBlueprintGlobalProjectImportsLayout::OnFilterNamespaceList)
+				.ButtonContent()
+				[
+					SNew(STextBlock)
+					.Text(LOCTEXT("BlueprintAddGlobalImportButton", "Add"))
+					.ToolTipText(LOCTEXT("BlueprintAddGlobalImportButton_Tooltip", "Choose a namespace that all Blueprints in this project should import by default (applies to all users)."))
+				];
+		}
+
+		virtual void GetManagedListItems(TArray<FManagedListItem>& OutListItems) const override
+		{
+			for(const FString& GlobalNamespace : GetDefault<UBlueprintEditorProjectSettings>()->NamespacesToAlwaysInclude)
+			{
+				FManagedListItem ItemDesc;
+				ItemDesc.ItemName = GlobalNamespace;
+				ItemDesc.DisplayName = FText::FromString(GlobalNamespace);
+				ItemDesc.bIsRemovable = true;
+
+				OutListItems.Add(MoveTemp(ItemDesc));
+			}
+		}
+
+		virtual void OnRemoveItem(const FManagedListItem& Item)
+		{
+			FScopedTransaction Transaction(LOCTEXT("RemoveGlobalImport_Transaction", "Remove Global Import"));
+
+			UBlueprintEditorProjectSettings* BlueprintEditorProjectSettings = GetMutableDefault<UBlueprintEditorProjectSettings>();
+			check(BlueprintEditorProjectSettings);
+
+			PropertyHandle->NotifyPreChange();
+			BlueprintEditorProjectSettings->NamespacesToAlwaysInclude.Remove(Item.ItemName);
+			PropertyHandle->NotifyPostChange(EPropertyChangeType::ArrayRemove);
+			PropertyHandle->NotifyFinishedChangingProperties();
+
+			RegenerateChildContent();
+		}
+		/** END FBlueprintManagedListDetails interface */
+
+		void OnNamespaceSelected(const FString& InNamespace)
+		{
+			FScopedTransaction Transaction(LOCTEXT("AddGlobalImport_Transaction", "Add Global Import"));
+
+			UBlueprintEditorProjectSettings* BlueprintEditorProjectSettings = GetMutableDefault<UBlueprintEditorProjectSettings>();
+			check(BlueprintEditorProjectSettings);
+
+			PropertyHandle->NotifyPreChange();
+			BlueprintEditorProjectSettings->NamespacesToAlwaysInclude.AddUnique(InNamespace);
+			PropertyHandle->NotifyPostChange(EPropertyChangeType::ArrayAdd);
+			PropertyHandle->NotifyFinishedChangingProperties();
+
+			RegenerateChildContent();
+		}
+
+		void OnFilterNamespaceList(TArray<FString>& InOutNamespaceList)
+		{
+			for (const FString& GlobalNamespace : GetDefault<UBlueprintEditorProjectSettings>()->NamespacesToAlwaysInclude)
+			{
+				InOutNamespaceList.RemoveSwap(GlobalNamespace);
+			}
+		}
+
+	private:
+		TSharedRef<IPropertyHandle> PropertyHandle;
+	};
 }
 
-FBlueprintEditorProjectSettingsCustomization::~FBlueprintEditorProjectSettingsCustomization()
+TSharedRef<IDetailCustomization> FBlueprintEditorProjectSettingsCustomization::MakeInstance()
 {
+	return MakeShared<FBlueprintEditorProjectSettingsCustomization>();
 }
 
 void FBlueprintEditorProjectSettingsCustomization::CustomizeDetails(IDetailLayoutBuilder& LayoutBuilder)
 {
-	// Add a custom edit condition on the project-specific 'NamespacesToAlwaysInclude' member to link it to the editor's namespace feature toggle flag (for consistency w/ the editor-specific set).
 	static FName PropertyName_NamespacesToAlwaysInclude = GET_MEMBER_NAME_CHECKED(UBlueprintEditorProjectSettings, NamespacesToAlwaysInclude);
 	TSharedRef<IPropertyHandle> PropertyHandle_NamespacesToAlwaysInclude = LayoutBuilder.GetProperty(PropertyName_NamespacesToAlwaysInclude);
-	if (IDetailPropertyRow* PropertyRow_NamespacesToAlwaysInclude = LayoutBuilder.EditDefaultProperty(PropertyHandle_NamespacesToAlwaysInclude))
-	{
-		PropertyRow_NamespacesToAlwaysInclude->EditCondition(
-			TAttribute<bool>::CreateLambda([]()
-			{
-				return GetDefault<UBlueprintEditorSettings>()->bEnableNamespaceEditorFeatures;
-			}),
-			FOnBooleanValueChanged()
-		);
-	}
+
+	PropertyHandle_NamespacesToAlwaysInclude->MarkHiddenByCustomization();
+
+	IDetailCategoryBuilder& CategoryBuilder = LayoutBuilder.EditCategory(PropertyHandle_NamespacesToAlwaysInclude->GetDefaultCategoryName());
+	CategoryBuilder.AddCustomBuilder(MakeShared<UE::Editor::BlueprintEditorProjectSettingsCustomization::Private::FBlueprintGlobalProjectImportsLayout>(PropertyHandle_NamespacesToAlwaysInclude));
 }
 
 #undef LOCTEXT_NAMESPACE
diff --git a/Engine/Source/Editor/UnrealEd/Private/Settings/BlueprintEditorProjectSettingsCustomization.h b/Engine/Source/Editor/UnrealEd/Private/Settings/BlueprintEditorProjectSettingsCustomization.h
--- a/Engine/Source/Editor/UnrealEd/Private/Settings/BlueprintEditorProjectSettingsCustomization.h
+++ b/Engine/Source/Editor/UnrealEd/Private/Settings/BlueprintEditorProjectSettingsCustomization.h
@@ -14,9 +14,6 @@ public:
 	// Factory creation.
 	static TSharedRef<IDetailCustomization> MakeInstance();
 
-	// Virtual destructor.
-	virtual ~FBlueprintEditorProjectSettingsCustomization();
-
 	// IDetailCustomization interface
 	virtual void CustomizeDetails(IDetailLayoutBuilder& LayoutBuilder) override;
 	// End IDetailCustomization interface
diff --git a/Engine/Source/Editor/UnrealEd/Private/Settings/BlueprintEditorSettingsCustomization.cpp b/Engine/Source/Editor/UnrealEd/Private/Settings/BlueprintEditorSettingsCustomization.cpp
new file mode 100644
--- /dev/null
+++ b/Engine/Source/Editor/UnrealEd/Private/Settings/BlueprintEditorSettingsCustomization.cpp
@@ -0,0 +1,126 @@
+// Copyright Epic Games, Inc. All Rights Reserved.
+
+#include "BlueprintEditorSettingsCustomization.h"
+#include "BlueprintEditorSettings.h"
+#include "PropertyHandle.h"
+#include "DetailLayoutBuilder.h"
+#include "IDetailPropertyRow.h"
+#include "DetailWidgetRow.h"
+#include "PropertyCustomizationHelpers.h"
+#include "BlueprintManagedListDetails.h"
+#include "SBlueprintNamespaceEntry.h"
+#include "ScopedTransaction.h"
+
+#define LOCTEXT_NAMESPACE "FBlueprintEditorSettingsCustomization"
+
+namespace UE::Editor::BlueprintEditorSettingsCustomization::Private
+{
+	class FBlueprintGlobalEditorImportsLayout : public FBlueprintManagedListDetails, public TSharedFromThis<FBlueprintGlobalEditorImportsLayout>
+	{
+	public:
+		FBlueprintGlobalEditorImportsLayout(TSharedRef<IPropertyHandle> InPropertyHandle)
+			: FBlueprintManagedListDetails()
+			, PropertyHandle(InPropertyHandle)
+		{
+			DisplayOptions.TitleText = PropertyHandle->GetPropertyDisplayName();
+			DisplayOptions.TitleTooltipText = PropertyHandle->GetToolTipText();
+
+			DisplayOptions.NoItemsLabelText = LOCTEXT("NoGlobalImports", "None");
+
+			// Add an edit condition to link it to the namespace feature toggle flag.
+			DisplayOptions.EditCondition = TAttribute<bool>::CreateLambda([]()
+			{
+				return GetDefault<UBlueprintEditorSettings>()->bEnableNamespaceEditorFeatures;
+			});
+		}
+
+	protected:
+		/** FBlueprintManagedListDetails interface*/
+		virtual TSharedPtr<SWidget> MakeAddItemWidget() override
+		{
+			return SNew(SBlueprintNamespaceEntry)
+				.AllowTextEntry(false)
+				.OnNamespaceSelected(this, &FBlueprintGlobalEditorImportsLayout::OnNamespaceSelected)
+				.OnFilterNamespaceList(this, &FBlueprintGlobalEditorImportsLayout::OnFilterNamespaceList)
+				.ButtonContent()
+				[
+					SNew(STextBlock)
+					.Text(LOCTEXT("BlueprintAddGlobalImportButton", "Add"))
+					.ToolTipText(LOCTEXT("BlueprintAddGlobalImportButton_Tooltip", "Choose a namespace that Blueprint editors should always import by default (applies only to you as the current local user)."))
+				];
+		}
+
+		virtual void GetManagedListItems(TArray<FManagedListItem>& OutListItems) const override
+		{
+			for (const FString& GlobalNamespace : GetDefault<UBlueprintEditorSettings>()->NamespacesToAlwaysInclude)
+			{
+				FManagedListItem ItemDesc;
+				ItemDesc.ItemName = GlobalNamespace;
+				ItemDesc.DisplayName = FText::FromString(GlobalNamespace);
+				ItemDesc.bIsRemovable = true;
+
+				OutListItems.Add(MoveTemp(ItemDesc));
+			}
+		}
+
+		virtual void OnRemoveItem(const FManagedListItem& Item)
+		{
+			FScopedTransaction Transaction(LOCTEXT("RemoveGlobalImport_Transaction", "Remove Global Import"));
+
+			UBlueprintEditorSettings* BlueprintEditorSettings = GetMutableDefault<UBlueprintEditorSettings>();
+			check(BlueprintEditorSettings);
+
+			PropertyHandle->NotifyPreChange();
+			BlueprintEditorSettings->NamespacesToAlwaysInclude.Remove(Item.ItemName);
+			PropertyHandle->NotifyPostChange(EPropertyChangeType::ArrayRemove);
+			PropertyHandle->NotifyFinishedChangingProperties();
+
+			RegenerateChildContent();
+		}
+		/** END FBlueprintManagedListDetails interface */
+
+		void OnNamespaceSelected(const FString& InNamespace)
+		{
+			FScopedTransaction Transaction(LOCTEXT("AddGlobalImport_Transaction", "Add Global Import"));
+
+			UBlueprintEditorSettings* BlueprintEditorSettings = GetMutableDefault<UBlueprintEditorSettings>();
+			check(BlueprintEditorSettings);
+
+			PropertyHandle->NotifyPreChange();
+			BlueprintEditorSettings->NamespacesToAlwaysInclude.AddUnique(InNamespace);
+			PropertyHandle->NotifyPostChange(EPropertyChangeType::ArrayAdd);
+			PropertyHandle->NotifyFinishedChangingProperties();
+
+			RegenerateChildContent();
+		}
+
+		void OnFilterNamespaceList(TArray<FString>& InOutNamespaceList)
+		{
+			for (const FString& GlobalNamespace : GetDefault<UBlueprintEditorSettings>()->NamespacesToAlwaysInclude)
+			{
+				InOutNamespaceList.RemoveSwap(GlobalNamespace);
+			}
+		}
+
+	private:
+		TSharedRef<IPropertyHandle> PropertyHandle;
+	};
+}
+
+TSharedRef<IDetailCustomization> FBlueprintEditorSettingsCustomization::MakeInstance()
+{
+	return MakeShared<FBlueprintEditorSettingsCustomization>();
+}
+
+void FBlueprintEditorSettingsCustomization::CustomizeDetails(IDetailLayoutBuilder& LayoutBuilder)
+{
+	static FName PropertyName_NamespacesToAlwaysInclude = GET_MEMBER_NAME_CHECKED(UBlueprintEditorSettings, NamespacesToAlwaysInclude);
+	TSharedRef<IPropertyHandle> PropertyHandle_NamespacesToAlwaysInclude = LayoutBuilder.GetProperty(PropertyName_NamespacesToAlwaysInclude);
+
+	PropertyHandle_NamespacesToAlwaysInclude->MarkHiddenByCustomization();
+
+	IDetailCategoryBuilder& CategoryBuilder = LayoutBuilder.EditCategory(PropertyHandle_NamespacesToAlwaysInclude->GetDefaultCategoryName());
+	CategoryBuilder.AddCustomBuilder(MakeShared<UE::Editor::BlueprintEditorSettingsCustomization::Private::FBlueprintGlobalEditorImportsLayout>(PropertyHandle_NamespacesToAlwaysInclude));
+}
+
+#undef LOCTEXT_NAMESPACE
diff --git a/Engine/Source/Editor/UnrealEd/Private/Settings/BlueprintEditorSettingsCustomization.h b/Engine/Source/Editor/UnrealEd/Private/Settings/BlueprintEditorSettingsCustomization.h
new file mode 100644
--- /dev/null
+++ b/Engine/Source/Editor/UnrealEd/Private/Settings/BlueprintEditorSettingsCustomization.h
@@ -0,0 +1,20 @@
+// Copyright Epic Games, Inc. All Rights Reserved.
+
+#pragma once
+
+#include "CoreMinimal.h"
+#include "IDetailCustomization.h"
+
+/**
+ * Implements a details customization for UBlueprintEditorSettings.
+ */
+class FBlueprintEditorSettingsCustomization : public IDetailCustomization
+{
+public:
+	// Factory creation.
+	static TSharedRef<IDetailCustomization> MakeInstance();
+
+	// IDetailCustomization interface
+	virtual void CustomizeDetails(IDetailLayoutBuilder& LayoutBuilder) override;
+	// End IDetailCustomization interface
+};
\ No newline at end of file
diff --git a/Engine/Source/Editor/UnrealEd/Private/UnrealEdEngine.cpp b/Engine/Source/Editor/UnrealEd/Private/UnrealEdEngine.cpp
--- a/Engine/Source/Editor/UnrealEd/Private/UnrealEdEngine.cpp
+++ b/Engine/Source/Editor/UnrealEd/Private/UnrealEdEngine.cpp
@@ -45,6 +45,7 @@
 #include "LevelEditor.h"
 #include "InstancedStaticMeshDelegates.h"
 #include "Interfaces/IMainFrameModule.h"
+#include "Settings/BlueprintEditorSettingsCustomization.h"
 #include "Settings/BlueprintEditorProjectSettingsCustomization.h"
 #include "Settings/EditorLoadingSavingSettingsCustomization.h"
 #include "Settings/GameMapsSettingsCustomization.h"
@@ -177,6 +178,7 @@ void UUnrealEdEngine::Init(IEngineLoop* InEngineLoop)
 	{
 		FPropertyEditorModule& PropertyModule = FModuleManager::LoadModuleChecked<FPropertyEditorModule>("PropertyEditor");
 
+		PropertyModule.RegisterCustomClassLayout("BlueprintEditorSettings", FOnGetDetailCustomizationInstance::CreateStatic(&FBlueprintEditorSettingsCustomization::MakeInstance));
 		PropertyModule.RegisterCustomClassLayout("BlueprintEditorProjectSettings", FOnGetDetailCustomizationInstance::CreateStatic(&FBlueprintEditorProjectSettingsCustomization::MakeInstance));
 		PropertyModule.RegisterCustomClassLayout("EditorLoadingSavingSettings", FOnGetDetailCustomizationInstance::CreateStatic(&FEditorLoadingSavingSettingsCustomization::MakeInstance));
 		PropertyModule.RegisterCustomClassLayout("GameMapsSettings", FOnGetDetailCustomizationInstance::CreateStatic(&FGameMapsSettingsCustomization::MakeInstance));
diff --git a/Engine/Source/Editor/UnrealEd/Public/Settings/BlueprintEditorProjectSettings.h b/Engine/Source/Editor/UnrealEd/Public/Settings/BlueprintEditorProjectSettings.h
--- a/Engine/Source/Editor/UnrealEd/Public/Settings/BlueprintEditorProjectSettings.h
+++ b/Engine/Source/Editor/UnrealEd/Public/Settings/BlueprintEditorProjectSettings.h
@@ -40,8 +40,8 @@ public:
 	UPROPERTY(EditAnywhere, config, Category = Experimental, meta = (EditCondition = "bEnableChildActorExpansionInTreeView"))
 	EChildActorComponentTreeViewVisualizationMode DefaultChildActorTreeViewMode;
 
-	// The list of namespaces to always expose in any Blueprint (for all users of the game/project). Requires Blueprint namespace features to be enabled in editor preferences.
-	UPROPERTY(EditAnywhere, config, Category = Experimental)
+	// A list of namespace identifiers that all Blueprint assets in the project should import by default. Requires Blueprint namespace features to be enabled in editor preferences. Editing this list will also cause any visible Blueprint editor windows to be closed.
+	UPROPERTY(EditAnywhere, config, Category = Experimental, DisplayName = "Global Namespace Imports (Shared)")
 	TArray<FString> NamespacesToAlwaysInclude;
 
 	/** 
