Add missing copyright boilerplate.

diff --git a/Engine/Source/Editor/Kismet/Private/BlueprintManagedListDetails.cpp b/Engine/Source/Editor/Kismet/Private/BlueprintManagedListDetails.cpp
--- a/Engine/Source/Editor/Kismet/Private/BlueprintManagedListDetails.cpp
+++ b/Engine/Source/Editor/Kismet/Private/BlueprintManagedListDetails.cpp
@@ -1,3 +1,5 @@
+// Copyright Epic Games, Inc. All Rights Reserved.
+
 #include "BlueprintManagedListDetails.h"
 #include "DetailLayoutBuilder.h"
 #include "DetailWidgetRow.h"
