Update the FBlueprintEditorModule::RegisterVariableCustomization to support multiple entry.

diff --git a/Engine/Plugins/Animation/ControlRig/Source/ControlRigEditor/Private/ControlRigEditorModule.cpp b/Engine/Plugins/Animation/ControlRig/Source/ControlRigEditor/Private/ControlRigEditorModule.cpp
--- a/Engine/Plugins/Animation/ControlRig/Source/ControlRigEditor/Private/ControlRigEditorModule.cpp
+++ b/Engine/Plugins/Animation/ControlRig/Source/ControlRigEditor/Private/ControlRigEditorModule.cpp
@@ -142,7 +142,7 @@ void FControlRigEditorModule::StartupModule()
 
 	// Register Blueprint editor variable customization
 	FBlueprintEditorModule& BlueprintEditorModule = FModuleManager::LoadModuleChecked<FBlueprintEditorModule>("Kismet");
-	BlueprintEditorModule.RegisterVariableCustomization(FProperty::StaticClass(), FOnGetVariableCustomizationInstance::CreateStatic(&FControlRigVariableDetailsCustomization::MakeInstance));
+	BlueprintVariableCustomizationHandle = BlueprintEditorModule.RegisterVariableCustomization(FProperty::StaticClass(), FOnGetVariableCustomizationInstance::CreateStatic(&FControlRigVariableDetailsCustomization::MakeInstance));
 	BlueprintEditorModule.RegisterGraphCustomization(GetDefault<UControlRigGraphSchema>(), FOnGetGraphCustomizationInstance::CreateStatic(&FControlRigGraphDetails::MakeInstance));
 
 	// Register to fixup newly created BPs
@@ -294,8 +294,7 @@ void FControlRigEditorModule::ShutdownModule()
 		FBlueprintEditorModule* BlueprintEditorModule = FModuleManager::GetModulePtr<FBlueprintEditorModule>("Kismet");
 		if (BlueprintEditorModule)
 		{
-			BlueprintEditorModule->UnregisterVariableCustomization(FProperty::StaticClass());
-			BlueprintEditorModule->UnregisterLocalVariableCustomization(FProperty::StaticClass());
+			BlueprintEditorModule->UnregisterVariableCustomization(FProperty::StaticClass(), BlueprintVariableCustomizationHandle);
 			BlueprintEditorModule->UnregisterGraphCustomization(GetDefault<UControlRigGraphSchema>());
 		}
 	}
diff --git a/Engine/Plugins/Animation/ControlRig/Source/ControlRigEditor/Private/ControlRigEditorModule.h b/Engine/Plugins/Animation/ControlRig/Source/ControlRigEditor/Private/ControlRigEditorModule.h
--- a/Engine/Plugins/Animation/ControlRig/Source/ControlRigEditor/Private/ControlRigEditorModule.h
+++ b/Engine/Plugins/Animation/ControlRig/Source/ControlRigEditor/Private/ControlRigEditorModule.h
@@ -86,6 +86,7 @@ private:
 	/** Delegate handles for blueprint utils */
 	FDelegateHandle RefreshAllNodesDelegateHandle;
 	FDelegateHandle ReconstructAllNodesDelegateHandle;
+	FDelegateHandle BlueprintVariableCustomizationHandle;
 
 	/** Animation Editor Handle*/
 	FDelegateHandle AnimationEditorExtenderHandle;
diff --git a/Engine/Source/Editor/Kismet/Private/BlueprintEditorModule.cpp b/Engine/Source/Editor/Kismet/Private/BlueprintEditorModule.cpp
--- a/Engine/Source/Editor/Kismet/Private/BlueprintEditorModule.cpp
+++ b/Engine/Source/Editor/Kismet/Private/BlueprintEditorModule.cpp
@@ -365,9 +365,11 @@ void FBlueprintEditorModule::UnregisterSCSEditorCustomization(const FName& InCom
 	SCSEditorCustomizations.Remove(InComponentName);
 }
 
-void FBlueprintEditorModule::RegisterVariableCustomization(FFieldClass* InFieldClass, FOnGetVariableCustomizationInstance InOnGetVariableCustomization)
+FDelegateHandle FBlueprintEditorModule::RegisterVariableCustomization(FFieldClass* InFieldClass, FOnGetVariableCustomizationInstance InOnGetVariableCustomization)
 {
+	FDelegateHandle Result = InOnGetVariableCustomization.GetHandle();
 	VariableCustomizations.Add(InFieldClass, InOnGetVariableCustomization);
+	return Result;
 }
 
 void FBlueprintEditorModule::UnregisterVariableCustomization(FFieldClass* InFieldClass)
@@ -375,9 +377,22 @@ void FBlueprintEditorModule::UnregisterVariableCustomization(FFieldClass* InFiel
 	VariableCustomizations.Remove(InFieldClass);
 }
 
-void FBlueprintEditorModule::RegisterLocalVariableCustomization(FFieldClass* InFieldClass, FOnGetLocalVariableCustomizationInstance InOnGetLocalVariableCustomization)
+void FBlueprintEditorModule::UnregisterVariableCustomization(FFieldClass* InFieldClass, FDelegateHandle InHandle)
 {
+	for (TMultiMap<FFieldClass*, FOnGetVariableCustomizationInstance>::TKeyIterator It = VariableCustomizations.CreateKeyIterator(InFieldClass); It; ++It)
+	{
+		if (It.Value().GetHandle() == InHandle)
+		{
+			It.RemoveCurrent();
+		}
+	}
+}
+
+FDelegateHandle FBlueprintEditorModule::RegisterLocalVariableCustomization(FFieldClass* InFieldClass, FOnGetLocalVariableCustomizationInstance InOnGetLocalVariableCustomization)
+{
+	FDelegateHandle Result = InOnGetLocalVariableCustomization.GetHandle();
 	LocalVariableCustomizations.Add(InFieldClass, InOnGetLocalVariableCustomization);
+	return Result;
 }
 
 void FBlueprintEditorModule::UnregisterLocalVariableCustomization(FFieldClass* InFieldClass)
@@ -385,6 +400,17 @@ void FBlueprintEditorModule::UnregisterLocalVariableCustomization(FFieldClass* I
 	LocalVariableCustomizations.Remove(InFieldClass);
 }
 
+void FBlueprintEditorModule::UnregisterLocalVariableCustomization(FFieldClass* InFieldClass, FDelegateHandle InHandle)
+{
+	for (TMultiMap<FFieldClass*, FOnGetVariableCustomizationInstance>::TKeyIterator It = LocalVariableCustomizations.CreateKeyIterator(InFieldClass); It; ++It)
+	{
+		if (It.Value().GetHandle() == InHandle)
+		{
+			It.RemoveCurrent();
+		}
+	}
+}
+
 void FBlueprintEditorModule::RegisterGraphCustomization(const UEdGraphSchema* InGraphSchema, FOnGetGraphCustomizationInstance InOnGetGraphCustomization)
 {
 	GraphCustomizations.Add(InGraphSchema, InOnGetGraphCustomization);
@@ -412,13 +438,17 @@ TArray<TSharedPtr<IDetailCustomization>> FBlueprintEditorModule::CustomizeVariab
 
 		for (FFieldClass* ClassToQuery : ParentClassesToQuery)
 		{
-			FOnGetVariableCustomizationInstance* CustomizationDelegate = VariableCustomizations.Find(ClassToQuery);
-			if (CustomizationDelegate && CustomizationDelegate->IsBound())
+			TArray<FOnGetVariableCustomizationInstance*, TInlineAllocator<4>> CustomizationDelegates;
+			VariableCustomizations.MultiFindPointer(ClassToQuery, CustomizationDelegates, false);
+			for (FOnGetVariableCustomizationInstance* CustomizationDelegate : CustomizationDelegates)
 			{
-				TSharedPtr<IDetailCustomization> Customization = CustomizationDelegate->Execute(InBlueprintEditor);
-				if(Customization.IsValid())
-				{ 
-					DetailsCustomizations.Add(Customization);
+				if (CustomizationDelegate && CustomizationDelegate->IsBound())
+				{
+					TSharedPtr<IDetailCustomization> Customization = CustomizationDelegate->Execute(InBlueprintEditor);
+					if (Customization.IsValid())
+					{
+						DetailsCustomizations.Add(Customization);
+					}
 				}
 			}
 		}
diff --git a/Engine/Source/Editor/Kismet/Public/BlueprintEditorModule.h b/Engine/Source/Editor/Kismet/Public/BlueprintEditorModule.h
--- a/Engine/Source/Editor/Kismet/Public/BlueprintEditorModule.h
+++ b/Engine/Source/Editor/Kismet/Public/BlueprintEditorModule.h
@@ -209,27 +209,43 @@ public:
 	 * @param	InStruct				The type of the variable to create the customization for
 	 * @param	InOnGetDetailCustomization	The delegate used to create customization instances
 	 */
-	virtual void RegisterVariableCustomization(FFieldClass* InFieldClass, FOnGetVariableCustomizationInstance InOnGetVariableCustomization);
+	virtual FDelegateHandle RegisterVariableCustomization(FFieldClass* InFieldClass, FOnGetVariableCustomizationInstance InOnGetVariableCustomization);
 
 	/** 
 	 * Unregister a previously registered customization for BP variables
 	 * @param	InStruct				The type to create the customization for
 	 */
+	UE_DEPRECATED(5.1, "UnregisterVariableCustomization without a delegate handle is deprecated.")
 	virtual void UnregisterVariableCustomization(FFieldClass* InFieldClass);
 
+	/**
+	 * Unregister a previously registered customization for BP variables
+	 * @param	InStruct				The type to create the customization for
+	 * @param	InHandle				The handle returned by RegisterVariableCustomization
+	 */
+	virtual void UnregisterVariableCustomization(FFieldClass* InFieldClass, FDelegateHandle InHandle);
+
 	/** 
-	* Register a customization for for Blueprint local variables
-	* @param	InFieldClass				The type of the variable to create the customization for
-	* @param	InOnGetLocalVariableCustomization	The delegate used to create customization instances
-	*/
-	virtual void RegisterLocalVariableCustomization(FFieldClass* InFieldClass, FOnGetLocalVariableCustomizationInstance InOnGetLocalVariableCustomization);
+	 * Register a customization for for Blueprint local variables
+	 * @param	InFieldClass				The type of the variable to create the customization for
+	 * @param	InOnGetLocalVariableCustomization	The delegate used to create customization instances
+	 */
+	virtual FDelegateHandle RegisterLocalVariableCustomization(FFieldClass* InFieldClass, FOnGetLocalVariableCustomizationInstance InOnGetLocalVariableCustomization);
 
 	/** 
 	* Unregister a previously registered customization for BP local variables
 	* @param	InFieldClass				The type to create the customization for
-	*/
+	 */
+	UE_DEPRECATED(5.1, "UnregisterLocalVariableCustomization without a delegate handle is deprecated.")
 	virtual void UnregisterLocalVariableCustomization(FFieldClass* InFieldClass);
 
+	/** 
+	 * Unregister a previously registered customization for BP local variables
+	 * @param	InFieldClass				The type to create the customization for
+	 * @param	InHandle				The handle returned by RegisterLocalVariableCustomization
+	 */
+	virtual void UnregisterLocalVariableCustomization(FFieldClass* InFieldClass, FDelegateHandle InHandle);
+
 	/** 
 	 * Register a customization for for Blueprint graphs
 	 * @param	InGraphSchema				The schema of the graph to create the customization for
@@ -296,10 +312,10 @@ private:
 	TMap<FName, FSCSEditorCustomizationBuilder> SCSEditorCustomizations;
 
 	/** Customizations for Blueprint variables */
-	TMap<FFieldClass*, FOnGetVariableCustomizationInstance> VariableCustomizations;
+	TMultiMap<FFieldClass*, FOnGetVariableCustomizationInstance> VariableCustomizations;
 
 	/** Customizations for Blueprint local variables */
-	TMap<FFieldClass*, FOnGetLocalVariableCustomizationInstance> LocalVariableCustomizations;
+	TMultiMap<FFieldClass*, FOnGetLocalVariableCustomizationInstance> LocalVariableCustomizations;
 
 	/** Customizations for Blueprint graphs */
 	TMap<const UEdGraphSchema*, FOnGetGraphCustomizationInstance> GraphCustomizations;
