Add customization for Blueprint Graph functions.
#prefligh 62312d4fac0173aa4125e0ea

diff --git a/Engine/Source/Editor/Kismet/Private/BlueprintDetailsCustomization.cpp b/Engine/Source/Editor/Kismet/Private/BlueprintDetailsCustomization.cpp
--- a/Engine/Source/Editor/Kismet/Private/BlueprintDetailsCustomization.cpp
+++ b/Engine/Source/Editor/Kismet/Private/BlueprintDetailsCustomization.cpp
@@ -4051,6 +4051,19 @@ void FBlueprintGraphActionDetails::CustomizeDetails( IDetailLayoutBuilder& Detai
 				];
 			OutputsCategory.HeaderContent(OutputsHeaderContentWidget);
 		}
+
+		// See if anything else wants to customize our details
+		TWeakPtr<FBlueprintEditor> BlueprintEditor = MyBlueprint.Pin()->GetBlueprintEditor();
+		FBlueprintEditorModule& BlueprintEditorModule = FModuleManager::GetModuleChecked<FBlueprintEditorModule>("Kismet");
+		TArray<TSharedPtr<IDetailCustomization>> Customizations = BlueprintEditorModule.CustomizeFunction(FunctionEntryNode->GetClass(), BlueprintEditor.Pin());
+		ExternalDetailCustomizations.Append(Customizations);
+		if (ExternalDetailCustomizations.Num() > 0)
+		{
+			for (TSharedPtr<IDetailCustomization> ExternalDetailCustomization : ExternalDetailCustomizations)
+			{
+				ExternalDetailCustomization->CustomizeDetails(DetailLayout);
+			}
+		}
 	}
 	else if (bHasAGraph)
 	{
diff --git a/Engine/Source/Editor/Kismet/Private/BlueprintEditorModule.cpp b/Engine/Source/Editor/Kismet/Private/BlueprintEditorModule.cpp
--- a/Engine/Source/Editor/Kismet/Private/BlueprintEditorModule.cpp
+++ b/Engine/Source/Editor/Kismet/Private/BlueprintEditorModule.cpp
@@ -421,6 +421,24 @@ void FBlueprintEditorModule::UnregisterGraphCustomization(const UEdGraphSchema*
 	GraphCustomizations.Remove(InGraphSchema);
 }
 
+FDelegateHandle FBlueprintEditorModule::RegisterFunctionCustomization(TSubclassOf<UK2Node_EditablePinBase> InFieldClass, FOnGetFunctionCustomizationInstance InOnGetFunctionCustomization)
+{
+	FDelegateHandle Result = InOnGetFunctionCustomization.GetHandle();
+	FunctionCustomizations.Add(InFieldClass, InOnGetFunctionCustomization);
+	return Result;
+}
+
+void FBlueprintEditorModule::UnregisterFunctionCustomization(TSubclassOf<UK2Node_EditablePinBase> InFieldClass, FDelegateHandle InHandle)
+{
+	for (TMultiMap<TSubclassOf<UK2Node_EditablePinBase>, FOnGetFunctionCustomizationInstance>::TKeyIterator It = FunctionCustomizations.CreateKeyIterator(InFieldClass); It; ++It)
+	{
+		if (It.Value().GetHandle() == InHandle)
+		{
+			It.RemoveCurrent();
+		}
+	}
+}
+
 TArray<TSharedPtr<IDetailCustomization>> FBlueprintEditorModule::CustomizeVariable(FFieldClass* InFieldClass, TSharedPtr<IBlueprintEditor> InBlueprintEditor)
 {
 	TArray<TSharedPtr<IDetailCustomization>> DetailsCustomizations;
@@ -491,6 +509,43 @@ TArray<TSharedPtr<IDetailCustomization>> FBlueprintEditorModule::CustomizeGraph(
 	return DetailsCustomizations;
 }
 
+TArray<TSharedPtr<IDetailCustomization>> FBlueprintEditorModule::CustomizeFunction(const TSubclassOf<UK2Node_EditablePinBase> InFunctionClass, TSharedPtr<IBlueprintEditor> InBlueprintEditor)
+{
+	TArray<TSharedPtr<IDetailCustomization>> DetailsCustomizations;
+	TArray<UClass*> ParentPinClassesToQuery;
+	if (InFunctionClass)
+	{
+		UClass* FunctionClass = InFunctionClass.Get();
+		ParentPinClassesToQuery.Add(FunctionClass);
+
+		UClass* ParentSchemaClass = FunctionClass->GetSuperClass();
+		while (ParentSchemaClass && ParentSchemaClass->IsChildOf(UK2Node_EditablePinBase::StaticClass()))
+		{
+			ParentPinClassesToQuery.Add(ParentSchemaClass);
+			ParentSchemaClass = ParentSchemaClass->GetSuperClass();
+		}
+
+		for (UClass* ClassToQuery : ParentPinClassesToQuery)
+		{
+			TArray<FOnGetFunctionCustomizationInstance*, TInlineAllocator<4>> CustomizationDelegates;
+			FunctionCustomizations.MultiFindPointer(ClassToQuery, CustomizationDelegates, false);
+			for (FOnGetFunctionCustomizationInstance* CustomizationDelegate : CustomizationDelegates)
+			{
+				if (CustomizationDelegate && CustomizationDelegate->IsBound())
+				{
+					TSharedPtr<IDetailCustomization> Customization = CustomizationDelegate->Execute(InBlueprintEditor);
+					if (Customization.IsValid())
+					{
+						DetailsCustomizations.Add(Customization);
+					}
+				}
+			}
+		}
+	}
+
+	return DetailsCustomizations;
+}
+
 void FBlueprintEditorModule::PrepareAutoGeneratedDefaultEvents()
 {
 	// Load up all default events that should be spawned for Blueprints that are children of specific classes
diff --git a/Engine/Source/Editor/Kismet/Public/BlueprintEditorModule.h b/Engine/Source/Editor/Kismet/Public/BlueprintEditorModule.h
--- a/Engine/Source/Editor/Kismet/Public/BlueprintEditorModule.h
+++ b/Engine/Source/Editor/Kismet/Public/BlueprintEditorModule.h
@@ -22,6 +22,7 @@ class IDetailCustomization;
 class FKismetCompilerContext;
 struct FBlueprintDebugger;
 class FSubobjectEditorTreeNode;
+class UK2Node_EditablePinBase;
 
 /** Delegate used to customize variable display */
 DECLARE_DELEGATE_RetVal_OneParam(TSharedPtr<IDetailCustomization>, FOnGetVariableCustomizationInstance, TSharedPtr<IBlueprintEditor> /*BlueprintEditor*/);
@@ -29,6 +30,9 @@ DECLARE_DELEGATE_RetVal_OneParam(TSharedPtr<IDetailCustomization>, FOnGetVariabl
 /** Delegate used to customize local variable display */
 DECLARE_DELEGATE_RetVal_OneParam(TSharedPtr<IDetailCustomization>, FOnGetLocalVariableCustomizationInstance, TSharedPtr<IBlueprintEditor> /*BlueprintEditor*/);
 
+/** Delegate used to customize function display */
+DECLARE_DELEGATE_RetVal_OneParam(TSharedPtr<IDetailCustomization>, FOnGetFunctionCustomizationInstance, TSharedPtr<IBlueprintEditor> /*BlueprintEditor*/);
+
 /** Delegate used to customize graph display */
 DECLARE_DELEGATE_RetVal_OneParam(TSharedPtr<IDetailCustomization>, FOnGetGraphCustomizationInstance, TSharedPtr<IBlueprintEditor> /*BlueprintEditor*/);
 
@@ -259,6 +263,20 @@ public:
 	 */
 	virtual void UnregisterGraphCustomization(const UEdGraphSchema* InGraphSchema);
 
+	/**
+	 * Register a customization for for Blueprint functions
+	 * @param	InStruct				The type of the pin to create the customization for
+	 * @param	InOnGetFunctionCustomization	The delegate used to create customization instances
+	 */
+	virtual FDelegateHandle RegisterFunctionCustomization(TSubclassOf<UK2Node_EditablePinBase> InFieldClass, FOnGetFunctionCustomizationInstance InOnGetFunctionCustomization);
+
+	/**
+	 * Unregister a previously registered customization for BP functions
+	 * @param	InStruct				The type to create the customization for
+	 * @param	InHandle				The handle returned by UnregisterFunctionCustomization
+	 */
+	virtual void UnregisterFunctionCustomization(TSubclassOf<UK2Node_EditablePinBase> InFieldClass, FDelegateHandle InHandle);
+
 
 	/** 
 	 * Build a set of details customizations for the passed-in type, if possible.
@@ -274,6 +292,13 @@ public:
 	 */
 	virtual TArray<TSharedPtr<IDetailCustomization>> CustomizeGraph(const UEdGraphSchema* InGraphSchema, TSharedPtr<IBlueprintEditor> InBlueprintEditor);
 
+	/** 
+	 * Build a set of details customizations for function with the passed-in type, if possible.
+	 * @param	InFunctionClass		The type to create the customization for
+	 * @param	InBlueprintEditor		The Blueprint Editor the customization will be created for
+	 */
+	virtual TArray<TSharedPtr<IDetailCustomization>> CustomizeFunction(TSubclassOf<UK2Node_EditablePinBase> InFunctionClass, TSharedPtr<IBlueprintEditor> InBlueprintEditor);
+
 	/** Delegate for binding functions to be called when the blueprint editor finishes getting created */
 	DECLARE_EVENT_OneParam( FBlueprintEditorModule, FBlueprintEditorOpenedEvent, EBlueprintType );
 	FBlueprintEditorOpenedEvent& OnBlueprintEditorOpened() { return BlueprintEditorOpened; }
@@ -320,6 +345,9 @@ private:
 	/** Customizations for Blueprint graphs */
 	TMap<const UEdGraphSchema*, FOnGetGraphCustomizationInstance> GraphCustomizations;
 
+	/** Customizations for Blueprint functions */
+	TMultiMap<TSubclassOf<UK2Node_EditablePinBase>, FOnGetFunctionCustomizationInstance> FunctionCustomizations;
+
 	/** Root customization for the BP editor details panel. */
 	TSharedPtr<class IDetailRootObjectCustomization> DetailsRootCustomization;
 
