Update the Blueprint namespace registry on removal/replacement of an assigned namespace.

This may occur during one of the following events:
- Reverting or reassigning a Blueprint asset's namespace.
- Removing a Blueprint asset (i.e. deleting it in the browser).
- Renaming a Blueprint asset (i.e. moving it to another location).
  - Note: this only applies when using package paths as a default namespace identifier (which is currently disabled by default).

Change summary:
- Added FBlueprintNamespaceRegistry::Rebuild().
- Exposed FBlueprintNamespaceUtilities::ConvertPackagePathToNamespacePath().
- Implemented FBlueprintNamespaceRegistry::OnAssetRemoved() and OnAssetRenamed().
- Modified FBlueprintNamespaceRegistry::FindAndRegisterAllNamespaces() with an exclusion set (to assist with asset removal).
- Added FBlueprintGlobalOptionsDetails::HandleNamespaceValueChange() to consolidate shared logic between reverting and reassignment.
- Extended Blueprint namespace member property details customization to refresh the namespace registry and editor context after updating the value.

diff --git a/Engine/Source/Editor/Kismet/Private/BlueprintDetailsCustomization.cpp b/Engine/Source/Editor/Kismet/Private/BlueprintDetailsCustomization.cpp
--- a/Engine/Source/Editor/Kismet/Private/BlueprintDetailsCustomization.cpp
+++ b/Engine/Source/Editor/Kismet/Private/BlueprintDetailsCustomization.cpp
@@ -92,6 +92,7 @@
 #include "IPropertyAccessEditor.h"
 #include "AssetRegistryModule.h"
 #include "Settings/BlueprintEditorProjectSettings.h"
+#include "Misc/ScopedSlowTask.h"
 
 #define LOCTEXT_NAMESPACE "BlueprintDetailsCustomization"
 
@@ -5967,25 +5968,7 @@ void FBlueprintGlobalOptionsDetails::OnNamespaceValueCommitted(const FString& In
 		check(NamespacePropertyHandle.IsValid());
 		NamespacePropertyHandle->SetValue(InNamespace);
 
-		// Ensure that the new path is added into the namespace registry.
-		FBlueprintNamespaceRegistry::Get().RegisterNamespace(InNamespace);
-
-		// Auto-import the namespace into the current editor context.
-		TSharedPtr<FBlueprintEditor> BlueprintEditor = GetBlueprintEditorPtr().Pin();
-		if (BlueprintEditor.IsValid())
-		{
-			// We need to refresh the details view if we flipped from an imported namespace to an assigned
-			// namespace (or vice-versa). In that case, it will switch to/from a default import in the view.
-			if (Blueprint->ImportedNamespaces.Contains(OldNamespace) || Blueprint->ImportedNamespaces.Contains(InNamespace))
-			{
-				BlueprintEditor->RefreshInspector();
-			}
-			else if (GetDefault<UBlueprintEditorSettings>()->bEnableNamespaceImportingFeatures)
-			{
-				// This will also refresh the details view.
-				BlueprintEditor->ImportNamespace(InNamespace);
-			}
-		}
+		HandleNamespaceValueChange(OldNamespace, InNamespace);
 	}
 }
 
@@ -6016,15 +5999,74 @@ void FBlueprintGlobalOptionsDetails::OnNamespaceResetToDefaultValue()
 		NamespaceValueWidget->SetCurrentNamespace(DefaultNamespaceValue);
 	}
 
+	HandleNamespaceValueChange(OriginalValue, DefaultNamespaceValue);
+}
+
+void FBlueprintGlobalOptionsDetails::HandleNamespaceValueChange(const FString& InOldValue, const FString& InNewValue)
+{
+	// Refresh the namespace registry.
+	FBlueprintNamespaceRegistry& BlueprintNamespaceRegistry = FBlueprintNamespaceRegistry::Get();
+	if (!InOldValue.IsEmpty() && BlueprintNamespaceRegistry.IsRegisteredPath(InOldValue))
+	{
+		// @todo_namespaces - This may not scale for larger projects.
+		// Using a slow task for now, but consider optimizing this path.
+		FScopedSlowTask SlowTask(0.0f, LOCTEXT("RebuildingNamespaceRegistry", "Updating the namespace registry..."));
+
+		// The old path is a non-global registered namespace. Revisit all assets and ensure that the registry is up-to-date.
+		// If the old namespace is no longer in use by another Blueprint asset, this effectively removes it from the registry.
+		BlueprintNamespaceRegistry.Rebuild();
+	}
+
+	if (!InNewValue.IsEmpty() && !BlueprintNamespaceRegistry.IsRegisteredPath(InNewValue))
+	{
+		// Add the new namespace into the registry (it has not been seen yet).
+		BlueprintNamespaceRegistry.RegisterNamespace(InNewValue);
+	}
+	
+	// Refresh the Blueprint editor context.
 	TSharedPtr<FBlueprintEditor> BlueprintEditor = GetBlueprintEditorPtr().Pin();
 	if (BlueprintEditor.IsValid())
 	{
-		// We need to refresh the details view if we reassigned away from an already-imported
-		// namespace. In that case, it needs to switch back to a user-added import in the view.
-		const UBlueprint* Blueprint = GetBlueprintObj();
-		if (Blueprint && Blueprint->ImportedNamespaces.Contains(OriginalValue))
+		if (const UBlueprint* Blueprint = BlueprintEditor->GetBlueprintObj())
 		{
-			BlueprintEditor->RefreshInspector();
+			bool bRefreshDetailsView = false;
+			if (Blueprint->ImportedNamespaces.Contains(InOldValue))
+			{
+				// Remove the import from the current editor context if it is no longer registered.
+				if (!BlueprintNamespaceRegistry.IsRegisteredPath(InOldValue))
+				{
+					BlueprintEditor->RemoveNamespace(InOldValue);
+				}
+
+				// We need to refresh the details view if we unassigned and/or removed an imported namespace path.
+				// If unassigned but still imported, it will return to a non-default namespace in the Imports table.
+				bRefreshDetailsView = true;
+			}
+			
+			if (Blueprint->ImportedNamespaces.Contains(InNewValue))
+			{
+				// We need to refresh the details view if we assigned an imported namespace.
+				// In that case, it will switch to a default namespace in the Imports table.
+				bRefreshDetailsView = true;
+			}
+			else if (GetDefault<UBlueprintEditorSettings>()->bEnableNamespaceImportingFeatures)
+			{
+				FBlueprintEditor::FImportNamespaceParameters Params;
+				Params.bIsAutoImport = false;
+				Params.OnImportCallback = FSimpleDelegate::CreateLambda([&bRefreshDetailsView]()
+				{
+					bRefreshDetailsView = true;
+				});
+
+				// Import the new namespace into the current editor context.
+				BlueprintEditor->ImportNamespace(InNewValue, Params);
+			}
+
+			// Refresh the details view if necessary.
+			if (bRefreshDetailsView)
+			{
+				BlueprintEditor->RefreshInspector();
+			}
 		}
 	}
 }
diff --git a/Engine/Source/Editor/Kismet/Private/BlueprintDetailsCustomization.h b/Engine/Source/Editor/Kismet/Private/BlueprintDetailsCustomization.h
--- a/Engine/Source/Editor/Kismet/Private/BlueprintDetailsCustomization.h
+++ b/Engine/Source/Editor/Kismet/Private/BlueprintDetailsCustomization.h
@@ -812,6 +812,9 @@ protected:
 	/** Invoked when the Blueprint namespace "reset to default" button is clicked */
 	void OnNamespaceResetToDefaultValue();
 
+	/** Called when the assigned Blueprint namespace value is changed */
+	void HandleNamespaceValueChange(const FString& InOldValue, const FString& InNewValue);
+
 private:
 	/** Weak reference to the Blueprint editor */
 	TWeakPtr<FBlueprintEditor> BlueprintEditorPtr;
diff --git a/Engine/Source/Editor/Kismet/Private/BlueprintNamespaceRegistry.cpp b/Engine/Source/Editor/Kismet/Private/BlueprintNamespaceRegistry.cpp
--- a/Engine/Source/Editor/Kismet/Private/BlueprintNamespaceRegistry.cpp
+++ b/Engine/Source/Editor/Kismet/Private/BlueprintNamespaceRegistry.cpp
@@ -102,12 +102,43 @@ void FBlueprintNamespaceRegistry::OnAssetAdded(const FAssetData& AssetData)
 
 void FBlueprintNamespaceRegistry::OnAssetRemoved(const FAssetData& AssetData)
 {
-	// @todo_namespaces - Handle Blueprint asset removal.
+	// Remove the asset's current namespace if it's non-empty and was previously registered.
+	FString AssetNamespace = FBlueprintNamespaceUtilities::GetAssetNamespace(AssetData);
+	if (!AssetNamespace.IsEmpty() && IsRegisteredPath(AssetNamespace))
+	{
+		PathTree->RemovePath(AssetNamespace);
+
+		// Add the removed asset to the exclusion set so we don't re-register it.
+		FSoftObjectPath RemovedObjectPath = AssetData.ToSoftObjectPath();
+		ExcludedObjectPaths.Add(RemovedObjectPath);
+
+		// Refresh the registry in case the same namespace is in use by another asset.
+		FindAndRegisterAllNamespaces();
+		
+		// Clear the removed asset from the exclusion set.
+		ExcludedObjectPaths.Remove(RemovedObjectPath);
+	}
 }
 
-void FBlueprintNamespaceRegistry::OnAssetRenamed(const FAssetData& AssetData, const FString& InOldName)
+void FBlueprintNamespaceRegistry::OnAssetRenamed(const FAssetData& AssetData, const FString& InOldPath)
 {
-	// @todo_namespaces - Handle Blueprint asset rename/relocation.
+	FString OldDefaultNamespacePath;
+	FBlueprintNamespaceUtilities::ConvertPackagePathToNamespacePath(InOldPath, OldDefaultNamespacePath);
+
+	// Remove the old path if it was explicitly registered as the default namespace.
+	if (IsRegisteredPath(OldDefaultNamespacePath))
+	{
+		PathTree->RemovePath(OldDefaultNamespacePath);
+	}
+
+	// Register the asset's new package path if we are using it as the default namespace.
+	if (FBlueprintNamespaceUtilities::GetDefaultBlueprintNamespaceType() == EDefaultBlueprintNamespaceType::UsePackagePathAsDefaultNamespace)
+	{
+		FString NewDefaultNamespacePath;
+		FBlueprintNamespaceUtilities::ConvertPackagePathToNamespacePath(AssetData.PackageName.ToString(), NewDefaultNamespacePath);
+
+		PathTree->AddPath(NewDefaultNamespacePath);
+	}
 }
 
 bool FBlueprintNamespaceRegistry::IsRegisteredPath(const FString& InPath) const
@@ -153,11 +184,32 @@ void FBlueprintNamespaceRegistry::GetAllRegisteredPaths(TArray<FString>& OutPath
 
 void FBlueprintNamespaceRegistry::FindAndRegisterAllNamespaces()
 {
+	// Resolve the subset of excluded objects that are loaded.
+	TSet<const UObject*> ExcludedObjects;
+	for (const FSoftObjectPath& ExcludedObjectPath : ExcludedObjectPaths)
+	{
+		if (const UObject* ExcludedObject = ExcludedObjectPath.ResolveObject())
+		{
+			if (const UBlueprint* ExcludedBlueprintObject = Cast<UBlueprint>(ExcludedObject))
+			{
+				// Redirect to the generated class when excluding a loaded Blueprint asset.
+				if (ExcludedBlueprintObject->GeneratedClass)
+				{
+					ExcludedObjects.Add(ExcludedBlueprintObject->GeneratedClass);
+				}
+			}
+			else
+			{
+				ExcludedObjects.Add(ExcludedObject);
+			}
+		}
+	}
+
 	// Register loaded class type namespace identifiers.
 	for (TObjectIterator<UClass> ClassIt; ClassIt; ++ClassIt)
 	{
 		const UClass* ClassObject = *ClassIt;
-		if (UEdGraphSchema_K2::IsAllowableBlueprintVariableType(ClassObject))
+		if (UEdGraphSchema_K2::IsAllowableBlueprintVariableType(ClassObject) && !ExcludedObjects.Contains(ClassObject))
 		{
 			RegisterNamespace(ClassObject);
 		}
@@ -167,7 +219,7 @@ void FBlueprintNamespaceRegistry::FindAndRegisterAllNamespaces()
 	for (TObjectIterator<UScriptStruct> StructIt; StructIt; ++StructIt)
 	{
 		const UScriptStruct* StructObject = *StructIt;
-		if (UEdGraphSchema_K2::IsAllowableBlueprintVariableType(StructObject))
+		if (UEdGraphSchema_K2::IsAllowableBlueprintVariableType(StructObject) && !ExcludedObjects.Contains(StructObject))
 		{
 			RegisterNamespace(StructObject);
 		}
@@ -177,7 +229,7 @@ void FBlueprintNamespaceRegistry::FindAndRegisterAllNamespaces()
 	for (TObjectIterator<UEnum> EnumIt; EnumIt; ++EnumIt)
 	{
 		const UEnum* EnumObject = *EnumIt;
-		if (UEdGraphSchema_K2::IsAllowableBlueprintVariableType(EnumObject))
+		if (UEdGraphSchema_K2::IsAllowableBlueprintVariableType(EnumObject) && !ExcludedObjects.Contains(EnumObject))
 		{
 			RegisterNamespace(EnumObject);
 		}
@@ -187,7 +239,7 @@ void FBlueprintNamespaceRegistry::FindAndRegisterAllNamespaces()
 	for (TObjectIterator<UBlueprintFunctionLibrary> LibraryIt; LibraryIt; ++LibraryIt)
 	{
 		const UBlueprintFunctionLibrary* LibraryObject = *LibraryIt;
-		if (LibraryObject)
+		if (LibraryObject && !ExcludedObjects.Contains(LibraryObject))
 		{
 			RegisterNamespace(LibraryObject);
 		}
@@ -207,7 +259,7 @@ void FBlueprintNamespaceRegistry::FindAndRegisterAllNamespaces()
 	AssetRegistry.GetAssets(ClassFilter, BlueprintAssets);
 	for (const FAssetData& BlueprintAsset : BlueprintAssets)
 	{
-		if (!BlueprintAsset.IsAssetLoaded())
+		if (!BlueprintAsset.IsAssetLoaded() && !ExcludedObjectPaths.Contains(BlueprintAsset.ToSoftObjectPath()))
 		{
 			RegisterNamespace(BlueprintAsset);
 		}
@@ -270,6 +322,11 @@ void FBlueprintNamespaceRegistry::DumpAllRegisteredPaths()
 void FBlueprintNamespaceRegistry::OnDefaultNamespaceTypeChanged()
 {
 	// Rebuild the registry to reflect the appropriate default namespace identifiers for all known types.
+	Rebuild();
+}
+
+void FBlueprintNamespaceRegistry::Rebuild()
+{
 	PathTree = MakeUnique<FBlueprintNamespacePathTree>();
 	FindAndRegisterAllNamespaces();
 }
diff --git a/Engine/Source/Editor/Kismet/Private/BlueprintNamespaceUtilities.cpp b/Engine/Source/Editor/Kismet/Private/BlueprintNamespaceUtilities.cpp
--- a/Engine/Source/Editor/Kismet/Private/BlueprintNamespaceUtilities.cpp
+++ b/Engine/Source/Editor/Kismet/Private/BlueprintNamespaceUtilities.cpp
@@ -18,16 +18,15 @@ namespace UE::Editor::Kismet::Private
 
 	// Delegate invoked whenever the default Blueprint namespace type is set to a different value.
 	static FBlueprintNamespaceUtilities::FOnDefaultBlueprintNamespaceTypeChanged OnDefaultBlueprintNamespaceTypeChangedDelegate;
+}
 
-	// Helper method to convert a package path to a Blueprint namespace identifier string.
-	static void ConvertPackagePathToNamespacePath(const FString& InPackagePath, FString& OutNamespacePath)
+void FBlueprintNamespaceUtilities::ConvertPackagePathToNamespacePath(const FString& InPackagePath, FString& OutNamespacePath)
+{
+	OutNamespacePath = InPackagePath;
+	OutNamespacePath.ReplaceCharInline(TEXT('/'), TEXT('.'));
+	if (OutNamespacePath.StartsWith(TEXT(".")))
 	{
-		OutNamespacePath = InPackagePath;
-		OutNamespacePath.ReplaceCharInline(TEXT('/'), TEXT('.'));
-		if (OutNamespacePath.StartsWith(TEXT(".")))
-		{
-			OutNamespacePath.RemoveAt(0);
-		}
+		OutNamespacePath.RemoveAt(0);
 	}
 }
 
diff --git a/Engine/Source/Editor/Kismet/Public/BlueprintNamespaceRegistry.h b/Engine/Source/Editor/Kismet/Public/BlueprintNamespaceRegistry.h
--- a/Engine/Source/Editor/Kismet/Public/BlueprintNamespaceRegistry.h
+++ b/Engine/Source/Editor/Kismet/Public/BlueprintNamespaceRegistry.h
@@ -3,6 +3,7 @@
 #pragma once
 
 #include "CoreMinimal.h"
+#include "UObject/SoftObjectPath.h"
 
 struct FAssetData;
 struct FBlueprintNamespacePathTree;
@@ -53,6 +54,11 @@ public:
 	 */
 	void RegisterNamespace(const FString& InPath);
 
+	/**
+	 * Recreates the namespace registry.
+	 */
+	void Rebuild();
+
 protected:
 	FBlueprintNamespaceRegistry();
 
@@ -83,4 +89,7 @@ private:
 
 	/** Handles storage and retrieval for namespace path identifiers. */
 	TUniquePtr<FBlueprintNamespacePathTree> PathTree;
+
+	/** Internal set of objects to exclude during namespace registration. */
+	TSet<FSoftObjectPath> ExcludedObjectPaths;
 };
\ No newline at end of file
diff --git a/Engine/Source/Editor/Kismet/Public/BlueprintNamespaceUtilities.h b/Engine/Source/Editor/Kismet/Public/BlueprintNamespaceUtilities.h
--- a/Engine/Source/Editor/Kismet/Public/BlueprintNamespaceUtilities.h
+++ b/Engine/Source/Editor/Kismet/Public/BlueprintNamespaceUtilities.h
@@ -88,4 +88,12 @@ public:
 	 * Refresh the Blueprint editor environment to align with current namespace editor feature settings.
 	 */
 	static void RefreshBlueprintEditorFeatures();
+
+	/**
+	 * Helper method to convert a package path to a Blueprint namespace identifier string.
+	 *
+	 * @param InPackagePath		Package path input string.
+	 * @param OutNamespacePath	The input string converted to namespace identifier format.
+	 */
+	static void ConvertPackagePathToNamespacePath(const FString& InPackagePath, FString& OutNamespacePath);
 };
\ No newline at end of file
