Auto-import the associated namespace after selecting a non-imported action in the Blueprint graph context menu.

Change summary:
- Added UBlueprintNodeSpawner::ImportTarget to store/serialize an object reference to allow for resolving action namespace scope on selection.
- Modified UBlueprintFieldNodeSpawner::SetField() to determine and store the import target object reference for a field-based node spawner type.
- Modified FBlueprintNodeSpawnerFactory::MakeMacroNodeSpawner() to determine and store the import target object for a macro instance node spawner.
- Modified SBlueprintActionMenu::OnActionSelected() to auto-import the associated namespace after selecting an action item containing a valid node spawner.

diff --git a/Engine/Source/Editor/BlueprintGraph/Private/BlueprintActionDatabase.cpp b/Engine/Source/Editor/BlueprintGraph/Private/BlueprintActionDatabase.cpp
--- a/Engine/Source/Editor/BlueprintGraph/Private/BlueprintActionDatabase.cpp
+++ b/Engine/Source/Editor/BlueprintGraph/Private/BlueprintActionDatabase.cpp
@@ -158,6 +158,15 @@ static UBlueprintNodeSpawner* FBlueprintNodeSpawnerFactory::MakeMacroNodeSpawner
 	UBlueprintNodeSpawner* NodeSpawner = UBlueprintNodeSpawner::Create(UK2Node_MacroInstance::StaticClass());
 	check(NodeSpawner != nullptr);
 
+	// Set the import reference for macro library instance node spawner actions.
+	if (const UBlueprint* MacroBlueprint = MacroGraph->GetTypedOuter<UBlueprint>())
+	{
+		if (MacroBlueprint->BlueprintType == BPTYPE_MacroLibrary)
+		{
+			NodeSpawner->ImportTarget = MacroBlueprint;
+		}
+	}
+
 	auto CustomizeMacroNodeLambda = [](UEdGraphNode* NewNode, bool bIsTemplateNode, TWeakObjectPtr<UEdGraph> InMacroGraph)
 	{
 		UK2Node_MacroInstance* MacroNode = CastChecked<UK2Node_MacroInstance>(NewNode);
diff --git a/Engine/Source/Editor/BlueprintGraph/Private/BlueprintFieldNodeSpawner.cpp b/Engine/Source/Editor/BlueprintGraph/Private/BlueprintFieldNodeSpawner.cpp
--- a/Engine/Source/Editor/BlueprintGraph/Private/BlueprintFieldNodeSpawner.cpp
+++ b/Engine/Source/Editor/BlueprintGraph/Private/BlueprintFieldNodeSpawner.cpp
@@ -66,14 +66,20 @@ FFieldVariant UBlueprintFieldNodeSpawner::GetField() const
 
 void UBlueprintFieldNodeSpawner::SetField(FFieldVariant InField)
 {
+	const UStruct* OwnerStruct;
+
 	if (InField.IsUObject())
 	{
 		Field = InField.Get<UField>();
+		OwnerStruct = Field->GetOwnerStruct();
 	}
 	else
 	{
 		Property = InField.Get<FProperty>();
+		OwnerStruct = Property->GetOwnerStruct();
 	}
+
+	ImportTarget = OwnerStruct;
 }
 
 #undef LOCTEXT_NAMESPACE
diff --git a/Engine/Source/Editor/BlueprintGraph/Public/BlueprintNodeSpawner.h b/Engine/Source/Editor/BlueprintGraph/Public/BlueprintNodeSpawner.h
--- a/Engine/Source/Editor/BlueprintGraph/Public/BlueprintNodeSpawner.h
+++ b/Engine/Source/Editor/BlueprintGraph/Public/BlueprintNodeSpawner.h
@@ -136,6 +136,13 @@ public:
 	UPROPERTY()
 	TSubclassOf<UEdGraphNode> NodeClass;
 
+	/**
+	 * A weak reference to an object whose namespace must be imported for this spawner's
+	 * action to be visible in the context menu. If NULL (default), no import is required.
+	 */
+	UPROPERTY()
+	TWeakObjectPtr<const UObject> ImportTarget;
+
 	/** Defines how this spawner is presented in the ui */
 	FBlueprintActionUiSpec  DefaultMenuSignature;
 
diff --git a/Engine/Source/Editor/Kismet/Private/SBlueprintActionMenu.cpp b/Engine/Source/Editor/Kismet/Private/SBlueprintActionMenu.cpp
--- a/Engine/Source/Editor/Kismet/Private/SBlueprintActionMenu.cpp
+++ b/Engine/Source/Editor/Kismet/Private/SBlueprintActionMenu.cpp
@@ -27,6 +27,9 @@
 #include "SBlueprintContextTargetMenu.h"
 #include "SBlueprintNamespaceEntry.h"
 #include "BlueprintEditorSettings.h"
+#include "BlueprintActionMenuItem.h"
+#include "BlueprintNodeSpawner.h"
+#include "BlueprintNamespaceUtilities.h"
 
 #define LOCTEXT_NAMESPACE "SBlueprintGraphContextMenu"
 
@@ -562,7 +565,7 @@ void SBlueprintActionMenu::OnActionSelected( const TArray< TSharedPtr<FEdGraphSc
 	{
 		for ( int32 ActionIndex = 0; ActionIndex < SelectedAction.Num(); ActionIndex++ )
 		{
-			if ( SelectedAction[ActionIndex].IsValid() && GraphObj != NULL )
+			if ( SelectedAction[ActionIndex].IsValid() && GraphObj != nullptr )
 			{
 				// Don't dismiss when clicking on dummy action
 				if ( !bActionExecuted && (SelectedAction[ActionIndex]->GetTypeId() != FEdGraphSchemaAction_Dummy::StaticGetTypeId()))
@@ -573,9 +576,28 @@ void SBlueprintActionMenu::OnActionSelected( const TArray< TSharedPtr<FEdGraphSc
 
 				UEdGraphNode* ResultNode = SelectedAction[ActionIndex]->PerformAction(GraphObj, DraggedFromPins, NewNodePosition);
 
-				if ( ResultNode != NULL )
+				if ( ResultNode != nullptr )
 				{
 					NewNodePosition.Y += UEdGraphSchema_K2::EstimateNodeHeight( ResultNode );
+
+					if (GetDefault<UBlueprintEditorSettings>()->bEnableNamespaceImportingFeatures)
+					{
+						if (SelectedAction[ActionIndex]->GetTypeId() == FBlueprintActionMenuItem::StaticGetTypeId())
+						{
+							// Check for a node spawner associated with this menu item.
+							const UBlueprintNodeSpawner* NodeSpawner = StaticCastSharedPtr<FBlueprintActionMenuItem>(SelectedAction[ActionIndex])->GetRawAction();
+
+							TSharedPtr<FBlueprintEditor> BlueprintEditorPtr = EditorPtr.Pin();
+							if (NodeSpawner && BlueprintEditorPtr.IsValid())
+							{
+								if (const UObject* ImportTarget = NodeSpawner->ImportTarget.Get())
+								{
+									// Auto-import the namespace associated with the node spawner's import target.
+									BlueprintEditorPtr->ImportNamespace(FBlueprintNamespaceUtilities::GetObjectNamespace(ImportTarget));
+								}
+							}
+						}
+					}
 				}
 			}
 		}
