Added bIsRegeneratingOnLoad to PostCDOCompiled context

diff --git a/Engine/Source/Editor/Kismet/Private/BlueprintCompilationManager.cpp b/Engine/Source/Editor/Kismet/Private/BlueprintCompilationManager.cpp
--- a/Engine/Source/Editor/Kismet/Private/BlueprintCompilationManager.cpp
+++ b/Engine/Source/Editor/Kismet/Private/BlueprintCompilationManager.cpp
@@ -1519,6 +1519,7 @@ void FBlueprintCompilationManagerImpl::FlushCompilationQueueImpl(bool bSuppressB
 			if (CompilerData.Compiler.IsValid())
 			{
 				UObject::FPostCDOCompiledContext PostCDOCompiledContext;
+				PostCDOCompiledContext.bIsRegeneratingOnLoad = CompilerData.BP->bIsRegeneratingOnLoad;
 				PostCDOCompiledContext.bIsSkeletonOnly = CompilerData.IsSkeletonOnly();
 
 				CompilerData.Compiler->PostCDOCompiled(PostCDOCompiledContext);
diff --git a/Engine/Source/Runtime/CoreUObject/Public/UObject/ObjectCompileContext.h b/Engine/Source/Runtime/CoreUObject/Public/UObject/ObjectCompileContext.h
--- a/Engine/Source/Runtime/CoreUObject/Public/UObject/ObjectCompileContext.h
+++ b/Engine/Source/Runtime/CoreUObject/Public/UObject/ObjectCompileContext.h
@@ -5,6 +5,8 @@
 #if WITH_EDITOR
 struct FObjectPostCDOCompiledContext
 {
+	/** True if the Blueprint is currently being regenerated for the first time after being loaded (aka, compile-on-load) */
+	bool bIsRegeneratingOnLoad = false;
 	/** True if this notification was from a 'skeleton-only' compile */
 	bool bIsSkeletonOnly = false;
 };
