Align Client names in Blueprint Debugger with PIE window titles

diff --git a/Engine/Source/Editor/Kismet/Private/Debugging/SKismetDebugTreeView.cpp b/Engine/Source/Editor/Kismet/Private/Debugging/SKismetDebugTreeView.cpp
--- a/Engine/Source/Editor/Kismet/Private/Debugging/SKismetDebugTreeView.cpp
+++ b/Engine/Source/Editor/Kismet/Private/Debugging/SKismetDebugTreeView.cpp
@@ -2571,34 +2571,7 @@ private:
 		{
 			if (UWorld* World = Object->GetTypedOuter<UWorld>())
 			{
-				ENetMode NetMode = World->GetNetMode();
-
-				FText WorldName;
-
-				switch (NetMode)
-				{
-				case NM_Standalone:
-					WorldName = LOCTEXT("DebugWorldStandalone", "Standalone");
-					break;
-
-				case NM_ListenServer:
-					WorldName = LOCTEXT("DebugWorldListenServer", "Listen Server");
-					break;
-
-				case NM_DedicatedServer:
-					WorldName = LOCTEXT("DebugWorldDedicatedServer", "Dedicated Server");
-					break;
-
-				case NM_Client:
-					if (FWorldContext* PieContext = GEngine->GetWorldContextFromWorld(World))
-					{
-						WorldName = FText::Format(LOCTEXT("DebugWorldClient", "Client {0}"), PieContext->PIEInstance - 1);
-						break;
-					}
-					[[fallthrough]];
-				default:
-					WorldName = LOCTEXT("UnknownWorldMode", "[Unknown]");
-				};
+				FText WorldName = FText::FromString(GetDebugStringForWorld(World));
 
 				return FText::FormatNamed(LOCTEXT("ParentLineTooltip", "{ObjectFullPath}\nWorld: {WorldFullPath}\nWorld Type: {WorldType}"),
 					TEXT("ObjectFullPath"), FText::FromString(Object->GetPathName()),
diff --git a/Engine/Source/Runtime/Engine/Private/World.cpp b/Engine/Source/Runtime/Engine/Private/World.cpp
--- a/Engine/Source/Runtime/Engine/Private/World.cpp
+++ b/Engine/Source/Runtime/Engine/Private/World.cpp
@@ -5539,7 +5539,6 @@ void UWorld::SetWorldDataLayers(AWorldDataLayers* NewWorldDataLayers)
 
 FString UWorld::GetDebugDisplayName() const
 {
-	extern ENGINE_API FString GetDebugStringForWorld(const UWorld* World);
 	return FString::Printf(TEXT("%s (%s)"), *GetDebugStringForWorld(this), *GetPathNameSafe(this));
 }
 
diff --git a/Engine/Source/Runtime/Engine/Public/UnrealEngine.h b/Engine/Source/Runtime/Engine/Public/UnrealEngine.h
--- a/Engine/Source/Runtime/Engine/Public/UnrealEngine.h
+++ b/Engine/Source/Runtime/Engine/Public/UnrealEngine.h
@@ -495,6 +495,9 @@ ENGINE_API extern int32 GUnbuiltHLODCount;
 // Update the debugging aid GPlayInEditorContextString based on the current world context (does nothing in WITH_EDITOR=0 builds)
 ENGINE_API void UpdatePlayInEditorWorldDebugString(const FWorldContext* WorldContext);
 
+// Returns the Debug string for a given world (Standalone, Listen Server, Client #, etc)
+ENGINE_API FString GetDebugStringForWorld(const UWorld* World);
+
 // Used to temporarily override GPlayInEditorID, correctly updating the debug string and other state as necessary
 struct ENGINE_API FTemporaryPlayInEditorIDOverride
 {
