Some minor refactoring to assist with auto-importing external type object dependency namespaces into the Blueprint editor context in response to user actions.

Change summary:
- Moved base definition of UK2Node::HasExternalDependencies() up one level to UEdGraphNode (to avoid blind casts in a context where we may not necessarily have a UK2Node subtype).
- Updated subtype headers that override HasExternalDependencies() to relocate the declaration into the class declaration's 'UEdGraphNode overrides' section (for clarity).
- Reverted previous UBlueprintNodeSpawner::ImportTarget solution in favor of using node dependencies; this allows auto-import actions to scale more easily to all existing node subtypes and avoids data duplication.
- Modified SBlueprintActionMenu::OnActionSelected() to utilize the new method for determining which namespace(s) to auto-import after executing a node spawner action in response to a context menu item selection (UE-146803).
- Reverted FBlueprintNamespaceUtilities::GetPropertyValueNamespaces() back to its original 5.0 signature so that it uses a TSet<FString> for the output rather than a TArray<FString>; this fits in better with input to other APIs.
- Renamed FBlueprintEditor::FImportNamespaceParameters to FImportNamespaceExParameters (for clarity).
- Changed the 'FBlueprintEditor::FImportNamespaceParameters::AdditionalNamespaces' member to a TSet<FString> containing *all* namespaces to import rather than limiting to only additional ones.
- Reverted FBlueprintEditor::ImportNamespace() back to its original 5.0 signature and added/implemented FBlueprintEditor::ImportNamespaceEx() as the "extended" version that allows for batched imports and future customizations.
- Revised existing FBlueprintEditor::ImportNamespace() call sites to utilize and/or conform to the extended ImportNamespaceEx() API where appropriate.

diff --git a/Engine/Source/Editor/BlueprintGraph/Classes/K2Node.h b/Engine/Source/Editor/BlueprintGraph/Classes/K2Node.h
--- a/Engine/Source/Editor/BlueprintGraph/Classes/K2Node.h
+++ b/Engine/Source/Editor/BlueprintGraph/Classes/K2Node.h
@@ -216,12 +216,6 @@ public:
 	/** Returns whether this node is considered 'pure' by the compiler */
 	virtual bool IsNodePure() const { return false; }
 
-	/** 
-	 * Returns whether or not this node has dependencies on an external structure 
-	 * If OptionalOutput isn't null, it should be filled with the known dependencies objects (Classes, Structures, Functions, etc).
-	 */
-	virtual bool HasExternalDependencies(TArray<class UStruct*>* OptionalOutput = nullptr) const { return false; }
-
 	/** Returns whether this node can have breakpoints placed on it in the debugger */
 	BLUEPRINTGRAPH_API virtual bool CanPlaceBreakpoints() const;
 
diff --git a/Engine/Source/Editor/BlueprintGraph/Private/BlueprintActionDatabase.cpp b/Engine/Source/Editor/BlueprintGraph/Private/BlueprintActionDatabase.cpp
--- a/Engine/Source/Editor/BlueprintGraph/Private/BlueprintActionDatabase.cpp
+++ b/Engine/Source/Editor/BlueprintGraph/Private/BlueprintActionDatabase.cpp
@@ -158,15 +158,6 @@ static UBlueprintNodeSpawner* FBlueprintNodeSpawnerFactory::MakeMacroNodeSpawner
 	UBlueprintNodeSpawner* NodeSpawner = UBlueprintNodeSpawner::Create(UK2Node_MacroInstance::StaticClass());
 	check(NodeSpawner != nullptr);
 
-	// Set the import reference for macro library instance node spawner actions.
-	if (const UBlueprint* MacroBlueprint = MacroGraph->GetTypedOuter<UBlueprint>())
-	{
-		if (MacroBlueprint->BlueprintType == BPTYPE_MacroLibrary)
-		{
-			NodeSpawner->ImportTarget = MacroBlueprint;
-		}
-	}
-
 	auto CustomizeMacroNodeLambda = [](UEdGraphNode* NewNode, bool bIsTemplateNode, TWeakObjectPtr<UEdGraph> InMacroGraph)
 	{
 		UK2Node_MacroInstance* MacroNode = CastChecked<UK2Node_MacroInstance>(NewNode);
diff --git a/Engine/Source/Editor/BlueprintGraph/Private/BlueprintFieldNodeSpawner.cpp b/Engine/Source/Editor/BlueprintGraph/Private/BlueprintFieldNodeSpawner.cpp
--- a/Engine/Source/Editor/BlueprintGraph/Private/BlueprintFieldNodeSpawner.cpp
+++ b/Engine/Source/Editor/BlueprintGraph/Private/BlueprintFieldNodeSpawner.cpp
@@ -66,20 +66,14 @@ FFieldVariant UBlueprintFieldNodeSpawner::GetField() const
 
 void UBlueprintFieldNodeSpawner::SetField(FFieldVariant InField)
 {
-	const UStruct* OwnerStruct;
-
 	if (InField.IsUObject())
 	{
 		Field = InField.Get<UField>();
-		OwnerStruct = Field->GetOwnerStruct();
 	}
 	else
 	{
 		Property = InField.Get<FProperty>();
-		OwnerStruct = Property->GetOwnerStruct();
 	}
-
-	ImportTarget = OwnerStruct;
 }
 
 #undef LOCTEXT_NAMESPACE
diff --git a/Engine/Source/Editor/BlueprintGraph/Public/BlueprintNodeSpawner.h b/Engine/Source/Editor/BlueprintGraph/Public/BlueprintNodeSpawner.h
--- a/Engine/Source/Editor/BlueprintGraph/Public/BlueprintNodeSpawner.h
+++ b/Engine/Source/Editor/BlueprintGraph/Public/BlueprintNodeSpawner.h
@@ -136,13 +136,6 @@ public:
 	UPROPERTY()
 	TSubclassOf<UEdGraphNode> NodeClass;
 
-	/**
-	 * A weak reference to an object whose namespace must be imported for this spawner's
-	 * action to be visible in the context menu. If NULL (default), no import is required.
-	 */
-	UPROPERTY()
-	TWeakObjectPtr<const UObject> ImportTarget;
-
 	/** Defines how this spawner is presented in the ui */
 	FBlueprintActionUiSpec  DefaultMenuSignature;
 
diff --git a/Engine/Source/Editor/Kismet/Private/BlueprintDetailsCustomization.cpp b/Engine/Source/Editor/Kismet/Private/BlueprintDetailsCustomization.cpp
--- a/Engine/Source/Editor/Kismet/Private/BlueprintDetailsCustomization.cpp
+++ b/Engine/Source/Editor/Kismet/Private/BlueprintDetailsCustomization.cpp
@@ -2782,20 +2782,15 @@ void FBlueprintVarActionDetails::ImportNamespacesForPropertyValue(const UStruct*
 	}
 
 	// Auto-import any namespace(s) associated with the property's value into the current editor context.
-	TArray<FString> AssociatedNamespaces;
-	FBlueprintNamespaceUtilities::GetPropertyValueNamespaces(InStruct, InProperty, InContainer, AssociatedNamespaces);
-	if (AssociatedNamespaces.Num() > 0)
+	TSharedPtr<SMyBlueprint> MyBlueprintPtr = MyBlueprint.Pin();
+	if (MyBlueprintPtr.IsValid())
 	{
-		TSharedPtr<SMyBlueprint> MyBlueprintPtr = MyBlueprint.Pin();
-		if (MyBlueprintPtr.IsValid())
+		TSharedPtr<FBlueprintEditor> BlueprintEditor = MyBlueprintPtr->GetBlueprintEditor().Pin();
+		if (BlueprintEditor.IsValid())
 		{
-			TSharedPtr<FBlueprintEditor> BlueprintEditor = MyBlueprintPtr->GetBlueprintEditor().Pin();
-			if (BlueprintEditor.IsValid())
-			{
-				FBlueprintEditor::FImportNamespaceParameters Params;
-				Params.AdditionalNamespaces = MakeArrayView(AssociatedNamespaces).RightChop(1);
-				BlueprintEditor->ImportNamespace(AssociatedNamespaces[0], Params);
-			}
+			FBlueprintEditor::FImportNamespaceExParameters Params;
+			FBlueprintNamespaceUtilities::GetPropertyValueNamespaces(InStruct, InProperty, InContainer, Params.NamespacesToImport);
+			BlueprintEditor->ImportNamespaceEx(Params);
 		}
 	}
 }
@@ -5678,15 +5673,16 @@ void FBlueprintImportsLayout::OnNamespaceSelected(const FString& InNamespace)
 	TSharedPtr<FBlueprintEditor> BlueprintEditorPtr = GetPinnedBlueprintEditorPtr();
 	if (BlueprintEditorPtr.IsValid())
 	{
-		FBlueprintEditor::FImportNamespaceParameters Params;
+		FBlueprintEditor::FImportNamespaceExParameters Params;
 		Params.bIsAutoImport = false;
-		Params.OnImportCallback = FSimpleDelegate::CreateLambda([this]()
+		Params.NamespacesToImport.Add(InNamespace);
+		Params.OnPostImportCallback = FSimpleDelegate::CreateLambda([this]()
 		{
 			RegenerateChildContent();
 		});
 
 		// Add to edited Blueprint(s) and import into the current editor context.
-		BlueprintEditorPtr->ImportNamespace(InNamespace, Params);
+		BlueprintEditorPtr->ImportNamespaceEx(Params);
 	}
 
 	OnRefreshInDetailsView();
@@ -6051,15 +6047,16 @@ void FBlueprintGlobalOptionsDetails::HandleNamespaceValueChange(const FString& I
 			}
 			else if (GetDefault<UBlueprintEditorSettings>()->bEnableNamespaceImportingFeatures)
 			{
-				FBlueprintEditor::FImportNamespaceParameters Params;
+				FBlueprintEditor::FImportNamespaceExParameters Params;
 				Params.bIsAutoImport = false;
-				Params.OnImportCallback = FSimpleDelegate::CreateLambda([&bRefreshDetailsView]()
+				Params.NamespacesToImport.Add(InNewValue);
+				Params.OnPostImportCallback = FSimpleDelegate::CreateLambda([&bRefreshDetailsView]()
 				{
 					bRefreshDetailsView = true;
 				});
 
 				// Import the new namespace into the current editor context.
-				BlueprintEditor->ImportNamespace(InNewValue, Params);
+				BlueprintEditor->ImportNamespaceEx(Params);
 			}
 
 			// Refresh the details view if necessary.
diff --git a/Engine/Source/Editor/Kismet/Private/BlueprintEditor.cpp b/Engine/Source/Editor/Kismet/Private/BlueprintEditor.cpp
--- a/Engine/Source/Editor/Kismet/Private/BlueprintEditor.cpp
+++ b/Engine/Source/Editor/Kismet/Private/BlueprintEditor.cpp
@@ -2039,17 +2039,25 @@ void FBlueprintEditor::LoadLibrariesFromAssetRegistry()
 	}
 }
 
-void FBlueprintEditor::ImportNamespace(const FString& InNamespace, const FImportNamespaceParameters& InParams)
+void FBlueprintEditor::ImportNamespace(const FString& InNamespace)
 {
-	// No need to import the global namespace.
-	if (InNamespace.IsEmpty())
+	FImportNamespaceExParameters Params;
+	Params.NamespacesToImport.Add(InNamespace);
+	ImportNamespaceEx(Params);
+}
+
+void FBlueprintEditor::ImportNamespaceEx(const FImportNamespaceExParameters& InParams)
+{
+	// Exit now if no input was given.
+	if (InParams.NamespacesToImport.IsEmpty())
 	{
 		return;
 	}
 
 	auto AddNamespaceToImportList = [](UBlueprint* InBlueprint, const FString& InNamespace) -> bool
 	{
-		if (!InBlueprint->ImportedNamespaces.Contains(InNamespace))
+		// No need to include the global namespace (empty string) in the import list.
+		if (!InNamespace.IsEmpty() && !InBlueprint->ImportedNamespaces.Contains(InNamespace))
 		{
 			InBlueprint->Modify();
 			InBlueprint->ImportedNamespaces.Add(InNamespace);
@@ -2067,18 +2075,18 @@ void FBlueprintEditor::ImportNamespace(const FString& InNamespace, const FImport
 	{
 		if (UBlueprint* BlueprintObj = Cast<UBlueprint>(EditingObj))
 		{
-			// Add it into the Blueprint's user-facing import set.
-			bWasAdded |= AddNamespaceToImportList(BlueprintObj, InNamespace);
-			for (const FString& AdditionalNamespace : InParams.AdditionalNamespaces)
+			// Add each namespace into the Blueprint's user-facing import set.
+			for (const FString& NamespaceToImport : InParams.NamespacesToImport)
 			{
-				bWasAdded |= AddNamespaceToImportList(BlueprintObj, AdditionalNamespace);
+				bWasAdded |= AddNamespaceToImportList(BlueprintObj, NamespaceToImport);
 			}
 		}
 	}
 
 	auto AddNamespaceToEditorContext = [](TSharedPtr<FBlueprintNamespaceHelper> ImportsHelper, const FString& InNamespace) -> bool
 	{
-		if (!ImportsHelper->IsIncludedInNamespaceList(InNamespace))
+		// Note: The global namespace (empty string) is implicitly included.
+		if (!InNamespace.IsEmpty() && !ImportsHelper->IsIncludedInNamespaceList(InNamespace))
 		{
 			ImportsHelper->AddNamespace(InNamespace);
 			return true;
@@ -2087,15 +2095,14 @@ void FBlueprintEditor::ImportNamespace(const FString& InNamespace, const FImport
 		return false;
 	};
 
-	// Add it to the current scope of the Blueprint's editor context. Note that in certain cases, imports may already be associated
+	// Add to the current scope of the Blueprint's editor context. Note that in certain cases, imports may already be associated
 	// with the Blueprint, but not yet associated with the editor context (e.g. - auto-import after setting a Blueprint's namespace;
 	// we won't add it to the Blueprint's import list, but we still want to add to the editor context and do any post-import actions).
 	if (ImportedNamespaceHelper.IsValid())
 	{
-		bWasAdded |= AddNamespaceToEditorContext(ImportedNamespaceHelper, InNamespace);
-		for (const FString& AdditionalNamespace : InParams.AdditionalNamespaces)
+		for (const FString& NamespaceToImport : InParams.NamespacesToImport)
 		{
-			bWasAdded |= AddNamespaceToEditorContext(ImportedNamespaceHelper, AdditionalNamespace);
+			bWasAdded |= AddNamespaceToEditorContext(ImportedNamespaceHelper, NamespaceToImport);
 		}
 	}
 
@@ -2112,7 +2119,7 @@ void FBlueprintEditor::ImportNamespace(const FString& InNamespace, const FImport
 		}
 
 		// If bound, execute the post-import callback.
-		InParams.OnImportCallback.ExecuteIfBound();
+		InParams.OnPostImportCallback.ExecuteIfBound();
 	}
 }
 
diff --git a/Engine/Source/Editor/Kismet/Private/BlueprintNamespaceUtilities.cpp b/Engine/Source/Editor/Kismet/Private/BlueprintNamespaceUtilities.cpp
--- a/Engine/Source/Editor/Kismet/Private/BlueprintNamespaceUtilities.cpp
+++ b/Engine/Source/Editor/Kismet/Private/BlueprintNamespaceUtilities.cpp
@@ -154,7 +154,7 @@ FString FBlueprintNamespaceUtilities::GetObjectNamespace(const FSoftObjectPath&
 	return GetAssetNamespace(AssetData);
 }
 
-void FBlueprintNamespaceUtilities::GetPropertyValueNamespaces(const UStruct* InStruct, const FProperty* InProperty, const void* InContainer, TArray<FString>& OutNamespaces)
+void FBlueprintNamespaceUtilities::GetPropertyValueNamespaces(const UStruct* InStruct, const FProperty* InProperty, const void* InContainer, TSet<FString>& OutNamespaces)
 {
 	if (!InStruct || !InProperty || !InContainer)
 	{
@@ -215,7 +215,7 @@ void FBlueprintNamespaceUtilities::GetPropertyValueNamespaces(const UStruct* InS
 			if (ObjectPath.IsValid())
 			{
 				FString Namespace = GetObjectNamespace(ObjectPath);
-				OutNamespaces.AddUnique(Namespace);
+				OutNamespaces.Add(Namespace);
 			}
 		}
 		else if (const FObjectPropertyBase* ObjectProperty = CastField<FObjectPropertyBase>(InProperty))
@@ -224,7 +224,7 @@ void FBlueprintNamespaceUtilities::GetPropertyValueNamespaces(const UStruct* InS
 			if (ObjectValue)
 			{
 				FString Namespace = GetObjectNamespace(ObjectValue);
-				OutNamespaces.AddUnique(Namespace);
+				OutNamespaces.Add(Namespace);
 			}
 		}
 	}
diff --git a/Engine/Source/Editor/Kismet/Private/SBlueprintActionMenu.cpp b/Engine/Source/Editor/Kismet/Private/SBlueprintActionMenu.cpp
--- a/Engine/Source/Editor/Kismet/Private/SBlueprintActionMenu.cpp
+++ b/Engine/Source/Editor/Kismet/Private/SBlueprintActionMenu.cpp
@@ -582,20 +582,31 @@ void SBlueprintActionMenu::OnActionSelected( const TArray< TSharedPtr<FEdGraphSc
 
 					if (GetDefault<UBlueprintEditorSettings>()->bEnableNamespaceImportingFeatures)
 					{
-						if (SelectedAction[ActionIndex]->GetTypeId() == FBlueprintActionMenuItem::StaticGetTypeId())
+						TSharedPtr<FBlueprintEditor> BlueprintEditorPtr = EditorPtr.Pin();
+						if (BlueprintEditorPtr.IsValid())
 						{
-							// Check for a node spawner associated with this menu item.
-							const UBlueprintNodeSpawner* NodeSpawner = StaticCastSharedPtr<FBlueprintActionMenuItem>(SelectedAction[ActionIndex])->GetRawAction();
-
-							TSharedPtr<FBlueprintEditor> BlueprintEditorPtr = EditorPtr.Pin();
-							if (NodeSpawner && BlueprintEditorPtr.IsValid())
+							// Determine which namespace(s) to import, based on the node's external dependencies.
+							TSet<FString> NamespacesToImport;
+							TArray<UStruct*> ExternalDependencies;
+							if (ResultNode->HasExternalDependencies(&ExternalDependencies))
 							{
-								if (const UObject* ImportTarget = NodeSpawner->ImportTarget.Get())
+								for (const UStruct* ExternalDependency : ExternalDependencies)
 								{
-									// Auto-import the namespace associated with the node spawner's import target.
-									BlueprintEditorPtr->ImportNamespace(FBlueprintNamespaceUtilities::GetObjectNamespace(ImportTarget));
+									FString ObjectNamespace = FBlueprintNamespaceUtilities::GetObjectNamespace(ExternalDependency);
+									if (!ObjectNamespace.IsEmpty())
+									{
+										NamespacesToImport.Add(MoveTemp(ObjectNamespace));
+									}
 								}
 							}
+
+							if (NamespacesToImport.Num() > 0)
+							{
+								// Auto-import the namespace(s) gathered above. Additional type objects within the imported scope may be loaded here.
+								FBlueprintEditor::FImportNamespaceExParameters Params;
+								Params.NamespacesToImport = MoveTemp(NamespacesToImport);
+								BlueprintEditorPtr->ImportNamespaceEx(Params);
+							}
 						}
 					}
 				}
@@ -647,8 +658,9 @@ void SBlueprintActionMenu::OnNamespaceSelectedForImport(const FString& InNamespa
 	TSharedPtr<FBlueprintEditor> BlueprintEditorPtr = EditorPtr.Pin();
 	if (BlueprintEditorPtr.IsValid())
 	{
-		FBlueprintEditor::FImportNamespaceParameters Params;
-		Params.OnImportCallback = FSimpleDelegate::CreateLambda([GraphActionMenu = this->GraphActionMenu]()
+		FBlueprintEditor::FImportNamespaceExParameters Params;
+		Params.NamespacesToImport.Add(InNamespace);
+		Params.OnPostImportCallback = FSimpleDelegate::CreateLambda([GraphActionMenu = this->GraphActionMenu]()
 		{
 			// Now that additional types have been loaded/imported, update the menu to include any additional action(s).
 			const bool bPreserveExpansion = true;
@@ -657,7 +669,7 @@ void SBlueprintActionMenu::OnNamespaceSelectedForImport(const FString& InNamespa
 		});
 
 		// Auto-import the namespace into the current editor context. This may load additional type assets.
-		BlueprintEditorPtr->ImportNamespace(InNamespace, Params);
+		BlueprintEditorPtr->ImportNamespaceEx(Params);
 	}
 }
 
diff --git a/Engine/Source/Editor/Kismet/Private/SKismetInspector.cpp b/Engine/Source/Editor/Kismet/Private/SKismetInspector.cpp
--- a/Engine/Source/Editor/Kismet/Private/SKismetInspector.cpp
+++ b/Engine/Source/Editor/Kismet/Private/SKismetInspector.cpp
@@ -1016,36 +1016,36 @@ void SKismetInspector::ImportNamespacesForPropertyValue(const FProperty* InPrope
 		return;
 	}
 
-	// Gather all namespace identifier strings associated with the property's value for each edited object.
-	TArray<FString> AssociatedNamespaces;
-	for (const TWeakObjectPtr<UObject>& SelectedObjectPtr : SelectedObjects)
+	TSharedPtr<FBlueprintEditor> BlueprintEditor = BlueprintEditorPtr.Pin();
+	if (BlueprintEditor.IsValid())
 	{
-		if (const UObject* SelectedObject = SelectedObjectPtr.Get())
+		// Gather all namespace identifier strings associated with the property's value for each edited object.
+		TSet<FString> AssociatedNamespaces;
+		for (const TWeakObjectPtr<UObject>& SelectedObjectPtr : SelectedObjects)
 		{
-			const UStruct* SelectedType = SelectedObject->GetClass();
-
-			// Ensure that the selected object type matches the property's owner.
-			// For example, a details customization may select an unrelated object
-			// and then customize each row with an external object reference. In
-			// those cases, the customization would need to handle this explicitly.
-			if (SelectedType != InProperty->GetOwnerStruct())
+			if (const UObject* SelectedObject = SelectedObjectPtr.Get())
 			{
-				continue;
-			}
+				const UStruct* SelectedType = SelectedObject->GetClass();
+
+				// Ensure that the selected object type matches the property's owner.
+				// For example, a details customization may select an unrelated object
+				// and then customize each row with an external object reference. In
+				// those cases, the customization would need to handle this explicitly.
+				if (SelectedType != InProperty->GetOwnerStruct())
+				{
+					continue;
+				}
 
-			FBlueprintNamespaceUtilities::GetPropertyValueNamespaces(SelectedType, InProperty, SelectedObject, AssociatedNamespaces);
+				FBlueprintNamespaceUtilities::GetPropertyValueNamespaces(SelectedType, InProperty, SelectedObject, AssociatedNamespaces);
+			}
 		}
-	}
 
-	// Auto-import any namespace(s) associated with the property's value into the current editor context.
-	if (AssociatedNamespaces.Num() > 0)
-	{
-		TSharedPtr<FBlueprintEditor> BlueprintEditor = BlueprintEditorPtr.Pin();
-		if (BlueprintEditor.IsValid())
+		// Auto-import any namespace(s) associated with the property's value into the current editor context.
+		if (AssociatedNamespaces.Num() > 0)
 		{
-			FBlueprintEditor::FImportNamespaceParameters Params;
-			Params.AdditionalNamespaces = MakeArrayView(AssociatedNamespaces).RightChop(1);
-			BlueprintEditor->ImportNamespace(AssociatedNamespaces[0], Params);
+			FBlueprintEditor::FImportNamespaceExParameters Params;
+			Params.NamespacesToImport = MoveTemp(AssociatedNamespaces);
+			BlueprintEditor->ImportNamespaceEx(Params);
 		}
 	}
 }
diff --git a/Engine/Source/Editor/Kismet/Public/BlueprintEditor.h b/Engine/Source/Editor/Kismet/Public/BlueprintEditor.h
--- a/Engine/Source/Editor/Kismet/Public/BlueprintEditor.h
+++ b/Engine/Source/Editor/Kismet/Public/BlueprintEditor.h
@@ -720,29 +720,32 @@ public:
 
 	/** Gets the default schema for this editor */
 	TSubclassOf<UEdGraphSchema> GetDefaultSchema() const { return GetDefaultSchemaClass(); }
+	
+	/** Imports the given namespace into the editor. This may trigger a load event for additional macro and/or function library assets if not already loaded. */
+	void ImportNamespace(const FString& InNamespace);
 
-	/** Extended parameters for the ImportNamespace() method */
-	struct FImportNamespaceParameters
+	/** Parameters for the extended ImportNamespaceEx() method */
+	struct FImportNamespaceExParameters
 	{
 		/** Whether this is an automatic or explicit (i.e. user-initiated) action. */
 		bool bIsAutoImport;
 
 		/** Callback to use for any post-import actions. Will not be invoked if nothing is imported. */
-		FSimpleDelegate OnImportCallback;
+		FSimpleDelegate OnPostImportCallback;
 
-		/** Additional namespaces to be imported as part of the single transaction. */
-		TArray<FString> AdditionalNamespaces;
+		/** One or more unique namespace identifiers to be imported as part of the single transaction. */
+		TSet<FString> NamespacesToImport;
 
 		/** Default ctor (for initialization). */
-		FImportNamespaceParameters()
+		FImportNamespaceExParameters()
 		{
 			// Treat as auto-import by default (implies that the editor will auto-refresh the details view).
 			bIsAutoImport = true;
 		}
 	};
-	
-	/** Imports the given namespace into the editor. This may trigger a load event for additional macro and/or function library assets if not already loaded. */
-	void ImportNamespace(const FString& InNamespace, const FImportNamespaceParameters& InParams = FImportNamespaceParameters());
+
+	/** Imports a set of namespace(s) into the editor. This may trigger a load event for additional macro and/or function library assets if not already loaded. */
+	void ImportNamespaceEx(const FImportNamespaceExParameters& InParams);
 
 	/** Removes the given namespace from the editor's current import context. However, this will NOT unload any associated macro and/or function library assets. */
 	void RemoveNamespace(const FString& InNamespace);
diff --git a/Engine/Source/Editor/Kismet/Public/BlueprintNamespaceUtilities.h b/Engine/Source/Editor/Kismet/Public/BlueprintNamespaceUtilities.h
--- a/Engine/Source/Editor/Kismet/Public/BlueprintNamespaceUtilities.h
+++ b/Engine/Source/Editor/Kismet/Public/BlueprintNamespaceUtilities.h
@@ -53,7 +53,7 @@ public:
 	 * @param InContainer	The source address of the struct/object containing the property's value.
 	 * @param OutNamespaces	Zero or more unique namespace identifier(s) referenced by the property value. An entry with an empty string equates to the default global namespace.
 	 */
-	static void GetPropertyValueNamespaces(const UStruct* InStruct, const FProperty* InProperty, const void* InContainer, TArray<FString>& OutNamespaces);
+	static void GetPropertyValueNamespaces(const UStruct* InStruct, const FProperty* InProperty, const void* InContainer, TSet<FString>& OutNamespaces);
 
 	/**
 	 * Gathers the set of global namespaces that are implicitly imported by all Blueprint assets.
diff --git a/Engine/Source/Runtime/Engine/Classes/EdGraph/EdGraphNode.h b/Engine/Source/Runtime/Engine/Classes/EdGraph/EdGraphNode.h
--- a/Engine/Source/Runtime/Engine/Classes/EdGraph/EdGraphNode.h
+++ b/Engine/Source/Runtime/Engine/Classes/EdGraph/EdGraphNode.h
@@ -827,6 +827,12 @@ public:
 	/** Update node size to new value */
 	virtual void ResizeNode(const FVector2D& NewSize) {}
 
+	/**
+	 * Returns whether or not this node has dependencies on an external structure
+	 * If OptionalOutput isn't null, it should be filled with the known dependencies objects (Classes, Structures, Functions, etc).
+	 */
+	virtual bool HasExternalDependencies(TArray<class UStruct*>* OptionalOutput = nullptr) const { return false; }
+
 	// Returns true if this node is deprecated
 	virtual bool IsDeprecated() const;
 
