Auto-import external type dependency namespaces after pasting one or more node(s) into a Blueprint graph.

Change summary:
- Extended FBlueprintEditor::PasteNodesHere() to gather external dependencies from pasted node(s) and subsequently auto-import any associated namespace(s) from any type object dependencies.

diff --git a/Engine/Source/Editor/Kismet/Private/BlueprintEditor.cpp b/Engine/Source/Editor/Kismet/Private/BlueprintEditor.cpp
--- a/Engine/Source/Editor/Kismet/Private/BlueprintEditor.cpp
+++ b/Engine/Source/Editor/Kismet/Private/BlueprintEditor.cpp
@@ -7267,6 +7267,8 @@ void FBlueprintEditor::PasteNodesHere(class UEdGraph* DestinationGraph, const FV
 		AvgNodePosition.X *= InvNumNodes;
 		AvgNodePosition.Y *= InvNumNodes;
 
+		TSet<FString> NamespacesToImport;
+
 		for (TSet<UEdGraphNode*>::TIterator It(PastedNodes); It; ++It)
 		{
 			UEdGraphNode* Node = *It;
@@ -7280,6 +7282,20 @@ void FBlueprintEditor::PasteNodesHere(class UEdGraph* DestinationGraph, const FV
 			// Give new node a different Guid from the old one
 			Node->CreateNewGuid();
 
+			// Collect any required imports from node dependencies
+			TArray<UStruct*> ExternalDependencies;
+			if (Node->HasExternalDependencies(&ExternalDependencies))
+			{
+				for (const UStruct* ExternalDependency : ExternalDependencies)
+				{
+					FString ObjectNamespace = FBlueprintNamespaceUtilities::GetObjectNamespace(ExternalDependency);
+					if (!ObjectNamespace.IsEmpty())
+					{
+						NamespacesToImport.Add(MoveTemp(ObjectNamespace));
+					}
+				}
+			}
+
 			UK2Node* K2Node = Cast<UK2Node>(Node);
 			if ((K2Node != nullptr) && K2Node->NodeCausesStructuralBlueprintChange())
 			{
@@ -7313,6 +7329,15 @@ void FBlueprintEditor::PasteNodesHere(class UEdGraph* DestinationGraph, const FV
 			// Log new node created to analytics
 			AnalyticsTrackNodeEvent(GetBlueprintObj(), Node, false);
 		}
+
+		// Auto-import any namespace dependencies.
+		const bool bShouldAutoImportNamespaceDependencies = GetDefault<UBlueprintEditorSettings>()->bEnableNamespaceImportingFeatures;
+		if (bShouldAutoImportNamespaceDependencies && NamespacesToImport.Num() > 0)
+		{
+			FImportNamespaceExParameters Params;
+			Params.NamespacesToImport = MoveTemp(NamespacesToImport);
+			ImportNamespaceEx(Params);
+		}
 	}
 
 	if (bNeedToModifyStructurally)
