Consolidate the Blueprint namespace importing features toggle flag check for existing auto-import code paths.

diff --git a/Engine/Source/Editor/Kismet/Private/BlueprintDetailsCustomization.cpp b/Engine/Source/Editor/Kismet/Private/BlueprintDetailsCustomization.cpp
--- a/Engine/Source/Editor/Kismet/Private/BlueprintDetailsCustomization.cpp
+++ b/Engine/Source/Editor/Kismet/Private/BlueprintDetailsCustomization.cpp
@@ -1449,19 +1449,16 @@ void FBlueprintVarActionDetails::OnVarTypeChanged(const FEdGraphPinType& NewPinT
 					FBlueprintEditorUtils::ChangeMemberVariableType(GetBlueprintObj(), VarName, NewPinType);
 				}
 
-				if (GetDefault<UBlueprintEditorSettings>()->bEnableNamespaceImportingFeatures)
-				{
-					TSharedPtr<FBlueprintEditor> BlueprintEditor = MyBlueprint.Pin()->GetBlueprintEditor().Pin();
+				TSharedPtr<FBlueprintEditor> BlueprintEditor = MyBlueprint.Pin()->GetBlueprintEditor().Pin();
 					
-					// If the underlying type object's namespace is not imported, auto-import it now into the current editor context.
-					const UObject* PinSubCategoryObject = NewPinType.PinSubCategoryObject.Get();
-					if (PinSubCategoryObject && BlueprintEditor.IsValid() && BlueprintEditor->IsNonImportedObject(PinSubCategoryObject))
+				// If the underlying type object's namespace is not imported, auto-import it now into the current editor context.
+				const UObject* PinSubCategoryObject = NewPinType.PinSubCategoryObject.Get();
+				if (PinSubCategoryObject && BlueprintEditor.IsValid() && BlueprintEditor->IsNonImportedObject(PinSubCategoryObject))
+				{
+					FString Namespace = FBlueprintNamespaceUtilities::GetObjectNamespace(PinSubCategoryObject);
+					if (!Namespace.IsEmpty())
 					{
-						FString Namespace = FBlueprintNamespaceUtilities::GetObjectNamespace(PinSubCategoryObject);
-						if (!Namespace.IsEmpty())
-						{
-							BlueprintEditor->ImportNamespace(Namespace);
-						}
+						BlueprintEditor->ImportNamespace(Namespace);
 					}
 				}
 			}
@@ -2776,11 +2773,6 @@ void FBlueprintVarActionDetails::OnFinishedChangingLocalVariable(const FProperty
 
 void FBlueprintVarActionDetails::ImportNamespacesForPropertyValue(const UStruct* InStruct, const FProperty* InProperty, const void* InContainer)
 {
-	if (!GetDefault<UBlueprintEditorSettings>()->bEnableNamespaceImportingFeatures)
-	{
-		return;
-	}
-
 	// Auto-import any namespace(s) associated with the property's value into the current editor context.
 	TSharedPtr<SMyBlueprint> MyBlueprintPtr = MyBlueprint.Pin();
 	if (MyBlueprintPtr.IsValid())
@@ -3445,23 +3437,20 @@ void FBlueprintGraphArgumentLayout::PinInfoChanged(const FEdGraphPinType& PinTyp
 							// Reset default value, it probably doesn't match
 							(*UDPinPtr)->PinDefaultValue.Reset();
 
-							if (GetDefault<UBlueprintEditorSettings>()->bEnableNamespaceImportingFeatures)
+							TSharedPtr<FBlueprintEditor> BlueprintEditor;
+							if(MyBPPinned.IsValid())
 							{
-								TSharedPtr<FBlueprintEditor> BlueprintEditor;
-								if(MyBPPinned.IsValid())
-								{
-									BlueprintEditor = MyBPPinned->GetBlueprintEditor().Pin();
-								}
+								BlueprintEditor = MyBPPinned->GetBlueprintEditor().Pin();
+							}
 
-								// If the underlying type object's namespace is not imported, auto-import it now into the current editor context.
-								const UObject* PinSubCategoryObject = PinType.PinSubCategoryObject.Get();
-								if (PinSubCategoryObject && BlueprintEditor.IsValid() && BlueprintEditor->IsNonImportedObject(PinSubCategoryObject))
+							// If the underlying type object's namespace is not imported, auto-import it now into the current editor context.
+							const UObject* PinSubCategoryObject = PinType.PinSubCategoryObject.Get();
+							if (PinSubCategoryObject && BlueprintEditor.IsValid() && BlueprintEditor->IsNonImportedObject(PinSubCategoryObject))
+							{
+								FString Namespace = FBlueprintNamespaceUtilities::GetObjectNamespace(PinSubCategoryObject);
+								if (!Namespace.IsEmpty())
 								{
-									FString Namespace = FBlueprintNamespaceUtilities::GetObjectNamespace(PinSubCategoryObject);
-									if (!Namespace.IsEmpty())
-									{
-										BlueprintEditor->ImportNamespace(Namespace);
-									}
+									BlueprintEditor->ImportNamespace(Namespace);
 								}
 							}
 						}
@@ -6045,7 +6034,7 @@ void FBlueprintGlobalOptionsDetails::HandleNamespaceValueChange(const FString& I
 				// In that case, it will switch to a default namespace in the Imports table.
 				bRefreshDetailsView = true;
 			}
-			else if (GetDefault<UBlueprintEditorSettings>()->bEnableNamespaceImportingFeatures)
+			else
 			{
 				FBlueprintEditor::FImportNamespaceExParameters Params;
 				Params.bIsAutoImport = false;
diff --git a/Engine/Source/Editor/Kismet/Private/BlueprintEditor.cpp b/Engine/Source/Editor/Kismet/Private/BlueprintEditor.cpp
--- a/Engine/Source/Editor/Kismet/Private/BlueprintEditor.cpp
+++ b/Engine/Source/Editor/Kismet/Private/BlueprintEditor.cpp
@@ -1226,12 +1226,8 @@ void FBlueprintEditor::OnComponentAddedToBlueprint(const FSubobjectData& NewSubo
 	check(NewSubobject);
 
 	// If necessary, auto-import the namespace associated with the new subobject's class.
-	const bool bShouldAutoImportTypeNamespace = GetDefault<UBlueprintEditorSettings>()->bEnableNamespaceImportingFeatures;
-	if (bShouldAutoImportTypeNamespace)
-	{
-		const FString SubobjectTypeNamespace = FBlueprintNamespaceUtilities::GetObjectNamespace(NewSubobject->GetClass());
-		ImportNamespace(SubobjectTypeNamespace);
-	}
+	const FString SubobjectTypeNamespace = FBlueprintNamespaceUtilities::GetObjectNamespace(NewSubobject->GetClass());
+	ImportNamespace(SubobjectTypeNamespace);
 }
 
 TSharedRef<SWidget> FBlueprintEditor::CreateGraphTitleBarWidget(TSharedRef<FTabInfo> InTabInfo, UEdGraph* InGraph)
@@ -2048,6 +2044,13 @@ void FBlueprintEditor::ImportNamespace(const FString& InNamespace)
 
 void FBlueprintEditor::ImportNamespaceEx(const FImportNamespaceExParameters& InParams)
 {
+	// No auto-import actions if features are disabled.
+	const bool bIsAutoInputEnabled = GetDefault<UBlueprintEditorSettings>()->bEnableNamespaceImportingFeatures;
+	if (InParams.bIsAutoImport && !bIsAutoInputEnabled)
+	{
+		return;
+	}
+
 	// Exit now if no input was given.
 	if (InParams.NamespacesToImport.IsEmpty())
 	{
diff --git a/Engine/Source/Editor/Kismet/Private/SBlueprintActionMenu.cpp b/Engine/Source/Editor/Kismet/Private/SBlueprintActionMenu.cpp
--- a/Engine/Source/Editor/Kismet/Private/SBlueprintActionMenu.cpp
+++ b/Engine/Source/Editor/Kismet/Private/SBlueprintActionMenu.cpp
@@ -580,33 +580,30 @@ void SBlueprintActionMenu::OnActionSelected( const TArray< TSharedPtr<FEdGraphSc
 				{
 					NewNodePosition.Y += UEdGraphSchema_K2::EstimateNodeHeight( ResultNode );
 
-					if (GetDefault<UBlueprintEditorSettings>()->bEnableNamespaceImportingFeatures)
+					TSharedPtr<FBlueprintEditor> BlueprintEditorPtr = EditorPtr.Pin();
+					if (BlueprintEditorPtr.IsValid())
 					{
-						TSharedPtr<FBlueprintEditor> BlueprintEditorPtr = EditorPtr.Pin();
-						if (BlueprintEditorPtr.IsValid())
+						// Determine which namespace(s) to import, based on the node's external dependencies.
+						TSet<FString> NamespacesToImport;
+						TArray<UStruct*> ExternalDependencies;
+						if (ResultNode->HasExternalDependencies(&ExternalDependencies))
 						{
-							// Determine which namespace(s) to import, based on the node's external dependencies.
-							TSet<FString> NamespacesToImport;
-							TArray<UStruct*> ExternalDependencies;
-							if (ResultNode->HasExternalDependencies(&ExternalDependencies))
+							for (const UStruct* ExternalDependency : ExternalDependencies)
 							{
-								for (const UStruct* ExternalDependency : ExternalDependencies)
+								FString ObjectNamespace = FBlueprintNamespaceUtilities::GetObjectNamespace(ExternalDependency);
+								if (!ObjectNamespace.IsEmpty())
 								{
-									FString ObjectNamespace = FBlueprintNamespaceUtilities::GetObjectNamespace(ExternalDependency);
-									if (!ObjectNamespace.IsEmpty())
-									{
-										NamespacesToImport.Add(MoveTemp(ObjectNamespace));
-									}
+									NamespacesToImport.Add(MoveTemp(ObjectNamespace));
 								}
 							}
+						}
 
-							if (NamespacesToImport.Num() > 0)
-							{
-								// Auto-import the namespace(s) gathered above. Additional type objects within the imported scope may be loaded here.
-								FBlueprintEditor::FImportNamespaceExParameters Params;
-								Params.NamespacesToImport = MoveTemp(NamespacesToImport);
-								BlueprintEditorPtr->ImportNamespaceEx(Params);
-							}
+						if (NamespacesToImport.Num() > 0)
+						{
+							// Auto-import the namespace(s) gathered above. Additional type objects within the imported scope may be loaded here.
+							FBlueprintEditor::FImportNamespaceExParameters Params;
+							Params.NamespacesToImport = MoveTemp(NamespacesToImport);
+							BlueprintEditorPtr->ImportNamespaceEx(Params);
 						}
 					}
 				}
diff --git a/Engine/Source/Editor/Kismet/Private/SBlueprintPalette.cpp b/Engine/Source/Editor/Kismet/Private/SBlueprintPalette.cpp
--- a/Engine/Source/Editor/Kismet/Private/SBlueprintPalette.cpp
+++ b/Engine/Source/Editor/Kismet/Private/SBlueprintPalette.cpp
@@ -835,17 +835,14 @@ private:
 				ActionPtr.Pin()->ChangeVariableType(InNewPinType);
 			}
 
-			if (GetDefault<UBlueprintEditorSettings>()->bEnableNamespaceImportingFeatures)
+			// If the underlying type object's namespace is not imported, auto-import it now into the current editor context.
+			const UObject* PinSubCategoryObject = InNewPinType.PinSubCategoryObject.Get();
+			if (PinSubCategoryObject && BlueprintEditor.IsValid() && BlueprintEditor->IsNonImportedObject(PinSubCategoryObject))
 			{
-				// If the underlying type object's namespace is not imported, auto-import it now into the current editor context.
-				const UObject* PinSubCategoryObject = InNewPinType.PinSubCategoryObject.Get();
-				if (PinSubCategoryObject && BlueprintEditor.IsValid() && BlueprintEditor->IsNonImportedObject(PinSubCategoryObject))
+				FString Namespace = FBlueprintNamespaceUtilities::GetObjectNamespace(PinSubCategoryObject);
+				if (!Namespace.IsEmpty() && BlueprintEditor.IsValid())
 				{
-					FString Namespace = FBlueprintNamespaceUtilities::GetObjectNamespace(PinSubCategoryObject);
-					if (!Namespace.IsEmpty() && BlueprintEditor.IsValid())
-					{
-						BlueprintEditor->ImportNamespace(Namespace);
-					}
+					BlueprintEditor->ImportNamespace(Namespace);
 				}
 			}
 		}
