Auto-import any Blueprint namespaces that are no longer included in the default import set after a reparenting operation.

diff --git a/Engine/Source/Editor/Kismet/Private/BlueprintEditor.cpp b/Engine/Source/Editor/Kismet/Private/BlueprintEditor.cpp
--- a/Engine/Source/Editor/Kismet/Private/BlueprintEditor.cpp
+++ b/Engine/Source/Editor/Kismet/Private/BlueprintEditor.cpp
@@ -3306,9 +3306,22 @@ void FBlueprintEditor::ReparentBlueprint_NewParentChosen(UClass* ChosenClass)
 				}
 			}
 
+			// Gather the set of default imports with the old parent class set.
+			TSet<FString> OldDefaultImports;
+			FBlueprintNamespaceUtilities::GetDefaultImportsForBlueprint(BlueprintObj, OldDefaultImports);
+
 			UClass* OldParentClass = BlueprintObj->ParentClass;
 			BlueprintObj->ParentClass = ChosenClass;
 
+			// Gather the set of default imports with the new parent class set.
+			TSet<FString> NewDefaultImports;
+			FBlueprintNamespaceUtilities::GetDefaultImportsForBlueprint(BlueprintObj, NewDefaultImports);
+
+			// Auto-import any old default imports that no longer appear in the new set.
+			FImportNamespaceExParameters Params;
+			Params.NamespacesToImport = OldDefaultImports.Difference(NewDefaultImports);
+			ImportNamespaceEx(Params);
+
 			// Ensure that the Blueprint is up-to-date (valid SCS etc.) before compiling
 			EnsureBlueprintIsUpToDate(BlueprintObj);
 			FBlueprintEditorUtils::RefreshAllNodes(GetBlueprintObj());
