Display a notification for Blueprint editor actions when they result in a namespace auto-import event.

diff --git a/Engine/Source/Editor/Kismet/Private/BlueprintEditor.cpp b/Engine/Source/Editor/Kismet/Private/BlueprintEditor.cpp
--- a/Engine/Source/Editor/Kismet/Private/BlueprintEditor.cpp
+++ b/Engine/Source/Editor/Kismet/Private/BlueprintEditor.cpp
@@ -2123,6 +2123,28 @@ void FBlueprintEditor::ImportNamespaceEx(const FImportNamespaceExParameters& InP
 
 		// If bound, execute the post-import callback.
 		InParams.OnPostImportCallback.ExecuteIfBound();
+
+		// Display a notification for auto-import events.
+		if (InParams.bIsAutoImport)
+		{
+			const int32 ImportCount = InParams.NamespacesToImport.Num();
+
+			FText NotificationText;
+			if (ImportCount > 1)
+			{
+				FFormatNamedArguments FormatArgs;
+				FormatArgs.Add(TEXT("ImportCount"), ImportCount);
+				NotificationText = FText::Format(LOCTEXT("AutoImportNotification_Multiple", "Imported {ImportCount} namespaces"), FormatArgs);
+			}
+			else
+			{
+				NotificationText = FText::Format(LOCTEXT("AutoImportNotification_Single", "Imported namespace \"{0}\""), FText::FromString(InParams.NamespacesToImport.Array()[0]));
+			}
+
+			FNotificationInfo Notification(NotificationText);
+			Notification.ExpireDuration = 3.0f;
+			FSlateNotificationManager::Get().AddNotification(Notification);
+		}
 	}
 }
 
@@ -3317,8 +3339,9 @@ void FBlueprintEditor::ReparentBlueprint_NewParentChosen(UClass* ChosenClass)
 			TSet<FString> NewDefaultImports;
 			FBlueprintNamespaceUtilities::GetDefaultImportsForBlueprint(BlueprintObj, NewDefaultImports);
 
-			// Auto-import any old default imports that no longer appear in the new set.
+			// Move namespace imports that no longer appear in the default set to the explicit set.
 			FImportNamespaceExParameters Params;
+			Params.bIsAutoImport = false;
 			Params.NamespacesToImport = OldDefaultImports.Difference(NewDefaultImports);
 			ImportNamespaceEx(Params);
 
