Consider Blueprint objects to be authoritative for namespace queries.

diff --git a/Engine/Source/Editor/Kismet/Private/BlueprintNamespaceUtilities.cpp b/Engine/Source/Editor/Kismet/Private/BlueprintNamespaceUtilities.cpp
--- a/Engine/Source/Editor/Kismet/Private/BlueprintNamespaceUtilities.cpp
+++ b/Engine/Source/Editor/Kismet/Private/BlueprintNamespaceUtilities.cpp
@@ -82,26 +82,23 @@ FString FBlueprintNamespaceUtilities::GetObjectNamespace(const UObject* InObject
 			Field = OwnerStruct;
 		}
 
-		if (const FString* TypeNamespace = Field->FindMetaData(FBlueprintMetadata::MD_Namespace))
+		const UBlueprint* Blueprint = nullptr;
+		if (const UClass* Class = Cast<UClass>(Field))
+		{
+			Blueprint = UBlueprint::GetBlueprintFromClass(Class);
+		}
+
+		if (Blueprint)
+		{
+			Namespace = GetObjectNamespace(Blueprint);
+		}
+		else if(const FString * TypeNamespace = Field->FindMetaData(FBlueprintMetadata::MD_Namespace))
 		{
 			Namespace = *TypeNamespace;
 		}
 		else
 		{
-			const UBlueprint* Blueprint = nullptr;
-			if (const UClass* Class = Cast<UClass>(Field))
-			{
-				Blueprint = UBlueprint::GetBlueprintFromClass(Class);
-			}
-
-			if (Blueprint)
-			{
-				Namespace = GetObjectNamespace(Blueprint);
-			}
-			else
-			{
-				Namespace = GetObjectNamespace(Field->GetPackage());
-			}
+			Namespace = GetObjectNamespace(Field->GetPackage());
 		}
 	}
 	else if (const UBlueprint* Blueprint = Cast<UBlueprint>(InObject))
