Deprecate a legacy Blueprint compiler flag in the project config settings and replace it with a console variable so it can still be used for debugging purposes.

diff --git a/Engine/Source/Editor/Kismet/Private/BlueprintCompilationManager.cpp b/Engine/Source/Editor/Kismet/Private/BlueprintCompilationManager.cpp
--- a/Engine/Source/Editor/Kismet/Private/BlueprintCompilationManager.cpp
+++ b/Engine/Source/Editor/Kismet/Private/BlueprintCompilationManager.cpp
@@ -28,7 +28,6 @@
 #include "ProfilingDebugging/ScopedTimers.h"
 #include "Serialization/ArchiveHasReferences.h"
 #include "Serialization/ArchiveReplaceObjectRef.h"
-#include "Settings/BlueprintEditorProjectSettings.h"
 #include "TickableEditorObject.h"
 #include "UObject/MetaData.h"
 #include "UObject/ReferenceChainSearch.h"
@@ -494,14 +493,24 @@ FReinstancingJob::FReinstancingJob(TPair<UClass*, UClass*> InOldToNew)
 {
 }
 
-namespace SkelReinstUtils
+namespace UE::Kismet::BlueprintCompilationManager::Private
 {
-	/** Flag to use the new FBlueprintCompileReinstancer::MoveDependentSkelToReinst and avoid problematic recursion during reinstancing */
-	static bool bEnableSkelReinstUpdate = true;
-	static FAutoConsoleVariableRef CVarEnableSkelReinstUpdate(
-		TEXT("BP.bEnableSkelReinstUpdate"), bEnableSkelReinstUpdate,
-		TEXT("If true the Reinstancing of SKEL classes will use the new FBlueprintCompileReinstancer::MoveDependentSkelToReinst(o(n)) instead of the old MoveSkelCDOAside (o(n^2))"),
-		ECVF_Default);
+	namespace ConsoleVariables
+	{
+		/** Flag to use the new FBlueprintCompileReinstancer::MoveDependentSkelToReinst and avoid problematic recursion during reinstancing */
+		static bool bEnableSkelReinstUpdate = true;
+		static FAutoConsoleVariableRef CVarEnableSkelReinstUpdate(
+			TEXT("BP.bEnableSkelReinstUpdate"), bEnableSkelReinstUpdate,
+			TEXT("If true the Reinstancing of SKEL classes will use the new FBlueprintCompileReinstancer::MoveDependentSkelToReinst(o(n)) instead of the old MoveSkelCDOAside (o(n^2))"),
+			ECVF_Default);
+
+		/** Flag to disable faster compiles for individual blueprints if they have no function signature changes */
+		static bool bForceAllDependenciesToRecompile = false;
+		static FAutoConsoleVariableRef CVarForceAllDependenciesToRecompile(
+			TEXT("BP.bForceAllDependenciesToRecompile"), bForceAllDependenciesToRecompile,
+			TEXT("If true all dependencies will be bytecode-compiled even when all referenced functions have no signature changes. Intended for compiler development/debugging purposes."),
+			ECVF_Default);
+	}
 }
 
 void FBlueprintCompilationManagerImpl::FlushCompilationQueueImpl(bool bSuppressBroadcastCompiled, TArray<UBlueprint*>* BlueprintsCompiled, TArray<UBlueprint*>* BlueprintsCompiledOrSkeletonCompiled, FUObjectSerializeContext* InLoadContext)
@@ -902,6 +911,8 @@ void FBlueprintCompilationManagerImpl::FlushCompilationQueueImpl(bool bSuppressB
 
 		// STAGE VII: safely throw away old skeleton CDOs:
 		{
+			using namespace UE::Kismet::BlueprintCompilationManager;
+
 			TMap<UClass*, UClass*> NewSkeletonToOldSkeleton;
 			for (FCompilerData& CompilerData : CurrentlyCompilingBPs)
 			{
@@ -909,7 +920,7 @@ void FBlueprintCompilationManagerImpl::FlushCompilationQueueImpl(bool bSuppressB
 				UClass* OldSkeletonClass = BP->SkeletonGeneratedClass;
 				if(OldSkeletonClass)
 				{
-					if (SkelReinstUtils::bEnableSkelReinstUpdate)
+					if (Private::ConsoleVariables::bEnableSkelReinstUpdate)
 					{
 						TRACE_CPUPROFILER_EVENT_SCOPE(MoveDependentSkelToReinst);
 
@@ -929,8 +940,7 @@ void FBlueprintCompilationManagerImpl::FlushCompilationQueueImpl(bool bSuppressB
 
 			// if any function signatures have changed in this skeleton class we will need to recompile all dependencies, but if not
 			// then we can avoid dependency recompilation:
-			const UBlueprintEditorProjectSettings* EditorProjectSettings = GetDefault<UBlueprintEditorProjectSettings>();
-			bool bSkipUnneededDependencyCompilation = !EditorProjectSettings->bForceAllDependenciesToRecompile;
+			bool bSkipUnneededDependencyCompilation = !Private::ConsoleVariables::bForceAllDependenciesToRecompile;
 			TSet<UObject*> OldFunctionsWithSignatureChanges;
 
 			for (FCompilerData& CompilerData : CurrentlyCompilingBPs)
diff --git a/Engine/Source/Editor/UnrealEd/Public/Settings/BlueprintEditorProjectSettings.h b/Engine/Source/Editor/UnrealEd/Public/Settings/BlueprintEditorProjectSettings.h
--- a/Engine/Source/Editor/UnrealEd/Public/Settings/BlueprintEditorProjectSettings.h
+++ b/Engine/Source/Editor/UnrealEd/Public/Settings/BlueprintEditorProjectSettings.h
@@ -17,12 +17,15 @@ class UNREALED_API UBlueprintEditorProjectSettings : public UDeveloperSettings
 
 public:
 	/**
-	 * Flag to disable faster compiles for individual blueprints if they have no function signature
-	 * changes. This flag is deprecated! In 4.21 there will be no way to force all dependencies to 
-	 * compile when no changes are detected. Report any issues immediately.
+	 * [DEPRECATED] Flag to disable faster compiles for individual blueprints if they have no function signature
+	 * changes.
+	 * 
+	 * If you still require this functionality for debugging purposes, it has moved to a console variable that can be
+	 * temporarily enabled during an editor session to bypass the fast path: 'BP.bForceAllDependenciesToRecompile 1'.
 	 */
-	UPROPERTY(EditAnywhere, config, Category=Blueprints, DisplayName = "Force All Dependencies To Recompile (DEPRECATED)")
-	uint8 bForceAllDependenciesToRecompile:1;
+	UE_DEPRECATED(5.1, "Fast path dependency compilation is now the standard behavior. As a result, this setting is no longer in use, and will eventually be removed.")
+	UPROPERTY(config)
+	uint8 bForceAllDependenciesToRecompile : 1;
 
 	/** If enabled, the editor will load packages to look for soft references to actors when deleting/renaming them. This can be slow in large projects so disable this to improve performance but increase the chance of breaking blueprints/sequences that use soft actor references */
 	UPROPERTY(EditAnywhere, config, Category=Actors)
