Add support for input pins with multiple connections to the PinValueInspection tooltip

diff --git a/Engine/Source/Editor/Kismet/Private/SPinValueInspector.cpp b/Engine/Source/Editor/Kismet/Private/SPinValueInspector.cpp
--- a/Engine/Source/Editor/Kismet/Private/SPinValueInspector.cpp
+++ b/Engine/Source/Editor/Kismet/Private/SPinValueInspector.cpp
@@ -201,8 +201,23 @@ void SPinValueInspector::PopulateTreeView()
 		return;
 	}
 	
+	if (GraphPin->Direction == EGPD_Input)
+	{
+		for (const UEdGraphPin* Pin : GraphPin->LinkedTo)
+		{
+			AddPinToTreeView(Pin, Blueprint);
+		}
+	}
+	else
+	{
+		AddPinToTreeView(GraphPin, Blueprint);
+	}
+}
+
+void SPinValueInspector::AddPinToTreeView(const UEdGraphPin* Pin, UBlueprint* Blueprint)
+{
 	TSharedPtr<FPropertyInstanceInfo> DebugInfo;
-	const FKismetDebugUtilities::EWatchTextResult WatchStatus = FKismetDebugUtilities::GetDebugInfo(DebugInfo, Blueprint, Blueprint->GetObjectBeingDebugged(), GraphPin);
+	const FKismetDebugUtilities::EWatchTextResult WatchStatus = FKismetDebugUtilities::GetDebugInfo(DebugInfo, Blueprint, Blueprint->GetObjectBeingDebugged(), Pin);
 	if (WatchStatus != FKismetDebugUtilities::EWTR_Valid)
 	{
 		switch (WatchStatus)
@@ -226,7 +241,7 @@ void SPinValueInspector::PopulateTreeView()
 
 	if (ensureMsgf(DebugInfo.IsValid(), TEXT("GetDebugInfo returned EWTR_Valid, but DebugInfo wasn't valid")))
 	{
-		TreeViewWidget->AddTreeItemUnique(SKismetDebugTreeView::MakeWatchLineItem(GraphPin, Blueprint->GetObjectBeingDebugged()));
+		TreeViewWidget->AddTreeItemUnique(SKismetDebugTreeView::MakeWatchLineItem(Pin, Blueprint->GetObjectBeingDebugged()));
 	}
 }
 
diff --git a/Engine/Source/Editor/Kismet/Public/SPinValueInspector.h b/Engine/Source/Editor/Kismet/Public/SPinValueInspector.h
--- a/Engine/Source/Editor/Kismet/Public/SPinValueInspector.h
+++ b/Engine/Source/Editor/Kismet/Public/SPinValueInspector.h
@@ -39,9 +39,12 @@ protected:
 	/** requests the constrained box be resized */
 	void OnExpansionChanged(FDebugTreeItemPtr InItem, bool bItemIsExpanded);
 
-	/** Adds the pin to the tree view */
+	/** Adds relevant pins to the tree view */
 	void PopulateTreeView();
 
+	/** Adds a single pin to the tree view */
+	void AddPinToTreeView(const UEdGraphPin* Pin, UBlueprint* Blueprint);
+
 	/** Sets the current watched pin */
 	void SetPinRef(const FEdGraphPinReference& InPinRef);
 
