Fix source of a potential for an editor crash during Blueprint namespace harvesting.

diff --git a/Engine/Source/Editor/Kismet/Private/BlueprintNamespaceUtilities.cpp b/Engine/Source/Editor/Kismet/Private/BlueprintNamespaceUtilities.cpp
--- a/Engine/Source/Editor/Kismet/Private/BlueprintNamespaceUtilities.cpp
+++ b/Engine/Source/Editor/Kismet/Private/BlueprintNamespaceUtilities.cpp
@@ -177,7 +177,7 @@ void FBlueprintNamespaceUtilities::GetPropertyValueNamespaces(const UStruct* InS
 		{
 			for (TFieldIterator<FProperty> It(StructProperty->Struct); It; ++It)
 			{
-				GetPropertyValueNamespaces(StructProperty->Struct, *It, (*It)->ContainerPtrToValuePtr<uint8>(ValuePtr), OutNamespaces);
+				GetPropertyValueNamespaces(StructProperty->Struct, *It, ValuePtr, OutNamespaces);
 			}
 		}
 		else if (const FArrayProperty* ArrayProperty = CastField<FArrayProperty>(InProperty))
