BP: Add generated function to the BP Skel class.

diff --git a/Engine/Source/Editor/Kismet/Private/BlueprintCompilationManager.cpp b/Engine/Source/Editor/Kismet/Private/BlueprintCompilationManager.cpp
--- a/Engine/Source/Editor/Kismet/Private/BlueprintCompilationManager.cpp
+++ b/Engine/Source/Editor/Kismet/Private/BlueprintCompilationManager.cpp
@@ -2922,6 +2922,7 @@ UClass* FBlueprintCompilationManagerImpl::FastGenerateSkeletonClass(UBlueprint*
 	}
 	
 	AddFunctionForGraphs(TEXT(""), BP->FunctionGraphs, CurrentFieldStorageLocation, BPTYPE_FunctionLibrary == BP->BlueprintType, false);
+	AddFunctionForGraphs(TEXT(""), CompilerContext.GeneratedFunctionGraphs, CurrentFieldStorageLocation, BPTYPE_FunctionLibrary == BP->BlueprintType, false);
 
 	// Add interface functions, often these are added by normal detection of implemented functions, but they won't be
 	// if the interface is added but the function is not implemented:
