Fixed some static analysis warnings.

diff --git a/Engine/Source/Developer/ShaderFormatOpenGL/Private/GlslBackend.cpp b/Engine/Source/Developer/ShaderFormatOpenGL/Private/GlslBackend.cpp
--- a/Engine/Source/Developer/ShaderFormatOpenGL/Private/GlslBackend.cpp
+++ b/Engine/Source/Developer/ShaderFormatOpenGL/Private/GlslBackend.cpp
@@ -4162,7 +4162,12 @@ static ir_rvalue* GenShaderInputSemantic(
 	bool& ApplyFlipFrontFacingAdjustment
 	)
 {
-	if (Semantic && FCStringAnsi::Strnicmp(Semantic, "SV_", 3) == 0)
+	if (Semantic == nullptr)
+	{
+		return nullptr;
+	}
+
+	if (FCStringAnsi::Strnicmp(Semantic, "SV_", 3) == 0)
 	{
 		FSystemValue* SystemValues = SystemValueTable[Frequency];
 		for (int i = 0; SystemValues[i].Semantic != NULL; ++i)
@@ -4291,7 +4296,7 @@ static ir_rvalue* GenShaderInputSemantic(
 
 	// If we're here, no built-in variables matched.
 
-	if (Semantic && FCStringAnsi::Strnicmp(Semantic, "SV_", 3) == 0)
+	if (FCStringAnsi::Strnicmp(Semantic, "SV_", 3) == 0)
 	{
 		_mesa_glsl_warning(ParseState, "unrecognized system "
 			"value input '%s'", Semantic);
diff --git a/Engine/Source/Editor/Kismet/Private/SReplaceNodeReferences.cpp b/Engine/Source/Editor/Kismet/Private/SReplaceNodeReferences.cpp
--- a/Engine/Source/Editor/Kismet/Private/SReplaceNodeReferences.cpp
+++ b/Engine/Source/Editor/Kismet/Private/SReplaceNodeReferences.cpp
@@ -652,7 +652,7 @@ FReply SReplaceNodeReferences::OnFindAndReplaceAll()
 
 void SReplaceNodeReferences::OnSubmitSearchQuery(bool bFindAndReplace)
 {
-	if (HasValidSource())
+	if (SourceProperty != nullptr)
 	{
 		FString SearchTerm;
 
@@ -690,7 +690,7 @@ void SReplaceNodeReferences::FindAllReplacementsComplete(TArray<FImaginaryFiBDat
 		}
 	}
 
-	if (HasValidSource() && SelectedTargetReferenceItem.IsValid() && SourceProperty!=nullptr )
+	if (SourceProperty != nullptr && SelectedTargetReferenceItem.IsValid())
 	{
 		FMemberReference SourceVariableReference;
 		SourceVariableReference.SetFromField<FProperty>(SourceProperty, SourceProperty->GetOwnerClass());
@@ -708,7 +708,7 @@ void SReplaceNodeReferences::FindAllReplacementsComplete(TArray<FImaginaryFiBDat
 
 					FReplaceNodeReferencesHelper::ReplaceReferences(SourceVariableReference, TargetVariableReference, Blueprint, InRawDataList);
 
-					if (SourceProperty && bShowReplacementsWhenFinished)
+					if (SourceProperty != nullptr && bShowReplacementsWhenFinished)
 					{
 						TSharedPtr<SFindInBlueprints> GlobalResults = FFindInBlueprintSearchManager::Get().GetGlobalFindResults();
 						if (GlobalResults)
@@ -756,7 +756,7 @@ void SReplaceNodeReferences::OnSourceSelectionChanged(FTreeViewItem Selection, E
 
 FText SReplaceNodeReferences::GetSourceDisplayText() const
 {
-	if (!HasValidSource())
+	if (SourceProperty == nullptr)
 	{
 		return LOCTEXT("UnselectedSourceReference", "Please select a source reference!");
 	}
@@ -765,7 +765,7 @@ FText SReplaceNodeReferences::GetSourceDisplayText() const
 
 const FSlateBrush* SReplaceNodeReferences::GetSourceReferenceIcon() const
 {
-	if (HasValidSource())
+	if (SourceProperty != nullptr)
 	{
 		return FBlueprintEditorUtils::GetIconFromPin(SourcePinType);
 	}
@@ -781,7 +781,7 @@ FSlateColor SReplaceNodeReferences::GetSourceReferenceIconColor() const
 
 const FSlateBrush* SReplaceNodeReferences::GetSecondarySourceReferenceIcon() const
 {
-	if (HasValidSource())
+	if (SourceProperty != nullptr)
 	{
 		return FBlueprintEditorUtils::GetSecondaryIconFromPin(SourcePinType);
 	}
@@ -795,11 +795,6 @@ FSlateColor SReplaceNodeReferences::GetSecondarySourceReferenceIconColor() const
 	return K2Schema->GetSecondaryPinTypeColor(SourcePinType);
 }
 
-bool SReplaceNodeReferences::HasValidSource() const
-{
-	return (SourceProperty != nullptr);
-}
-
 FText SReplaceNodeReferences::GetFindAllButtonText() const
 {
 	if (bFindWithinBlueprint)
@@ -838,7 +833,7 @@ FText SReplaceNodeReferences::GetFindAndReplaceToolTipText(bool bFindAndReplace)
 	}
 	else
 	{
-		if (!HasValidSource())
+		if (SourceProperty == nullptr)
 		{
 			return LOCTEXT("PickSourceVariable", "Pick a source variable from the menu!");
 		}
@@ -855,7 +850,7 @@ FText SReplaceNodeReferences::GetFindAndReplaceToolTipText(bool bFindAndReplace)
 
 bool SReplaceNodeReferences::CanBeginSearch(bool bFindAndReplace) const
 {
-	bool bCanSearch = HasValidSource() && !IsSearchInProgress();
+	bool bCanSearch = SourceProperty != nullptr && !IsSearchInProgress();
 
 	if (bFindAndReplace)
 	{
diff --git a/Engine/Source/Editor/Kismet/Private/SReplaceNodeReferences.h b/Engine/Source/Editor/Kismet/Private/SReplaceNodeReferences.h
--- a/Engine/Source/Editor/Kismet/Private/SReplaceNodeReferences.h
+++ b/Engine/Source/Editor/Kismet/Private/SReplaceNodeReferences.h
@@ -159,9 +159,6 @@ protected:
 	/** Returns the secondary icon color for the source reference */
 	FSlateColor GetSecondarySourceReferenceIconColor() const;
 
-	/** Returns whether the source property is a valid property to search for and replace */
-	bool HasValidSource() const;
-
 	/** Returns the text for the Find All button */
 	FText GetFindAllButtonText() const;
 
