Fix for a failure to include native Blueprint function library class objects when harvesting Blueprint namespaces.

diff --git a/Engine/Source/Editor/Kismet/Private/BlueprintNamespaceRegistry.cpp b/Engine/Source/Editor/Kismet/Private/BlueprintNamespaceRegistry.cpp
--- a/Engine/Source/Editor/Kismet/Private/BlueprintNamespaceRegistry.cpp
+++ b/Engine/Source/Editor/Kismet/Private/BlueprintNamespaceRegistry.cpp
@@ -209,7 +209,13 @@ void FBlueprintNamespaceRegistry::FindAndRegisterAllNamespaces()
 	for (TObjectIterator<UClass> ClassIt; ClassIt; ++ClassIt)
 	{
 		const UClass* ClassObject = *ClassIt;
-		if (UEdGraphSchema_K2::IsAllowableBlueprintVariableType(ClassObject) && !ExcludedObjects.Contains(ClassObject))
+
+		// Include native function library class objects as well.
+		const bool bShouldRegisterClassObject =
+			UEdGraphSchema_K2::IsAllowableBlueprintVariableType(ClassObject)
+			|| ClassObject->IsChildOf<UBlueprintFunctionLibrary>();
+
+		if (bShouldRegisterClassObject && !ExcludedObjects.Contains(ClassObject))
 		{
 			RegisterNamespace(ClassObject);
 		}
@@ -235,16 +241,6 @@ void FBlueprintNamespaceRegistry::FindAndRegisterAllNamespaces()
 		}
 	}
 
-	// Register loaded function library namespace identifiers.
-	for (TObjectIterator<UBlueprintFunctionLibrary> LibraryIt; LibraryIt; ++LibraryIt)
-	{
-		const UBlueprintFunctionLibrary* LibraryObject = *LibraryIt;
-		if (LibraryObject && !ExcludedObjects.Contains(LibraryObject))
-		{
-			RegisterNamespace(LibraryObject);
-		}
-	}
-
 	FARFilter ClassFilter;
 	ClassFilter.ClassNames.Add(UBlueprint::StaticClass()->GetFName());
 	ClassFilter.ClassNames.Add(UBlueprintGeneratedClass::StaticClass()->GetFName());
