Reflect native class namespace changes on Live Coding completion events.

diff --git a/Engine/Source/Editor/Kismet/Private/BlueprintNamespaceRegistry.cpp b/Engine/Source/Editor/Kismet/Private/BlueprintNamespaceRegistry.cpp
--- a/Engine/Source/Editor/Kismet/Private/BlueprintNamespaceRegistry.cpp
+++ b/Engine/Source/Editor/Kismet/Private/BlueprintNamespaceRegistry.cpp
@@ -59,6 +59,7 @@ void FBlueprintNamespaceRegistry::Initialize()
 
 		FindAndRegisterAllNamespaces();
 
+		OnReloadCompleteDelegateHandle = FCoreUObjectDelegates::ReloadCompleteDelegate.AddRaw(this, &FBlueprintNamespaceRegistry::OnReloadComplete);
 		OnDefaultNamespaceTypeChangedDelegateHandle = FBlueprintNamespaceUtilities::OnDefaultBlueprintNamespaceTypeChanged().AddRaw(this, &FBlueprintNamespaceRegistry::OnDefaultNamespaceTypeChanged);
 	}
 
@@ -72,6 +73,7 @@ void FBlueprintNamespaceRegistry::Shutdown()
 		return;
 	}
 
+	FCoreUObjectDelegates::ReloadCompleteDelegate.Remove(OnReloadCompleteDelegateHandle);
 	FBlueprintNamespaceUtilities::OnDefaultBlueprintNamespaceTypeChanged().Remove(OnDefaultNamespaceTypeChangedDelegateHandle);
 
 	if (FModuleManager::Get().IsModuleLoaded(TEXT("AssetRegistry")))
@@ -141,6 +143,12 @@ void FBlueprintNamespaceRegistry::OnAssetRenamed(const FAssetData& AssetData, co
 	}
 }
 
+void FBlueprintNamespaceRegistry::OnReloadComplete(EReloadCompleteReason InReason)
+{
+	// Rebuild the registry to reflect the appropriate default namespace identifiers for any reloaded types.
+	Rebuild();
+}
+
 bool FBlueprintNamespaceRegistry::IsRegisteredPath(const FString& InPath) const
 {
 	TSharedPtr<FBlueprintNamespacePathTree::FNode> Node = PathTree->FindPathNode(InPath);
diff --git a/Engine/Source/Editor/Kismet/Public/BlueprintNamespaceRegistry.h b/Engine/Source/Editor/Kismet/Public/BlueprintNamespaceRegistry.h
--- a/Engine/Source/Editor/Kismet/Public/BlueprintNamespaceRegistry.h
+++ b/Engine/Source/Editor/Kismet/Public/BlueprintNamespaceRegistry.h
@@ -77,6 +77,9 @@ protected:
 	void DumpAllRegisteredPaths();
 	void OnDefaultNamespaceTypeChanged();
 
+	/** Handler for hot reload / live coding completion events. */
+	void OnReloadComplete(EReloadCompleteReason InReason);
+
 private:
 	/** Indicates whether the registry has been initialized. */
 	bool bIsInitialized;
@@ -85,6 +88,7 @@ private:
 	FDelegateHandle OnAssetAddedDelegateHandle;
 	FDelegateHandle OnAssetRemovedDelegateHandle;
 	FDelegateHandle OnAssetRenamedDelegateHandle;
+	FDelegateHandle OnReloadCompleteDelegateHandle;
 	FDelegateHandle OnDefaultNamespaceTypeChangedDelegateHandle;
 
 	/** Handles storage and retrieval for namespace path identifiers. */
