Add static shutdown method to FPinValueInspectorTooltip to reset static shared pointers (fixes crash on Mac on engine shutdown)

diff --git a/Engine/Source/Editor/Kismet/Private/BlueprintEditorModule.cpp b/Engine/Source/Editor/Kismet/Private/BlueprintEditorModule.cpp
--- a/Engine/Source/Editor/Kismet/Private/BlueprintEditorModule.cpp
+++ b/Engine/Source/Editor/Kismet/Private/BlueprintEditorModule.cpp
@@ -31,6 +31,7 @@
 #include "KismetCompiler.h"
 #include "KismetWidgets.h"
 #include "BlueprintNamespaceRegistry.h"
+#include "SPinValueInspector.h"
 
 #define LOCTEXT_NAMESPACE "BlueprintEditor"
 
@@ -259,6 +260,7 @@ void FBlueprintEditorModule::ShutdownModule()
 	UnregisterSCSEditorCustomization("InstancedStaticMeshComponent");
 
 	UEdGraphPin::ShutdownVerification();
+	FPinValueInspectorTooltip::ShutdownTooltip();
 
 	FBlueprintNamespaceRegistry::Get().Shutdown();
 }
diff --git a/Engine/Source/Editor/Kismet/Private/SPinValueInspector.cpp b/Engine/Source/Editor/Kismet/Private/SPinValueInspector.cpp
--- a/Engine/Source/Editor/Kismet/Private/SPinValueInspector.cpp
+++ b/Engine/Source/Editor/Kismet/Private/SPinValueInspector.cpp
@@ -279,6 +279,13 @@ TWeakPtr<FPinValueInspectorTooltip> FPinValueInspectorTooltip::SummonTooltip(FEd
 	return nullptr;
 }
 
+void FPinValueInspectorTooltip::ShutdownTooltip()
+{
+	TooltipWindow.Reset();
+	TooltipWidget.Reset();
+	ValueInspectorWidget.Reset();
+}
+
 void FPinValueInspectorTooltip::MoveTooltip(const FVector2D& InNewLocation)
 {
 	if (ensureMsgf(TooltipWindow.IsValid(), TEXT("SPinValueInspector::MoveTooltip was called before any tooltip was created")))
diff --git a/Engine/Source/Editor/Kismet/Public/SPinValueInspector.h b/Engine/Source/Editor/Kismet/Public/SPinValueInspector.h
--- a/Engine/Source/Editor/Kismet/Public/SPinValueInspector.h
+++ b/Engine/Source/Editor/Kismet/Public/SPinValueInspector.h
@@ -91,6 +91,9 @@ public:
 	/** Summons a new tooltip in the shared window */
 	static TWeakPtr<FPinValueInspectorTooltip> SummonTooltip(FEdGraphPinReference InPinRef);
 
+	/** Release the references to the static widget shared pointers */
+	static void ShutdownTooltip();
+
 	/** Inspector widget in the tooltip */
 	static TSharedPtr<SPinValueInspector> ValueInspectorWidget;
 
