Child Blueprint classes will now inherit assigned namespaces all the way through the native parent class hierarchy.

Also changed the UX to not show the auto-import alert unless a namespace was added into the current editor context. We technically auto-import a namespace once it's used but we don't necessarily need to tell the user about it in certain scenarios (e.g. when it's already part of the default set). The reason we do it that way is that if it is no longer present in the default set (e.g. after reparenting the child BP), the child will then display it as an explicit import.

#rnx

diff --git a/Engine/Source/Editor/Kismet/Private/BlueprintEditor.cpp b/Engine/Source/Editor/Kismet/Private/BlueprintEditor.cpp
--- a/Engine/Source/Editor/Kismet/Private/BlueprintEditor.cpp
+++ b/Engine/Source/Editor/Kismet/Private/BlueprintEditor.cpp
@@ -2045,8 +2045,8 @@ void FBlueprintEditor::ImportNamespace(const FString& InNamespace)
 void FBlueprintEditor::ImportNamespaceEx(const FImportNamespaceExParameters& InParams)
 {
 	// No auto-import actions if features are disabled.
-	const bool bIsAutoInputEnabled = GetDefault<UBlueprintEditorSettings>()->bEnableNamespaceImportingFeatures;
-	if (InParams.bIsAutoImport && !bIsAutoInputEnabled)
+	const bool bIsAutoImportEnabled = GetDefault<UBlueprintEditorSettings>()->bEnableNamespaceImportingFeatures;
+	if (InParams.bIsAutoImport && !bIsAutoImportEnabled)
 	{
 		return;
 	}
@@ -2072,7 +2072,7 @@ void FBlueprintEditor::ImportNamespaceEx(const FImportNamespaceExParameters& InP
 	};
 
 	// Update the imported set for all edited objects.
-	bool bWasAdded = false;
+	bool bWasAddedAsImport = false;
 	const TArray<UObject*>& EditingObjs = GetEditingObjects();
 	for (UObject* EditingObj : EditingObjs)
 	{
@@ -2081,7 +2081,7 @@ void FBlueprintEditor::ImportNamespaceEx(const FImportNamespaceExParameters& InP
 			// Add each namespace into the Blueprint's user-facing import set.
 			for (const FString& NamespaceToImport : InParams.NamespacesToImport)
 			{
-				bWasAdded |= AddNamespaceToImportList(BlueprintObj, NamespaceToImport);
+				bWasAddedAsImport |= AddNamespaceToImportList(BlueprintObj, NamespaceToImport);
 			}
 		}
 	}
@@ -2101,15 +2101,16 @@ void FBlueprintEditor::ImportNamespaceEx(const FImportNamespaceExParameters& InP
 	// Add to the current scope of the Blueprint's editor context. Note that in certain cases, imports may already be associated
 	// with the Blueprint, but not yet associated with the editor context (e.g. - auto-import after setting a Blueprint's namespace;
 	// we won't add it to the Blueprint's import list, but we still want to add to the editor context and do any post-import actions).
+	bool bWasAddedToEditorContext = false;
 	if (ImportedNamespaceHelper.IsValid())
 	{
 		for (const FString& NamespaceToImport : InParams.NamespacesToImport)
 		{
-			bWasAdded |= AddNamespaceToEditorContext(ImportedNamespaceHelper, NamespaceToImport);
+			bWasAddedToEditorContext |= AddNamespaceToEditorContext(ImportedNamespaceHelper, NamespaceToImport);
 		}
 	}
 
-	if (bWasAdded)
+	if (bWasAddedAsImport || bWasAddedToEditorContext)
 	{
 		// Load additional libraries that may now be in scope.
 		// @todo_namespaces - Make this more targeted - i.e. get/load only those assets tagged w/ the given namespace
@@ -2124,8 +2125,8 @@ void FBlueprintEditor::ImportNamespaceEx(const FImportNamespaceExParameters& InP
 		// If bound, execute the post-import callback.
 		InParams.OnPostImportCallback.ExecuteIfBound();
 
-		// Display a notification for auto-import events.
-		if (InParams.bIsAutoImport)
+		// Display a notification for auto-import events only if we've added it to the editor context.
+		if (InParams.bIsAutoImport && bWasAddedToEditorContext)
 		{
 			const int32 ImportCount = InParams.NamespacesToImport.Num();
 
diff --git a/Engine/Source/Editor/Kismet/Private/BlueprintNamespaceUtilities.cpp b/Engine/Source/Editor/Kismet/Private/BlueprintNamespaceUtilities.cpp
--- a/Engine/Source/Editor/Kismet/Private/BlueprintNamespaceUtilities.cpp
+++ b/Engine/Source/Editor/Kismet/Private/BlueprintNamespaceUtilities.cpp
@@ -255,26 +255,27 @@ void FBlueprintNamespaceUtilities::GetDefaultImportsForBlueprint(const UBlueprin
 
 	const bool bAddParentClassImportedNamespaces = GetDefault<UBlueprintEditorSettings>()->bInheritImportedNamespacesFromParentBP;
 
-	// Inherited Blueprint namespaces (if set).
-	TArray<UBlueprint*> ParentBPStack;
-	UBlueprint::GetBlueprintHierarchyFromClass(InBlueprint->ParentClass, ParentBPStack);
-	while (ParentBPStack.Num() > 0)
+	// Inherited namespaces from the parent class hierarchy (if set).
+	UClass* ParentClass = InBlueprint->ParentClass;
+	while (ParentClass)
 	{
-		if (const UBlueprint* ParentBP = ParentBPStack.Pop())
+		// Parent class namespace (if set).
+		FString ParentClassNamespace = GetObjectNamespace(ParentClass);
+		if (!ParentClassNamespace.IsEmpty())
 		{
-			// Parent Blueprint's namespace (if set).
-			FString ParentBPNamespace = GetObjectNamespace(ParentBP);
-			if (!ParentBPNamespace.IsEmpty())
-			{
-				OutNamespaces.Add(ParentBPNamespace);
-			}
+			OutNamespaces.Add(ParentClassNamespace);
+		}
 
-			// If enabled, also include namespaces that are explicitly imported by all ancestor BPs.
-			if (bAddParentClassImportedNamespaces)
+		// If enabled, also include namespaces that are explicitly imported by all ancestor BPs.
+		if (bAddParentClassImportedNamespaces)
+		{
+			if (const UBlueprint* ParentBP = UBlueprint::GetBlueprintFromClass(ParentClass))
 			{
 				OutNamespaces.Append(ParentBP->ImportedNamespaces);
 			}
 		}
+
+		ParentClass = ParentClass->GetSuperClass();
 	}
 }
 
