Add static shutdown method to FPinValueInspectorTooltip to reset static shared pointers (hotfix version)

diff --git a/Engine/Source/Editor/Kismet/Public/SPinValueInspector.h b/Engine/Source/Editor/Kismet/Public/SPinValueInspector.h
--- a/Engine/Source/Editor/Kismet/Public/SPinValueInspector.h
+++ b/Engine/Source/Editor/Kismet/Public/SPinValueInspector.h
@@ -66,6 +66,8 @@ private:
 class KISMET_API FPinValueInspectorTooltip
 {
 public:
+	friend class FBlueprintEditorModule;
+
 	/** Moves the tooltip to the new location */
 	void MoveTooltip(const FVector2D& InNewLocation);
 
@@ -101,6 +103,9 @@ private:
 	/** Handles Creating a custom tooltip window for all PinValueInspector tooltips */
 	static void CreatePinValueTooltipWindow();
 
+	/** Release the references to the static widget shared pointers */
+	static void ShutdownTooltip();
+
 	/** A reusable tooltip window for PinValueInspector */
 	static TSharedPtr<class SWindow> TooltipWindow;
 
