Effectively backout hotfix version of existing change merged from 5.0 to fix build errors
#rnx

diff --git a/Engine/Source/Editor/Kismet/Public/SPinValueInspector.h b/Engine/Source/Editor/Kismet/Public/SPinValueInspector.h
--- a/Engine/Source/Editor/Kismet/Public/SPinValueInspector.h
+++ b/Engine/Source/Editor/Kismet/Public/SPinValueInspector.h
@@ -66,8 +66,6 @@ private:
 class KISMET_API FPinValueInspectorTooltip
 {
 public:
-	friend class FBlueprintEditorModule;
-
 	/** Moves the tooltip to the new location */
 	void MoveTooltip(const FVector2D& InNewLocation);
 
@@ -103,9 +101,6 @@ private:
 	/** Handles Creating a custom tooltip window for all PinValueInspector tooltips */
 	static void CreatePinValueTooltipWindow();
 
-	/** Release the references to the static widget shared pointers */
-	static void ShutdownTooltip();
-
 	/** A reusable tooltip window for PinValueInspector */
 	static TSharedPtr<class SWindow> TooltipWindow;
 
