Fix in BlueprintDebugger main menu availability when FTabManager doesn't have a MenuMultiBox available.

- AddSetMenuMultiBox() for global menu bar platforms.   This part is a clone of CL 18842821.

This worked on first launch of the window but caused different menu items to appears on closing and reopening the BluprintDebugger window.

- SDockTab doesn't match one used during create on subsequent window opens - this is causing different TabManager MultiBox items to be used on window re-open (i.e the MultiBox items from the global TabManager one instead of the intended TabManger).   To workaround this recreate the debugger on each call to CreateBluprintDebuggerTab() so it then picks the correct tab manager in FGlobalTabmanager::UpdateMainMenu() - this makes it work the same as the NiagaraDebugger and presents the correct menu items each time.

diff --git a/Engine/Source/Editor/Kismet/Private/BlueprintDebugger.cpp b/Engine/Source/Editor/Kismet/Private/BlueprintDebugger.cpp
--- a/Engine/Source/Editor/Kismet/Private/BlueprintDebugger.cpp
+++ b/Engine/Source/Editor/Kismet/Private/BlueprintDebugger.cpp
@@ -103,26 +103,19 @@ TSharedRef<SDockTab> FBlueprintDebuggerImpl::CreateBluprintDebuggerTab(const FSp
 		.TabRole(ETabRole::NomadTab)
 		.Label(NSLOCTEXT("BlueprintDebugger", "TabTitle", "Blueprint Debugger"));
 
-	if (!DebuggingToolsTabManager.IsValid())
-	{
-		DebuggingToolsTabManager = FGlobalTabmanager::Get()->NewTabManager(NomadTab);
-		// on persist layout will handle saving layout if the editor is shut down:
-		DebuggingToolsTabManager->SetOnPersistLayout(
-			FTabManager::FOnPersistLayout::CreateStatic(
-				[](const TSharedRef<FTabManager::FLayout>& InLayout)
+	DebuggingToolsTabManager = FGlobalTabmanager::Get()->NewTabManager(NomadTab);
+	// on persist layout will handle saving layout if the editor is shut down:
+	DebuggingToolsTabManager->SetOnPersistLayout(
+		FTabManager::FOnPersistLayout::CreateStatic(
+			[](const TSharedRef<FTabManager::FLayout>& InLayout)
+			{
+				if (InLayout->GetPrimaryArea().Pin().IsValid())
 				{
-					if (InLayout->GetPrimaryArea().Pin().IsValid())
-					{
-						FLayoutSaveRestore::SaveToConfig(GEditorLayoutIni, InLayout);
-					}
+					FLayoutSaveRestore::SaveToConfig(GEditorLayoutIni, InLayout);
 				}
-			)
-		);
-	}
-	else
-	{
-		ensure(BlueprintDebuggerLayout.IsValid());
-	}
+			}
+		)
+	);
 
 	// Register Toolbar
 	SKismetDebuggingView::TryRegisterDebugToolbar();
@@ -147,45 +140,42 @@ TSharedRef<SDockTab> FBlueprintDebuggerImpl::CreateBluprintDebuggerTab(const FSp
 		, DebuggingToolsTabManagerWeak
 	));
 
-	if (!BlueprintDebuggerLayout.IsValid())
-	{
-		DebuggingToolsTabManager->RegisterTabSpawner(
-			ExecutionFlowTabName,
-			FOnSpawnTab::CreateStatic(
-				[](const FSpawnTabArgs&)->TSharedRef<SDockTab>
-				{
-				const TSharedPtr<SKismetDebuggingView> KismetDebuggingView = SNew(SKismetDebuggingView)
-					.BlueprintToWatch(nullptr);
-				return SNew(SDockTab)
-					.TabRole(ETabRole::PanelTab)
-					.Label_Raw(KismetDebuggingView.Get(), &SKismetDebuggingView::GetTabLabel)
-					[
-						KismetDebuggingView.ToSharedRef()
-					];
-				}
-			)
+	DebuggingToolsTabManager->RegisterTabSpawner(
+		ExecutionFlowTabName,
+		FOnSpawnTab::CreateStatic(
+			[](const FSpawnTabArgs&)->TSharedRef<SDockTab>
+			{
+			const TSharedPtr<SKismetDebuggingView> KismetDebuggingView = SNew(SKismetDebuggingView)
+				.BlueprintToWatch(nullptr);
+			return SNew(SDockTab)
+				.TabRole(ETabRole::PanelTab)
+				.Label_Raw(KismetDebuggingView.Get(), &SKismetDebuggingView::GetTabLabel)
+				[
+					KismetDebuggingView.ToSharedRef()
+				];
+			}
 		)
-		.SetDisplayName(NSLOCTEXT("BlueprintDebugger", "ExecutionFlowTabTitle", "Blueprint Data Flow"))
-		.SetTooltipText(NSLOCTEXT("BlueprintDebugger", "ExecutionFlowTooltipText", "Open the Blueprint Data Flow tab."));
-
-		CallStackViewer::RegisterTabSpawner(*DebuggingToolsTabManager); 
-
-		BlueprintDebuggerLayout = FTabManager::NewLayout("Standalone_BlueprintDebugger_Layout_v2")
-		->AddArea
+	)
+	.SetDisplayName(NSLOCTEXT("BlueprintDebugger", "ExecutionFlowTabTitle", "Blueprint Data Flow"))
+	.SetTooltipText(NSLOCTEXT("BlueprintDebugger", "ExecutionFlowTooltipText", "Open the Blueprint Data Flow tab."));
+
+	CallStackViewer::RegisterTabSpawner(*DebuggingToolsTabManager);
+
+	BlueprintDebuggerLayout = FTabManager::NewLayout("Standalone_BlueprintDebugger_Layout_v2")
+	->AddArea
+	(
+		FTabManager::NewPrimaryArea()
+		->SetOrientation(Orient_Vertical)
+		->Split
 		(
-			FTabManager::NewPrimaryArea()
-			->SetOrientation(Orient_Vertical)
-			->Split
-			(
-				FTabManager::NewStack()
-				->SetSizeCoefficient(.4f)
-				->SetHideTabWell(true)
-				->AddTab(CallStackTabName, ETabState::OpenedTab)
-				->AddTab(ExecutionFlowTabName, ETabState::OpenedTab)
-				->SetForegroundTab(CallStackTabName)
-			)
-		);
-	}
+			FTabManager::NewStack()
+			->SetSizeCoefficient(.4f)
+			->SetHideTabWell(true)
+			->AddTab(CallStackTabName, ETabState::OpenedTab)
+			->AddTab(ExecutionFlowTabName, ETabState::OpenedTab)
+			->SetForegroundTab(CallStackTabName)
+		)
+	);
 
 	BlueprintDebuggerLayout = FLayoutSaveRestore::LoadFromConfig(GEditorLayoutIni, BlueprintDebuggerLayout.ToSharedRef());
 
@@ -266,6 +256,7 @@ TSharedRef<SDockTab> FBlueprintDebuggerImpl::CreateBluprintDebuggerTab(const FSp
 			Builder.AddMenuEntry(FBlueprintDebuggerCommands::Get().ShowExecutionTrace);
 			})
 	);
+	
 
 	TSharedRef<SWidget> MenuBarWidget = MenuBarBuilder.MakeWidget();
 
@@ -286,6 +277,7 @@ TSharedRef<SDockTab> FBlueprintDebuggerImpl::CreateBluprintDebuggerTab(const FSp
 			]
 		]
 	);
+	
 
 	// Tell tab-manager about the multi-box for platforms with a global menu bar
 	DebuggingToolsTabManager->SetMenuMultiBox(MenuBarBuilder.GetMultiBox(), MenuBarWidget);
