Fix static analysis warning fallout from CL# 19710337
#rnx

diff --git a/Engine/Source/Editor/Kismet/Private/SReplaceNodeReferences.cpp b/Engine/Source/Editor/Kismet/Private/SReplaceNodeReferences.cpp
--- a/Engine/Source/Editor/Kismet/Private/SReplaceNodeReferences.cpp
+++ b/Engine/Source/Editor/Kismet/Private/SReplaceNodeReferences.cpp
@@ -708,7 +708,7 @@ void SReplaceNodeReferences::FindAllReplacementsComplete(TArray<FImaginaryFiBDat
 
 					FReplaceNodeReferencesHelper::ReplaceReferences(SourceVariableReference, TargetVariableReference, Blueprint, InRawDataList);
 
-					if (SourceProperty != nullptr && bShowReplacementsWhenFinished)
+					if (bShowReplacementsWhenFinished)
 					{
 						TSharedPtr<SFindInBlueprints> GlobalResults = FFindInBlueprintSearchManager::Get().GetGlobalFindResults();
 						if (GlobalResults)
