Fix unit test as a result of CL# 19718761
#rnx

diff --git a/Engine/Source/Editor/Kismet/Private/Tests/BlueprintEditorTests.cpp b/Engine/Source/Editor/Kismet/Private/Tests/BlueprintEditorTests.cpp
--- a/Engine/Source/Editor/Kismet/Private/Tests/BlueprintEditorTests.cpp
+++ b/Engine/Source/Editor/Kismet/Private/Tests/BlueprintEditorTests.cpp
@@ -1119,7 +1119,7 @@ namespace BlueprintEditorPromotionTestHelper
 
 				Test->AddInfo(TEXT("Promoted the return pin on the add mesh node to a variable"));
 
-				const FName OldVarName(TEXT("NewVar_0")); // Default variable name
+				const FName OldVarName(TEXT("NewVar")); // Default variable name
 				const FName NewVarName(TEXT("MyMesh"));
 				FBlueprintEditorUtils::RenameMemberVariable(BlueprintObject, OldVarName, NewVarName);
 				Test->TestNotEqual(TEXT("New variable was renamed"), FBlueprintEditorUtils::FindMemberVariableGuidByName(BlueprintObject, OldVarName), FBlueprintEditorUtils::FindMemberVariableGuidByName(BlueprintObject, NewVarName));
