Fix compilation manager duplicating a subobject when one of the same name already exists

diff --git a/Engine/Source/Editor/Kismet/Private/BlueprintCompilationManager.cpp b/Engine/Source/Editor/Kismet/Private/BlueprintCompilationManager.cpp
--- a/Engine/Source/Editor/Kismet/Private/BlueprintCompilationManager.cpp
+++ b/Engine/Source/Editor/Kismet/Private/BlueprintCompilationManager.cpp
@@ -2365,8 +2365,16 @@ void FBlueprintCompilationManagerImpl::ReinstanceBatch(TArray<FReinstancingJob>&
 							continue;
 						}
 
+						// Even if the object created the DSO itself, we want to carry over the existing one from the old archetype, so rename the newly constructed one out of the way
+						if (UObject* ExistingObject = static_cast<UObject*>(FindObjectWithOuter(NewArchetype, nullptr, Subobject->GetFName())))
+						{
+							UClass* ExistingObjectClass = ExistingObject->GetClass();
+							UObject* TransientOuterForRename = GetTransientOuterForRename(ExistingObjectClass);
+							ExistingObject->Rename(*MakeUniqueObjectName(TransientOuterForRename, ExistingObjectClass).ToString(), TransientOuterForRename, REN_DoNotDirty | REN_DontCreateRedirectors | REN_ForceNoResetLoaders);
+						}
+
 						// Non DSO subobject - just reuse the subobject:
-						Subobject->Rename( 
+						Subobject->Rename(
 							nullptr, 
 							// destination:
 							NewArchetype, 
