Fix a regression that prevents users from editing an interface graph function signature when a terminator node is selected in the read-only graph view.

diff --git a/Engine/Source/Editor/Kismet/Private/SKismetInspector.cpp b/Engine/Source/Editor/Kismet/Private/SKismetInspector.cpp
--- a/Engine/Source/Editor/Kismet/Private/SKismetInspector.cpp
+++ b/Engine/Source/Editor/Kismet/Private/SKismetInspector.cpp
@@ -30,6 +30,7 @@
 #include "K2Node_VariableGet.h"
 #include "K2Node_VariableSet.h"
 #include "K2Node_AddComponent.h" // for GetTemplateFromNode()
+#include "K2Node_FunctionTerminator.h"
 #include "IDetailCustomization.h"
 #include "Editor.h"
 #include "PropertyEditorModule.h"
@@ -928,8 +929,13 @@ bool SKismetInspector::IsPropertyEditingEnabled() const
 			{
 				if(BlueprintEditorPtr.IsValid() && !BlueprintEditorPtr.Pin()->IsEditable(OuterGraph))
 				{
-					bIsEditable = false;
-					break;
+					// Allow property editing on interface function graph terminator nodes (i.e. allow users to modify the interface function signature).
+					const bool bIsInterfaceGraphTerminatorNode = FBlueprintEditorUtils::IsInterfaceGraph(OuterGraph) && EdGraphNode->IsA<UK2Node_FunctionTerminator>();
+					if (!bIsInterfaceGraphTerminatorNode)
+					{
+						bIsEditable = false;
+						break;
+					}
 				}
 			}
 		}
