Fixed structs with a native identical function failing to compare correctly in the Blueprint diff

It would only compare property data, so missed any non-property state that the custom identical function might be comparing.

#rnx

diff --git a/Engine/Source/Editor/Kismet/Private/DiffUtils.cpp b/Engine/Source/Editor/Kismet/Private/DiffUtils.cpp
--- a/Engine/Source/Editor/Kismet/Private/DiffUtils.cpp
+++ b/Engine/Source/Editor/Kismet/Private/DiffUtils.cpp
@@ -462,18 +462,29 @@ static void IdenticalHelper(const FProperty* AProperty, const FProperty* BProper
 	if (APropAsStruct != nullptr)
 	{
 		const FStructProperty* BPropAsStruct = CastFieldChecked<FStructProperty>(const_cast<FProperty*>(BProperty));
-		if (APropAsStruct->Struct->StructFlags & STRUCT_IdenticalNative && BPropAsStruct->Struct != APropAsStruct->Struct)
+		if (BPropAsStruct->Struct == APropAsStruct->Struct)
 		{
-			// If the struct uses CPP identical tests, then we can't dig into it, and we already know it's not identical from the test when we started
-			DifferingSubProperties.Push(RootPath);
+			if (APropAsStruct->Struct->StructFlags & STRUCT_IdenticalNative)
+			{
+				// If the struct uses CPP identical tests then we need to honor that
+				if (!AProperty->Identical(AValue, BValue, PPF_DeepComparison))
+				{
+					DifferingSubProperties.Push(RootPath);
+				}
+			}
+			else
+			{
+				// Compare sub-properties to detect more granular changes
+				for (TFieldIterator<FProperty> PropertyIt(APropAsStruct->Struct); PropertyIt; ++PropertyIt)
+				{
+					const FProperty* StructProp = *PropertyIt;
+					IdenticalHelper(StructProp, StructProp, StructProp->ContainerPtrToValuePtr<void>(AValue, 0), StructProp->ContainerPtrToValuePtr<void>(BValue, 0), FPropertySoftPath(RootPath, StructProp), DifferingSubProperties);
+				}
+			}
 		}
 		else
 		{
-			for (TFieldIterator<FProperty> PropertyIt(APropAsStruct->Struct); PropertyIt; ++PropertyIt)
-			{
-				const FProperty* StructProp = *PropertyIt;
-				IdenticalHelper(StructProp, StructProp, StructProp->ContainerPtrToValuePtr<void>(AValue, 0), StructProp->ContainerPtrToValuePtr<void>(BValue, 0), FPropertySoftPath(RootPath, StructProp), DifferingSubProperties);
-			}
+			DifferingSubProperties.Push(RootPath);
 		}
 	}
 	else if (APropAsArray != nullptr)
