AssetRegistry includes (Engine/Source): change #include "AssetData.h" -> #include "AssetRegistry/AssetData.h", and similar for the other moved AssetRegistry headers.
#rnx

diff --git a/Engine/Source/Developer/AssetTools/Private/AdvancedCopyCustomization.cpp b/Engine/Source/Developer/AssetTools/Private/AdvancedCopyCustomization.cpp
--- a/Engine/Source/Developer/AssetTools/Private/AdvancedCopyCustomization.cpp
+++ b/Engine/Source/Developer/AssetTools/Private/AdvancedCopyCustomization.cpp
@@ -2,7 +2,7 @@
 
 #include "AdvancedCopyCustomization.h"
 #include "Containers/UnrealString.h"
-#include "AssetRegistryModule.h"
+#include "AssetRegistry/AssetRegistryModule.h"
 #include "Interfaces/IPluginManager.h"
 #include "Engine/World.h"
 #include "Engine/Level.h"
diff --git a/Engine/Source/Developer/AssetTools/Private/AssetFixUpRedirectors.cpp b/Engine/Source/Developer/AssetTools/Private/AssetFixUpRedirectors.cpp
--- a/Engine/Source/Developer/AssetTools/Private/AssetFixUpRedirectors.cpp
+++ b/Engine/Source/Developer/AssetTools/Private/AssetFixUpRedirectors.cpp
@@ -16,7 +16,7 @@
 #include "FileHelpers.h"
 #include "SDiscoveringAssetsDialog.h"
 #include "AssetRenameManager.h"
-#include "AssetRegistryModule.h"
+#include "AssetRegistry/AssetRegistryModule.h"
 #include "ICollectionManager.h"
 #include "CollectionManagerModule.h"
 #include "ObjectTools.h"
diff --git a/Engine/Source/Developer/AssetTools/Private/AssetTools.h b/Engine/Source/Developer/AssetTools/Private/AssetTools.h
--- a/Engine/Source/Developer/AssetTools/Private/AssetTools.h
+++ b/Engine/Source/Developer/AssetTools/Private/AssetTools.h
@@ -5,7 +5,7 @@
 #include "CoreMinimal.h"
 #include "IAssetTools.h"
 #include "IAssetTypeActions.h"
-#include "AssetData.h"
+#include "AssetRegistry/AssetData.h"
 #include "AssetRenameManager.h"
 #include "Misc/NamePermissionList.h"
 #include "UObject/StrongObjectPtr.h"
diff --git a/Engine/Source/Developer/AssetTools/Private/AssetTypeActions/AssetTypeActions_Class.cpp b/Engine/Source/Developer/AssetTools/Private/AssetTypeActions/AssetTypeActions_Class.cpp
--- a/Engine/Source/Developer/AssetTools/Private/AssetTypeActions/AssetTypeActions_Class.cpp
+++ b/Engine/Source/Developer/AssetTools/Private/AssetTypeActions/AssetTypeActions_Class.cpp
@@ -2,7 +2,7 @@
 
 #include "AssetTypeActions/AssetTypeActions_Class.h"
 #include "ToolMenus.h"
-#include "AssetData.h"
+#include "AssetRegistry/AssetData.h"
 #include "HAL/FileManager.h"
 #include "EditorStyleSet.h"
 #include "SourceCodeNavigation.h"
diff --git a/Engine/Source/Developer/AssetTools/Private/AssetTypeActions/AssetTypeActions_MaterialFunction.cpp b/Engine/Source/Developer/AssetTools/Private/AssetTypeActions/AssetTypeActions_MaterialFunction.cpp
--- a/Engine/Source/Developer/AssetTools/Private/AssetTypeActions/AssetTypeActions_MaterialFunction.cpp
+++ b/Engine/Source/Developer/AssetTools/Private/AssetTypeActions/AssetTypeActions_MaterialFunction.cpp
@@ -13,8 +13,8 @@
 #include "MaterialEditorModule.h"
 #include "ContentBrowserModule.h"
 #include "IContentBrowserSingleton.h"
-#include "ARFilter.h"
-#include "AssetData.h"
+#include "AssetRegistry/ARFilter.h"
+#include "AssetRegistry/AssetData.h"
 #include "Misc/PackageName.h"
 
 #define LOCTEXT_NAMESPACE "AssetTypeActions"
diff --git a/Engine/Source/Developer/AssetTools/Private/AssetViewUtils.cpp b/Engine/Source/Developer/AssetTools/Private/AssetViewUtils.cpp
--- a/Engine/Source/Developer/AssetTools/Private/AssetViewUtils.cpp
+++ b/Engine/Source/Developer/AssetTools/Private/AssetViewUtils.cpp
@@ -22,8 +22,8 @@
 #include "ISourceControlModule.h"
 #include "SourceControlHelpers.h"
 #include "FileHelpers.h"
-#include "ARFilter.h"
-#include "AssetRegistryModule.h"
+#include "AssetRegistry/ARFilter.h"
+#include "AssetRegistry/AssetRegistryModule.h"
 #include "IAssetTools.h"
 #include "AssetToolsModule.h"
 #include "Settings/EditorExperimentalSettings.h"
diff --git a/Engine/Source/Developer/AssetTools/Public/AdvancedCopyCustomization.h b/Engine/Source/Developer/AssetTools/Public/AdvancedCopyCustomization.h
--- a/Engine/Source/Developer/AssetTools/Public/AdvancedCopyCustomization.h
+++ b/Engine/Source/Developer/AssetTools/Public/AdvancedCopyCustomization.h
@@ -4,7 +4,7 @@
 
 #include "CoreMinimal.h"
 #include "UObject/ObjectMacros.h"
-#include "ARFilter.h"
+#include "AssetRegistry/ARFilter.h"
 #include "IAssetTools.h"
 #include "AdvancedCopyCustomization.generated.h"
 
diff --git a/Engine/Source/Developer/AssetTools/Public/AssetTypeActions_Base.h b/Engine/Source/Developer/AssetTools/Public/AssetTypeActions_Base.h
--- a/Engine/Source/Developer/AssetTools/Public/AssetTypeActions_Base.h
+++ b/Engine/Source/Developer/AssetTools/Public/AssetTypeActions_Base.h
@@ -12,7 +12,7 @@
 #include "IAssetTypeActions.h"
 #include "AssetToolsModule.h"
 #include "Toolkits/SimpleAssetEditor.h"
-#include "ARFilter.h"
+#include "AssetRegistry/ARFilter.h"
 
 struct FAssetData;
 struct FARFilter;
diff --git a/Engine/Source/Developer/AutomationWindow/Private/SAutomationTestItem.cpp b/Engine/Source/Developer/AutomationWindow/Private/SAutomationTestItem.cpp
--- a/Engine/Source/Developer/AutomationWindow/Private/SAutomationTestItem.cpp
+++ b/Engine/Source/Developer/AutomationWindow/Private/SAutomationTestItem.cpp
@@ -16,10 +16,10 @@
 #include "SAutomationWindow.h"
 
 #if WITH_EDITOR
-	#include "AssetData.h"
+	#include "AssetRegistry/AssetData.h"
 	#include "EngineGlobals.h"
 	#include "Editor.h"
-	#include "AssetRegistryModule.h"
+	#include "AssetRegistry/AssetRegistryModule.h"
 	#include "Dialogs/Dialogs.h"
 	#include "SKismetInspector.h"
 #endif
diff --git a/Engine/Source/Developer/AutomationWindow/Private/SAutomationWindow.h b/Engine/Source/Developer/AutomationWindow/Private/SAutomationWindow.h
--- a/Engine/Source/Developer/AutomationWindow/Private/SAutomationWindow.h
+++ b/Engine/Source/Developer/AutomationWindow/Private/SAutomationWindow.h
@@ -18,7 +18,7 @@
 #include "Developer/AutomationWindow/Private/SAutomationTestTreeView.h"
 
 #if WITH_EDITOR
-#include "IAssetRegistry.h"
+#include "AssetRegistry/IAssetRegistry.h"
 #endif
 
 class FAutomationFilter;
diff --git a/Engine/Source/Developer/CookedEditor/Private/CookedEditorPackageManager.cpp b/Engine/Source/Developer/CookedEditor/Private/CookedEditorPackageManager.cpp
--- a/Engine/Source/Developer/CookedEditor/Private/CookedEditorPackageManager.cpp
+++ b/Engine/Source/Developer/CookedEditor/Private/CookedEditorPackageManager.cpp
@@ -3,9 +3,9 @@
 #include "CookedEditorPackageManager.h"
 #include "CoreMinimal.h"
 #include "Logging/LogMacros.h"
-#include "AssetRegistryModule.h"
-#include "IAssetRegistry.h"
-#include "AssetRegistryModule.h"
+#include "AssetRegistry/AssetRegistryModule.h"
+#include "AssetRegistry/IAssetRegistry.h"
+#include "AssetRegistry/AssetRegistryModule.h"
 #include "Engine/AssetManager.h"
 #include "GameDelegates.h"
 
diff --git a/Engine/Source/Developer/FunctionalTesting/Private/FunctionalTestingModule.cpp b/Engine/Source/Developer/FunctionalTesting/Private/FunctionalTestingModule.cpp
--- a/Engine/Source/Developer/FunctionalTesting/Private/FunctionalTestingModule.cpp
+++ b/Engine/Source/Developer/FunctionalTesting/Private/FunctionalTestingModule.cpp
@@ -9,8 +9,8 @@
 #include "EngineUtils.h"
 #include "FunctionalTest.h"
 #include "EngineGlobals.h"
-#include "ARFilter.h"
-#include "AssetRegistryModule.h"
+#include "AssetRegistry/ARFilter.h"
+#include "AssetRegistry/AssetRegistryModule.h"
 #include "Misc/CommandLine.h"
 #include "IAutomationControllerModule.h"
 
diff --git a/Engine/Source/Developer/SourceControl/Private/SourceControlAssetDataCache.cpp b/Engine/Source/Developer/SourceControl/Private/SourceControlAssetDataCache.cpp
--- a/Engine/Source/Developer/SourceControl/Private/SourceControlAssetDataCache.cpp
+++ b/Engine/Source/Developer/SourceControl/Private/SourceControlAssetDataCache.cpp
@@ -3,9 +3,9 @@
 #include "SourceControlAssetDataCache.h"
 
 #include "Async/Async.h"
-#include "AssetData.h"
+#include "AssetRegistry/AssetData.h"
 #include "HAL/IConsoleManager.h"
-#include "IAssetRegistry.h"
+#include "AssetRegistry/IAssetRegistry.h"
 #include "ISourceControlModule.h"
 #include "ISourceControlState.h"
 #include "SourceControlHelpers.h"
diff --git a/Engine/Source/Editor/AIGraph/Private/AIGraphTypes.cpp b/Engine/Source/Editor/AIGraph/Private/AIGraphTypes.cpp
--- a/Engine/Source/Editor/AIGraph/Private/AIGraphTypes.cpp
+++ b/Engine/Source/Editor/AIGraph/Private/AIGraphTypes.cpp
@@ -10,12 +10,12 @@
 #include "Misc/PackageName.h"
 #include "UObject/ConstructorHelpers.h"
 #include "Engine/Blueprint.h"
-#include "AssetData.h"
+#include "AssetRegistry/AssetData.h"
 #include "Editor.h"
 #include "ObjectEditorUtils.h"
 #include "Logging/MessageLog.h"
-#include "ARFilter.h"
-#include "AssetRegistryModule.h"
+#include "AssetRegistry/ARFilter.h"
+#include "AssetRegistry/AssetRegistryModule.h"
 
 #define LOCTEXT_NAMESPACE "AIGraph"
 
diff --git a/Engine/Source/Editor/AnimationEditor/Private/AnimationEditor.cpp b/Engine/Source/Editor/AnimationEditor/Private/AnimationEditor.cpp
--- a/Engine/Source/Editor/AnimationEditor/Private/AnimationEditor.cpp
+++ b/Engine/Source/Editor/AnimationEditor/Private/AnimationEditor.cpp
@@ -10,7 +10,7 @@
 #include "Animation/SmartName.h"
 #include "Animation/AnimationAsset.h"
 #include "Animation/DebugSkelMeshComponent.h"
-#include "AssetData.h"
+#include "AssetRegistry/AssetData.h"
 #include "EdGraph/EdGraphSchema.h"
 #include "Animation/AnimSequence.h"
 #include "Animation/AnimMontage.h"
@@ -36,7 +36,7 @@
 #include "IContentBrowserSingleton.h"
 #include "ContentBrowserModule.h"
 #include "AnimationEditorUtils.h"
-#include "AssetRegistryModule.h"
+#include "AssetRegistry/AssetRegistryModule.h"
 #include "IAssetFamily.h"
 #include "IAnimationSequenceBrowser.h"
 #include "Framework/Notifications/NotificationManager.h"
diff --git a/Engine/Source/Editor/Blutility/Private/BlutilityMenuExtensions.cpp b/Engine/Source/Editor/Blutility/Private/BlutilityMenuExtensions.cpp
--- a/Engine/Source/Editor/Blutility/Private/BlutilityMenuExtensions.cpp
+++ b/Engine/Source/Editor/Blutility/Private/BlutilityMenuExtensions.cpp
@@ -1,7 +1,7 @@
 // Copyright Epic Games, Inc. All Rights Reserved.
 
 #include "BlutilityMenuExtensions.h"
-#include "AssetRegistryModule.h"
+#include "AssetRegistry/AssetRegistryModule.h"
 #include "EditorUtilityBlueprint.h"
 #include "Misc/PackageName.h"
 
@@ -21,7 +21,7 @@
 #include "Widgets/Layout/SScrollBox.h"
 #include "ScopedTransaction.h"
 #include "Subsystems/AssetEditorSubsystem.h"
-#include "AssetData.h"
+#include "AssetRegistry/AssetData.h"
 #include "PropertyPathHelpers.h"
 
 #include "UObject/NoExportTypes.h"
diff --git a/Engine/Source/Editor/Blutility/Private/EditorUtilitySubsystem.cpp b/Engine/Source/Editor/Blutility/Private/EditorUtilitySubsystem.cpp
--- a/Engine/Source/Editor/Blutility/Private/EditorUtilitySubsystem.cpp
+++ b/Engine/Source/Editor/Blutility/Private/EditorUtilitySubsystem.cpp
@@ -9,9 +9,9 @@
 #include "IBlutilityModule.h"
 #include "EditorUtilityWidget.h"
 #include "ScopedTransaction.h"
-#include "ARFilter.h"
-#include "AssetRegistryModule.h"
-#include "IAssetRegistry.h"
+#include "AssetRegistry/ARFilter.h"
+#include "AssetRegistry/AssetRegistryModule.h"
+#include "AssetRegistry/IAssetRegistry.h"
 #include "EditorUtilityTask.h"
 
 #define LOCTEXT_NAMESPACE "EditorUtilitySubsystem"
diff --git a/Engine/Source/Editor/ClassViewer/Private/SClassViewer.cpp b/Engine/Source/Editor/ClassViewer/Private/SClassViewer.cpp
--- a/Engine/Source/Editor/ClassViewer/Private/SClassViewer.cpp
+++ b/Engine/Source/Editor/ClassViewer/Private/SClassViewer.cpp
@@ -27,7 +27,7 @@
 #include "Engine/BlueprintCore.h"
 #include "Engine/Blueprint.h"
 #include "Engine/Brush.h"
-#include "AssetData.h"
+#include "AssetRegistry/AssetData.h"
 #include "Editor/UnrealEdEngine.h"
 #include "Animation/AnimBlueprint.h"
 #include "Engine/BlueprintGeneratedClass.h"
@@ -41,7 +41,7 @@
 #include "ContentBrowserDataDragDropOp.h"
 
 #include "IAssetTools.h"
-#include "ARFilter.h"
+#include "AssetRegistry/ARFilter.h"
 #include "IContentBrowserSingleton.h"
 #include "ContentBrowserModule.h"
 
@@ -51,7 +51,7 @@
 #include "PackageTools.h"
 #include "Logging/MessageLog.h"
 
-#include "AssetRegistryModule.h"
+#include "AssetRegistry/AssetRegistryModule.h"
 #include "AssetToolsModule.h"
 
 #include "ClassViewerNode.h"
diff --git a/Engine/Source/Editor/ContentBrowser/Private/CollectionAssetRegistryBridge.cpp b/Engine/Source/Editor/ContentBrowser/Private/CollectionAssetRegistryBridge.cpp
--- a/Engine/Source/Editor/ContentBrowser/Private/CollectionAssetRegistryBridge.cpp
+++ b/Engine/Source/Editor/ContentBrowser/Private/CollectionAssetRegistryBridge.cpp
@@ -6,8 +6,8 @@
 #include "UObject/Class.h"
 #include "Misc/PackageName.h"
 #include "UObject/ConstructorHelpers.h"
-#include "AssetData.h"
-#include "AssetRegistryModule.h"
+#include "AssetRegistry/AssetData.h"
+#include "AssetRegistry/AssetRegistryModule.h"
 #include "CollectionManagerTypes.h"
 #include "ICollectionManager.h"
 #include "CollectionManagerModule.h"
diff --git a/Engine/Source/Editor/ContentBrowser/Private/PathContextMenu.cpp b/Engine/Source/Editor/ContentBrowser/Private/PathContextMenu.cpp
--- a/Engine/Source/Editor/ContentBrowser/Private/PathContextMenu.cpp
+++ b/Engine/Source/Editor/ContentBrowser/Private/PathContextMenu.cpp
@@ -19,11 +19,11 @@
 #include "SourceControlOperations.h"
 #include "ISourceControlModule.h"
 #include "SourceControlHelpers.h"
-#include "AssetData.h"
+#include "AssetRegistry/AssetData.h"
 #include "Editor.h"
 #include "FileHelpers.h"
-#include "ARFilter.h"
-#include "AssetRegistryModule.h"
+#include "AssetRegistry/ARFilter.h"
+#include "AssetRegistry/AssetRegistryModule.h"
 #include "IAssetTools.h"
 #include "AssetToolsModule.h"
 #include "ContentBrowserLog.h"
diff --git a/Engine/Source/Editor/ContentBrowser/Private/SAssetPicker.h b/Engine/Source/Editor/ContentBrowser/Private/SAssetPicker.h
--- a/Engine/Source/Editor/ContentBrowser/Private/SAssetPicker.h
+++ b/Engine/Source/Editor/ContentBrowser/Private/SAssetPicker.h
@@ -7,8 +7,8 @@
 #include "Widgets/DeclarativeSyntaxSupport.h"
 #include "Widgets/SWidget.h"
 #include "Widgets/SCompoundWidget.h"
-#include "AssetData.h"
-#include "ARFilter.h"
+#include "AssetRegistry/AssetData.h"
+#include "AssetRegistry/ARFilter.h"
 #include "IContentBrowserSingleton.h"
 #include "SourcesData.h"
 
diff --git a/Engine/Source/Editor/ContentBrowser/Private/SAssetView.cpp b/Engine/Source/Editor/ContentBrowser/Private/SAssetView.cpp
--- a/Engine/Source/Editor/ContentBrowser/Private/SAssetView.cpp
+++ b/Engine/Source/Editor/ContentBrowser/Private/SAssetView.cpp
@@ -48,7 +48,7 @@
 #include "Misc/FileHelper.h"
 #include "Misc/TextFilterUtils.h"
 #include "Misc/NamePermissionList.h"
-#include "AssetRegistryState.h"
+#include "AssetRegistry/AssetRegistryState.h"
 #include "Materials/Material.h"
 #include "ContentBrowserMenuContexts.h"
 #include "ContentBrowserUtils.h"
diff --git a/Engine/Source/Editor/ContentBrowser/Private/SContentBrowser.cpp b/Engine/Source/Editor/ContentBrowser/Private/SContentBrowser.cpp
--- a/Engine/Source/Editor/ContentBrowser/Private/SContentBrowser.cpp
+++ b/Engine/Source/Editor/ContentBrowser/Private/SContentBrowser.cpp
@@ -35,8 +35,8 @@
 #include "Settings/EditorSettings.h"
 #include "Editor.h"
 #include "FileHelpers.h"
-#include "AssetRegistryModule.h"
-#include "AssetRegistryState.h"
+#include "AssetRegistry/AssetRegistryModule.h"
+#include "AssetRegistry/AssetRegistryState.h"
 #include "AssetToolsModule.h"
 #include "Widgets/Navigation/SBreadcrumbTrail.h"
 #include "ContentBrowserLog.h"
diff --git a/Engine/Source/Editor/ContentBrowser/Private/SFilterList.h b/Engine/Source/Editor/ContentBrowser/Private/SFilterList.h
--- a/Engine/Source/Editor/ContentBrowser/Private/SFilterList.h
+++ b/Engine/Source/Editor/ContentBrowser/Private/SFilterList.h
@@ -7,8 +7,8 @@
 #include "Widgets/DeclarativeSyntaxSupport.h"
 #include "Widgets/SWidget.h"
 #include "Widgets/SCompoundWidget.h"
-#include "AssetData.h"
-#include "ARFilter.h"
+#include "AssetRegistry/AssetData.h"
+#include "AssetRegistry/ARFilter.h"
 #include "FrontendFilterBase.h"
 
 class UToolMenu;
diff --git a/Engine/Source/Editor/ContentBrowser/Public/IContentBrowserSingleton.h b/Engine/Source/Editor/ContentBrowser/Public/IContentBrowserSingleton.h
--- a/Engine/Source/Editor/ContentBrowser/Public/IContentBrowserSingleton.h
+++ b/Engine/Source/Editor/ContentBrowser/Public/IContentBrowserSingleton.h
@@ -5,9 +5,9 @@
 #include "CoreMinimal.h"
 #include "SlateFwd.h"
 #include "Misc/Attribute.h"
-#include "AssetData.h"
+#include "AssetRegistry/AssetData.h"
 #include "Developer/AssetTools/Public/AssetTypeCategories.h"
-#include "ARFilter.h"
+#include "AssetRegistry/ARFilter.h"
 #include "ContentBrowserItem.h"
 #include "ContentBrowserDelegates.h"
 #include "ContentBrowserDataSubsystem.h"
diff --git a/Engine/Source/Editor/EditorWidgets/Private/SAssetDiscoveryIndicator.h b/Engine/Source/Editor/EditorWidgets/Private/SAssetDiscoveryIndicator.h
--- a/Engine/Source/Editor/EditorWidgets/Private/SAssetDiscoveryIndicator.h
+++ b/Engine/Source/Editor/EditorWidgets/Private/SAssetDiscoveryIndicator.h
@@ -11,7 +11,7 @@
 #include "Widgets/DeclarativeSyntaxSupport.h"
 #include "Widgets/SCompoundWidget.h"
 #include "AssetDiscoveryIndicator.h"
-#include "IAssetRegistry.h"
+#include "AssetRegistry/IAssetRegistry.h"
 
 /** An indicator for the progress of the asset registry background search */
 class SAssetDiscoveryIndicator : public SCompoundWidget
diff --git a/Engine/Source/Editor/Kismet/Private/STimelineEditor.cpp b/Engine/Source/Editor/Kismet/Private/STimelineEditor.cpp
--- a/Engine/Source/Editor/Kismet/Private/STimelineEditor.cpp
+++ b/Engine/Source/Editor/Kismet/Private/STimelineEditor.cpp
@@ -26,7 +26,7 @@
 #include "Editor/PropertyEditor/Public/PropertyCustomizationHelpers.h"
 
 #include "BlueprintEditor.h"
-#include "AssetRegistryModule.h"
+#include "AssetRegistry/AssetRegistryModule.h"
 #include "IAssetTools.h"
 #include "AssetToolsModule.h"
 #include "Dialogs/DlgPickAssetPath.h"
@@ -35,7 +35,7 @@
 #include "Framework/Notifications/NotificationManager.h"
 #include "Widgets/Notifications/SNotificationList.h"
 #include "Engine/Selection.h"
-#include "AssetData.h"
+#include "AssetRegistry/AssetData.h"
 #include "Styling/StyleColors.h"
 
 #define LOCTEXT_NAMESPACE "STimelineEditor"
diff --git a/Engine/Source/Editor/Kismet/Private/Tests/BlueprintEditorTests.cpp b/Engine/Source/Editor/Kismet/Private/Tests/BlueprintEditorTests.cpp
--- a/Engine/Source/Editor/Kismet/Private/Tests/BlueprintEditorTests.cpp
+++ b/Engine/Source/Editor/Kismet/Private/Tests/BlueprintEditorTests.cpp
@@ -21,7 +21,7 @@
 #include "Engine/StaticMesh.h"
 #include "Kismet/KismetSystemLibrary.h"
 #include "Engine/BlueprintGeneratedClass.h"
-#include "AssetData.h"
+#include "AssetRegistry/AssetData.h"
 #include "Editor.h"
 #include "UObject/SavePackage.h"
 
@@ -39,13 +39,13 @@
 #include "BlueprintEditorModes.h"
 #include "Kismet2/BlueprintEditorUtils.h"
 #include "ComponentAssetBroker.h"
-#include "ARFilter.h"
+#include "AssetRegistry/ARFilter.h"
 
 #include "ScopedTransaction.h"
 #include "ObjectTools.h"
 
 // Automation
-#include "AssetRegistryModule.h"
+#include "AssetRegistry/AssetRegistryModule.h"
 #include "Tests/AutomationTestSettings.h"
 #include "Tests/AutomationEditorCommon.h"
 #include "Tests/AutomationEditorPromotionCommon.h"
diff --git a/Engine/Source/Editor/LevelEditor/Private/SLevelViewport.cpp b/Engine/Source/Editor/LevelEditor/Private/SLevelViewport.cpp
--- a/Engine/Source/Editor/LevelEditor/Private/SLevelViewport.cpp
+++ b/Engine/Source/Editor/LevelEditor/Private/SLevelViewport.cpp
@@ -69,8 +69,8 @@
 #include "FoliageType.h"
 #include "IVREditorModule.h"
 #include "ShowFlagMenuCommands.h"
-#include "AssetRegistryModule.h"
-#include "IAssetRegistry.h"
+#include "AssetRegistry/AssetRegistryModule.h"
+#include "AssetRegistry/IAssetRegistry.h"
 #include "BufferVisualizationMenuCommands.h"
 #include "NaniteVisualizationMenuCommands.h"
 #include "LumenVisualizationMenuCommands.h"
diff --git a/Engine/Source/Editor/MergeActors/Private/MeshApproximationTool/MeshApproximationTool.cpp b/Engine/Source/Editor/MergeActors/Private/MeshApproximationTool/MeshApproximationTool.cpp
--- a/Engine/Source/Editor/MergeActors/Private/MeshApproximationTool/MeshApproximationTool.cpp
+++ b/Engine/Source/Editor/MergeActors/Private/MeshApproximationTool/MeshApproximationTool.cpp
@@ -21,12 +21,12 @@
 #include "MaterialUtilities.h"
 #include "IContentBrowserSingleton.h"
 #include "ContentBrowserModule.h"
-#include "AssetRegistryModule.h"
+#include "AssetRegistry/AssetRegistryModule.h"
 #include "IMeshReductionInterfaces.h"
 #include "IMeshMergeUtilities.h"
 #include "MeshMergeModule.h"
 
-#include "AssetRegistryModule.h"
+#include "AssetRegistry/AssetRegistryModule.h"
 #include "AssetToolsModule.h"
 #include "ObjectTools.h"
 #include "Algo/ForEach.h"
diff --git a/Engine/Source/Editor/Persona/Private/PersonaModule.cpp b/Engine/Source/Editor/Persona/Private/PersonaModule.cpp
--- a/Engine/Source/Editor/Persona/Private/PersonaModule.cpp
+++ b/Engine/Source/Editor/Persona/Private/PersonaModule.cpp
@@ -13,10 +13,10 @@
 #include "Materials/Material.h"
 #include "IPersonaPreviewScene.h"
 #include "Logging/TokenizedMessage.h"
-#include "ARFilter.h"
+#include "AssetRegistry/ARFilter.h"
 #include "AnimGraphDefinitions.h"
 #include "Toolkits/ToolkitManager.h"
-#include "AssetRegistryModule.h"
+#include "AssetRegistry/AssetRegistryModule.h"
 #include "SkeletalMeshSocketDetails.h"
 #include "AnimNotifyDetails.h"
 #include "AnimGraphNodeDetails.h"
diff --git a/Engine/Source/Editor/Persona/Private/SAnimationSequenceBrowser.h b/Engine/Source/Editor/Persona/Private/SAnimationSequenceBrowser.h
--- a/Engine/Source/Editor/Persona/Private/SAnimationSequenceBrowser.h
+++ b/Engine/Source/Editor/Persona/Private/SAnimationSequenceBrowser.h
@@ -10,13 +10,13 @@
 #include "Widgets/SWidget.h"
 #include "PreviewScene.h"
 #include "EditorViewportClient.h"
-#include "AssetData.h"
+#include "AssetRegistry/AssetData.h"
 #include "Widgets/SToolTip.h"
 #include "PersonaDelegates.h"
 #include "Editor/ContentBrowser/Public/ContentBrowserDelegates.h"
 #include "EditorAnimUtils.h"
 #include "IAnimationSequenceBrowser.h"
-#include "ARFilter.h"
+#include "AssetRegistry/ARFilter.h"
 
 class FSceneViewport;
 class FUICommandList;
diff --git a/Engine/Source/Editor/SequenceRecorder/Public/SequenceRecorderUtils.h b/Engine/Source/Editor/SequenceRecorder/Public/SequenceRecorderUtils.h
--- a/Engine/Source/Editor/SequenceRecorder/Public/SequenceRecorderUtils.h
+++ b/Engine/Source/Editor/SequenceRecorder/Public/SequenceRecorderUtils.h
@@ -7,8 +7,8 @@
 #include "MovieScene.h"
 #include "Modules/ModuleManager.h"
 #include "Misc/PackageName.h"
-#include "AssetData.h"
-#include "AssetRegistryModule.h"
+#include "AssetRegistry/AssetData.h"
+#include "AssetRegistry/AssetRegistryModule.h"
 #include "ObjectTools.h"
 
 class AActor;
diff --git a/Engine/Source/Editor/StructViewer/Private/SStructViewer.cpp b/Engine/Source/Editor/StructViewer/Private/SStructViewer.cpp
--- a/Engine/Source/Editor/StructViewer/Private/SStructViewer.cpp
+++ b/Engine/Source/Editor/StructViewer/Private/SStructViewer.cpp
@@ -30,9 +30,9 @@
 #include "Framework/Docking/TabManager.h"
 #include "SListViewSelectorDropdownMenu.h"
 
-#include "ARFilter.h"
-#include "AssetData.h"
-#include "AssetRegistryModule.h"
+#include "AssetRegistry/ARFilter.h"
+#include "AssetRegistry/AssetData.h"
+#include "AssetRegistry/AssetRegistryModule.h"
 
 #include "ContentBrowserDataDragDropOp.h"
 
diff --git a/Engine/Source/Editor/ToolMenusEditor/Private/SEditToolMenuDialog.cpp b/Engine/Source/Editor/ToolMenusEditor/Private/SEditToolMenuDialog.cpp
--- a/Engine/Source/Editor/ToolMenusEditor/Private/SEditToolMenuDialog.cpp
+++ b/Engine/Source/Editor/ToolMenusEditor/Private/SEditToolMenuDialog.cpp
@@ -24,8 +24,8 @@
 
 #include "Editor.h"
 
-#include "IAssetRegistry.h"
-#include "AssetRegistryModule.h"
+#include "AssetRegistry/IAssetRegistry.h"
+#include "AssetRegistry/AssetRegistryModule.h"
 
 #include "Serialization/StaticMemoryReader.h"
 #include "Widgets/Colors/SColorBlock.h"
diff --git a/Engine/Source/Editor/UnrealEd/Classes/Commandlets/DiffAssetRegistriesCommandlet.h b/Engine/Source/Editor/UnrealEd/Classes/Commandlets/DiffAssetRegistriesCommandlet.h
--- a/Engine/Source/Editor/UnrealEd/Classes/Commandlets/DiffAssetRegistriesCommandlet.h
+++ b/Engine/Source/Editor/UnrealEd/Classes/Commandlets/DiffAssetRegistriesCommandlet.h
@@ -5,7 +5,7 @@
 #include "UObject/ObjectMacros.h"
 #include "Commandlets/Commandlet.h"
 #include "Misc/AssetRegistryInterface.h"
-#include "AssetRegistryState.h"
+#include "AssetRegistry/AssetRegistryState.h"
 #include "DiffAssetRegistriesCommandlet.generated.h"
 
 DECLARE_LOG_CATEGORY_EXTERN(LogDiffAssets, Log, All);
diff --git a/Engine/Source/Editor/UnrealEd/Private/AssetThumbnail.cpp b/Engine/Source/Editor/UnrealEd/Private/AssetThumbnail.cpp
--- a/Engine/Source/Editor/UnrealEd/Private/AssetThumbnail.cpp
+++ b/Engine/Source/Editor/UnrealEd/Private/AssetThumbnail.cpp
@@ -29,8 +29,8 @@
 #include "UnrealEdGlobals.h"
 #include "Slate/SlateTextures.h"
 #include "ObjectTools.h"
-#include "AssetRegistryModule.h"
-#include "IAssetRegistry.h"
+#include "AssetRegistry/AssetRegistryModule.h"
+#include "AssetRegistry/IAssetRegistry.h"
 #include "ShaderCompiler.h"
 #include "AssetCompilingManager.h"
 #include "IAssetTools.h"
diff --git a/Engine/Source/Editor/UnrealEd/Private/Commandlets/AssetRegUtilCommandlet.cpp b/Engine/Source/Editor/UnrealEd/Private/Commandlets/AssetRegUtilCommandlet.cpp
--- a/Engine/Source/Editor/UnrealEd/Private/Commandlets/AssetRegUtilCommandlet.cpp
+++ b/Engine/Source/Editor/UnrealEd/Private/Commandlets/AssetRegUtilCommandlet.cpp
@@ -7,8 +7,8 @@
 
 #include "Commandlets/AssetRegUtilCommandlet.h"
 
-#include "AssetRegistryModule.h"
-#include "AssetRegistryState.h"
+#include "AssetRegistry/AssetRegistryModule.h"
+#include "AssetRegistry/AssetRegistryState.h"
 #include "Engine/Texture.h"
 #include "Logging/LogMacros.h"
 #include "Materials/Material.h"
diff --git a/Engine/Source/Editor/UnrealEd/Private/Commandlets/GatherTextCommandletBase.cpp b/Engine/Source/Editor/UnrealEd/Private/Commandlets/GatherTextCommandletBase.cpp
--- a/Engine/Source/Editor/UnrealEd/Private/Commandlets/GatherTextCommandletBase.cpp
+++ b/Engine/Source/Editor/UnrealEd/Private/Commandlets/GatherTextCommandletBase.cpp
@@ -9,10 +9,10 @@
 #include "SourceControlOperations.h"
 #include "ISourceControlModule.h"
 #include "EngineGlobals.h"
-#include "AssetData.h"
+#include "AssetRegistry/AssetData.h"
 #include "Editor.h"
-#include "IAssetRegistry.h"
-#include "ARFilter.h"
+#include "AssetRegistry/IAssetRegistry.h"
+#include "AssetRegistry/ARFilter.h"
 #include "PackageHelperFunctions.h"
 #include "ObjectTools.h"
 
diff --git a/Engine/Source/Editor/UnrealEd/Private/Commandlets/GenerateAssetManifestCommandlet.cpp b/Engine/Source/Editor/UnrealEd/Private/Commandlets/GenerateAssetManifestCommandlet.cpp
--- a/Engine/Source/Editor/UnrealEd/Private/Commandlets/GenerateAssetManifestCommandlet.cpp
+++ b/Engine/Source/Editor/UnrealEd/Private/Commandlets/GenerateAssetManifestCommandlet.cpp
@@ -7,10 +7,10 @@
 =============================================================================*/
 
 #include "Commandlets/GenerateAssetManifestCommandlet.h"
-#include "AssetRegistryModule.h"
+#include "AssetRegistry/AssetRegistryModule.h"
 #include "HAL/FileManager.h"
 #include "Misc/Paths.h"
-#include "ARFilter.h"
+#include "AssetRegistry/ARFilter.h"
 #include "Engine/World.h"
 #include "Misc/FileHelper.h"
 
diff --git a/Engine/Source/Editor/UnrealEd/Private/Commandlets/ReplaceAssetsCommandlet.cpp b/Engine/Source/Editor/UnrealEd/Private/Commandlets/ReplaceAssetsCommandlet.cpp
--- a/Engine/Source/Editor/UnrealEd/Private/Commandlets/ReplaceAssetsCommandlet.cpp
+++ b/Engine/Source/Editor/UnrealEd/Private/Commandlets/ReplaceAssetsCommandlet.cpp
@@ -6,10 +6,10 @@
 =============================================================================*/
 
 #include "Commandlets/ReplaceAssetsCommandlet.h"
-#include "AssetRegistryModule.h"
+#include "AssetRegistry/AssetRegistryModule.h"
 #include "HAL/FileManager.h"
 #include "Misc/Paths.h"
-#include "ARFilter.h"
+#include "AssetRegistry/ARFilter.h"
 #include "Engine/World.h"
 
 DEFINE_LOG_CATEGORY_STATIC(LogReplaceAssetsCommandlet, Log, All);
diff --git a/Engine/Source/Editor/UnrealEd/Private/Factories/MaterialImportHelpers.cpp b/Engine/Source/Editor/UnrealEd/Private/Factories/MaterialImportHelpers.cpp
--- a/Engine/Source/Editor/UnrealEd/Private/Factories/MaterialImportHelpers.cpp
+++ b/Engine/Source/Editor/UnrealEd/Private/Factories/MaterialImportHelpers.cpp
@@ -1,9 +1,9 @@
 // Copyright Epic Games, Inc. All Rights Reserved.
 
 #include "Factories/MaterialImportHelpers.h"
-#include "AssetRegistryModule.h"
-#include "AssetData.h"
-#include "ARFilter.h"
+#include "AssetRegistry/AssetRegistryModule.h"
+#include "AssetRegistry/AssetData.h"
+#include "AssetRegistry/ARFilter.h"
 #include "Materials/MaterialInterface.h"
 #include "Misc/Paths.h"
 
diff --git a/Engine/Source/Editor/UnrealEd/Private/Fbx/FbxMainImport.cpp b/Engine/Source/Editor/UnrealEd/Private/Fbx/FbxMainImport.cpp
--- a/Engine/Source/Editor/UnrealEd/Private/Fbx/FbxMainImport.cpp
+++ b/Engine/Source/Editor/UnrealEd/Private/Fbx/FbxMainImport.cpp
@@ -28,8 +28,8 @@
 #include "UObject/MetaData.h"
 #include "UObject/UObjectGlobals.h"
 #include "UObject/Package.h"
-#include "AssetRegistryModule.h"
-#include "ARFilter.h"
+#include "AssetRegistry/AssetRegistryModule.h"
+#include "AssetRegistry/ARFilter.h"
 #include "Animation/Skeleton.h"
 #include "HAL/PlatformApplicationMisc.h"
 #include "Editor/EditorPerProjectUserSettings.h"
diff --git a/Engine/Source/Editor/UnrealEd/Private/Fbx/FbxSkeletalMeshImport.cpp b/Engine/Source/Editor/UnrealEd/Private/Fbx/FbxSkeletalMeshImport.cpp
--- a/Engine/Source/Editor/UnrealEd/Private/Fbx/FbxSkeletalMeshImport.cpp
+++ b/Engine/Source/Editor/UnrealEd/Private/Fbx/FbxSkeletalMeshImport.cpp
@@ -39,9 +39,9 @@
 #include "Logging/TokenizedMessage.h"
 #include "FbxImporter.h"
 
-#include "AssetData.h"
-#include "ARFilter.h"
-#include "AssetRegistryModule.h"
+#include "AssetRegistry/AssetData.h"
+#include "AssetRegistry/ARFilter.h"
+#include "AssetRegistry/AssetRegistryModule.h"
 #include "AssetNotifications.h"
 
 #include "ObjectTools.h"
diff --git a/Engine/Source/Editor/UnrealEd/Private/Kismet2/BlueprintEditorUtils.cpp b/Engine/Source/Editor/UnrealEd/Private/Kismet2/BlueprintEditorUtils.cpp
--- a/Engine/Source/Editor/UnrealEd/Private/Kismet2/BlueprintEditorUtils.cpp
+++ b/Engine/Source/Editor/UnrealEd/Private/Kismet2/BlueprintEditorUtils.cpp
@@ -125,10 +125,10 @@
 #include "AnimGraphNode_LinkedInputPose.h"
 #include "AnimGraphNode_Root.h"
 #include "Subsystems/AssetEditorSubsystem.h"
-#include "AssetRegistryModule.h"
+#include "AssetRegistry/AssetRegistryModule.h"
 #include "Misc/FeedbackContext.h"
 
-#include "AssetRegistryModule.h"
+#include "AssetRegistry/AssetRegistryModule.h"
 #include "Misc/MessageDialog.h"
 
 #define LOCTEXT_NAMESPACE "Blueprint"
diff --git a/Engine/Source/Editor/UnrealEd/Private/LocalizedAssetUtil.cpp b/Engine/Source/Editor/UnrealEd/Private/LocalizedAssetUtil.cpp
--- a/Engine/Source/Editor/UnrealEd/Private/LocalizedAssetUtil.cpp
+++ b/Engine/Source/Editor/UnrealEd/Private/LocalizedAssetUtil.cpp
@@ -4,9 +4,9 @@
 #include "LocalizationSourceControlUtil.h"
 #include "Misc/Paths.h"
 #include "Misc/PackageName.h"
-#include "AssetData.h"
-#include "IAssetRegistry.h"
-#include "ARFilter.h"
+#include "AssetRegistry/AssetData.h"
+#include "AssetRegistry/IAssetRegistry.h"
+#include "AssetRegistry/ARFilter.h"
 #include "PackageHelperFunctions.h"
 #include "ObjectTools.h"
 
diff --git a/Engine/Source/Editor/UnrealEd/Private/ObjectTools.cpp b/Engine/Source/Editor/UnrealEd/Private/ObjectTools.cpp
--- a/Engine/Source/Editor/UnrealEd/Private/ObjectTools.cpp
+++ b/Engine/Source/Editor/UnrealEd/Private/ObjectTools.cpp
@@ -68,12 +68,12 @@
 #include "BusyCursor.h"
 #include "Dialogs/DlgMoveAssets.h"
 #include "Dialogs/DlgReferenceTree.h"
-#include "ARFilter.h"
+#include "AssetRegistry/ARFilter.h"
 #include "AssetDeleteModel.h"
 #include "Dialogs/SDeleteAssetsDialog.h"
 #include "AudioDevice.h"
 #include "ReferencedAssetsUtils.h"
-#include "AssetRegistryModule.h"
+#include "AssetRegistry/AssetRegistryModule.h"
 #include "PackagesDialog.h"
 #include "PropertyEditorModule.h"
 #include "Kismet2/KismetEditorUtilities.h"
diff --git a/Engine/Source/Editor/UnrealEd/Private/PackageTools.cpp b/Engine/Source/Editor/UnrealEd/Private/PackageTools.cpp
--- a/Engine/Source/Editor/UnrealEd/Private/PackageTools.cpp
+++ b/Engine/Source/Editor/UnrealEd/Private/PackageTools.cpp
@@ -42,15 +42,15 @@
 #include "UObject/WeakObjectPtrTemplates.h"
 #include "Widgets/Notifications/SNotificationList.h"
 
-#include "AssetData.h"
-#include "AssetRegistryModule.h"
+#include "AssetRegistry/AssetData.h"
+#include "AssetRegistry/AssetRegistryModule.h"
 #include "ComponentReregisterContext.h"
 #include "Engine/BlueprintGeneratedClass.h"
 #include "Engine/GameEngine.h"
 #include "Engine/LevelStreaming.h"
 #include "Engine/MapBuildDataRegistry.h"
 #include "Engine/Selection.h"
-#include "IAssetRegistry.h"
+#include "AssetRegistry/IAssetRegistry.h"
 #include "Logging/MessageLog.h"
 #include "UObject/UObjectGlobals.h"
 #include "UObject/UObjectIterator.h"
diff --git a/Engine/Source/Editor/UnrealEd/Private/UnrealEdPrivatePCH.h b/Engine/Source/Editor/UnrealEd/Private/UnrealEdPrivatePCH.h
--- a/Engine/Source/Editor/UnrealEd/Private/UnrealEdPrivatePCH.h
+++ b/Engine/Source/Editor/UnrealEd/Private/UnrealEdPrivatePCH.h
@@ -449,10 +449,10 @@
 #include "IMessageContext.h"
 
 // From AssetRegistry:
-#include "AssetData.h"
-#include "IAssetRegistry.h"
-#include "AssetRegistryModule.h"
-#include "ARFilter.h"
+#include "AssetRegistry/AssetData.h"
+#include "AssetRegistry/IAssetRegistry.h"
+#include "AssetRegistry/AssetRegistryModule.h"
+#include "AssetRegistry/ARFilter.h"
 
 // From Internationalization:
 #include "LocTextHelper.h"
diff --git a/Engine/Source/Editor/UnrealEd/Public/Commandlets/AssetRegistryGenerator.h b/Engine/Source/Editor/UnrealEd/Public/Commandlets/AssetRegistryGenerator.h
--- a/Engine/Source/Editor/UnrealEd/Public/Commandlets/AssetRegistryGenerator.h
+++ b/Engine/Source/Editor/UnrealEd/Public/Commandlets/AssetRegistryGenerator.h
@@ -4,8 +4,8 @@
 
 #include "CoreMinimal.h"
 
-#include "AssetData.h"
-#include "AssetRegistryState.h"
+#include "AssetRegistry/AssetData.h"
+#include "AssetRegistry/AssetRegistryState.h"
 #include "Misc/AssetRegistryInterface.h"
 #include "Misc/Paths.h"
 #include "Templates/UniquePtr.h"
diff --git a/Engine/Source/Editor/VirtualTexturingEditor/Private/RuntimeVirtualTextureAssetTypeActions.cpp b/Engine/Source/Editor/VirtualTexturingEditor/Private/RuntimeVirtualTextureAssetTypeActions.cpp
--- a/Engine/Source/Editor/VirtualTexturingEditor/Private/RuntimeVirtualTextureAssetTypeActions.cpp
+++ b/Engine/Source/Editor/VirtualTexturingEditor/Private/RuntimeVirtualTextureAssetTypeActions.cpp
@@ -2,13 +2,13 @@
 
 #include "RuntimeVirtualTextureAssetTypeActions.h"
 
-#include "AssetRegistryModule.h"
+#include "AssetRegistry/AssetRegistryModule.h"
 #include "ContentBrowserModule.h"
 #include "EditorSupportDelegates.h"
 #include "FileHelpers.h"
 #include "Framework/Commands/UIAction.h"
 #include "Framework/MultiBox/MultiBoxBuilder.h"
-#include "IAssetRegistry.h"
+#include "AssetRegistry/IAssetRegistry.h"
 #include "IContentBrowserSingleton.h"
 #include "MaterialEditingLibrary.h"
 #include "Materials/Material.h"
diff --git a/Engine/Source/Runtime/CoreUObject/Public/AssetRegistry/AssetDataTagMapSerializationDetails.h b/Engine/Source/Runtime/CoreUObject/Public/AssetRegistry/AssetDataTagMapSerializationDetails.h
--- a/Engine/Source/Runtime/CoreUObject/Public/AssetRegistry/AssetDataTagMapSerializationDetails.h
+++ b/Engine/Source/Runtime/CoreUObject/Public/AssetRegistry/AssetDataTagMapSerializationDetails.h
@@ -2,7 +2,7 @@
 
 #pragma once
 
-#include "AssetDataTagMap.h"
+#include "AssetRegistry/AssetDataTagMap.h"
 #include "Async/Async.h"
 
 struct FAssetRegistrySerializationOptions;
diff --git a/Engine/Source/Runtime/Engine/Classes/Engine/AssetManager.h b/Engine/Source/Runtime/Engine/Classes/Engine/AssetManager.h
--- a/Engine/Source/Runtime/Engine/Classes/Engine/AssetManager.h
+++ b/Engine/Source/Runtime/Engine/Classes/Engine/AssetManager.h
@@ -2,12 +2,12 @@
 
 #pragma once
 
-#include "AssetData.h"
+#include "AssetRegistry/AssetData.h"
 #include "AssetManagerTypes.h"
 #include "Misc/AssetRegistryInterface.h"
 #include "StreamableManager.h"
-#include "AssetBundleData.h"
-#include "AssetRegistryModule.h"
+#include "AssetRegistry/AssetBundleData.h"
+#include "AssetRegistry/AssetRegistryModule.h"
 #include "GenericPlatform/GenericPlatformChunkInstall.h"
 #include "ContentEncryptionConfig.h"
 #include "AssetManager.generated.h"
diff --git a/Engine/Source/Runtime/Engine/Classes/Engine/AssetManagerTypes.h b/Engine/Source/Runtime/Engine/Classes/Engine/AssetManagerTypes.h
--- a/Engine/Source/Runtime/Engine/Classes/Engine/AssetManagerTypes.h
+++ b/Engine/Source/Runtime/Engine/Classes/Engine/AssetManagerTypes.h
@@ -4,8 +4,8 @@
 
 #include "CoreMinimal.h"
 #include "UObject/PrimaryAssetId.h"
-#include "AssetData.h"
-#include "AssetBundleData.h"
+#include "AssetRegistry/AssetData.h"
+#include "AssetRegistry/AssetBundleData.h"
 #include "EngineTypes.h"
 #include "AssetManagerTypes.generated.h"
 
diff --git a/Engine/Source/Runtime/Engine/Classes/Engine/DataAsset.h b/Engine/Source/Runtime/Engine/Classes/Engine/DataAsset.h
--- a/Engine/Source/Runtime/Engine/Classes/Engine/DataAsset.h
+++ b/Engine/Source/Runtime/Engine/Classes/Engine/DataAsset.h
@@ -5,7 +5,7 @@
 #include "CoreMinimal.h"
 #include "UObject/ObjectMacros.h"
 #include "UObject/Object.h"
-#include "AssetBundleData.h"
+#include "AssetRegistry/AssetBundleData.h"
 #include "Templates/SubclassOf.h"
 #include "DataAsset.generated.h"
 
diff --git a/Engine/Source/Runtime/Engine/Private/Animation/Skeleton.cpp b/Engine/Source/Runtime/Engine/Private/Animation/Skeleton.cpp
--- a/Engine/Source/Runtime/Engine/Private/Animation/Skeleton.cpp
+++ b/Engine/Source/Runtime/Engine/Private/Animation/Skeleton.cpp
@@ -17,9 +17,9 @@
 #include "Animation/AnimSequenceBase.h"
 #include "Animation/AnimSequence.h"
 #include "AnimationRuntime.h"
-#include "AssetData.h"
-#include "ARFilter.h"
-#include "AssetRegistryModule.h"
+#include "AssetRegistry/AssetData.h"
+#include "AssetRegistry/ARFilter.h"
+#include "AssetRegistry/AssetRegistryModule.h"
 #include "Animation/Rig.h"
 #include "Animation/BlendProfile.h"
 #include "Logging/TokenizedMessage.h"
diff --git a/Engine/Source/Runtime/Engine/Private/AssetManager.cpp b/Engine/Source/Runtime/Engine/Private/AssetManager.cpp
--- a/Engine/Source/Runtime/Engine/Private/AssetManager.cpp
+++ b/Engine/Source/Runtime/Engine/Private/AssetManager.cpp
@@ -3,8 +3,8 @@
 #include "Engine/AssetManager.h"
 #include "Engine/AssetManagerSettings.h"
 #include "Engine/PrimaryAssetLabel.h"
-#include "AssetData.h"
-#include "ARFilter.h"
+#include "AssetRegistry/AssetData.h"
+#include "AssetRegistry/ARFilter.h"
 #include "Containers/StringView.h"
 #include "Engine/Engine.h"
 #include "Engine/BlueprintGeneratedClass.h"
@@ -19,7 +19,7 @@
 #include "Misc/Paths.h"
 #include "Misc/StringBuilder.h"
 #include "Serialization/MemoryReader.h"
-#include "AssetRegistryState.h"
+#include "AssetRegistry/AssetRegistryState.h"
 #include "HAL/PlatformFileManager.h"
 #include "IPlatformFilePak.h"
 #include "Stats/StatsMisc.h"
diff --git a/Engine/Source/Runtime/Engine/Private/Level.cpp b/Engine/Source/Runtime/Engine/Private/Level.cpp
--- a/Engine/Source/Runtime/Engine/Private/Level.cpp
+++ b/Engine/Source/Runtime/Engine/Private/Level.cpp
@@ -61,9 +61,9 @@ Level.cpp: Level-related functions
 #include "Kismet2/KismetEditorUtilities.h"
 #include "Kismet2/BlueprintEditorUtils.h"
 #include "Algo/AnyOf.h"
-#include "AssetRegistryModule.h"
-#include "IAssetRegistry.h"
-#include "AssetData.h"
+#include "AssetRegistry/AssetRegistryModule.h"
+#include "AssetRegistry/IAssetRegistry.h"
+#include "AssetRegistry/AssetData.h"
 #include "PieFixupSerializer.h"
 #include "Editor.h"
 #include "Subsystems/EditorActorSubsystem.h"
diff --git a/Engine/Source/Runtime/Engine/Private/LevelInstance/LevelInstanceSubsystem.cpp b/Engine/Source/Runtime/Engine/Private/LevelInstance/LevelInstanceSubsystem.cpp
--- a/Engine/Source/Runtime/Engine/Private/LevelInstance/LevelInstanceSubsystem.cpp
+++ b/Engine/Source/Runtime/Engine/Private/LevelInstance/LevelInstanceSubsystem.cpp
@@ -25,8 +25,8 @@
 #include "Misc/ITransaction.h"
 #include "Misc/Paths.h"
 #include "Misc/PackageName.h"
-#include "AssetRegistryModule.h"
-#include "AssetData.h"
+#include "AssetRegistry/AssetRegistryModule.h"
+#include "AssetRegistry/AssetData.h"
 #include "FileHelpers.h"
 #include "Editor.h"
 #include "EditorLevelUtils.h"
diff --git a/Engine/Source/Runtime/Experimental/Interchange/Engine/Private/InterchangeManager.cpp b/Engine/Source/Runtime/Experimental/Interchange/Engine/Private/InterchangeManager.cpp
--- a/Engine/Source/Runtime/Experimental/Interchange/Engine/Private/InterchangeManager.cpp
+++ b/Engine/Source/Runtime/Experimental/Interchange/Engine/Private/InterchangeManager.cpp
@@ -2,8 +2,8 @@
 
 #include "InterchangeManager.h"
 
-#include "AssetDataTagMap.h"
-#include "AssetRegistryModule.h"
+#include "AssetRegistry/AssetDataTagMap.h"
+#include "AssetRegistry/AssetRegistryModule.h"
 #include "Async/Async.h"
 #include "CoreMinimal.h"
 #include "Engine/Blueprint.h"
