Include sparse class data in the BP diff

#rnx

diff --git a/Engine/Source/Editor/Kismet/Private/DiffUtils.cpp b/Engine/Source/Editor/Kismet/Private/DiffUtils.cpp
--- a/Engine/Source/Editor/Kismet/Private/DiffUtils.cpp
+++ b/Engine/Source/Editor/Kismet/Private/DiffUtils.cpp
@@ -42,7 +42,18 @@ namespace UEDiffUtils_Private
 
 FResolvedProperty FPropertySoftPath::Resolve(const UObject* Object) const
 {
-	return Resolve(Object->GetClass(), Object);
+	FResolvedProperty ResolvedProperty = Resolve(Object->GetClass(), Object);
+	if (!ResolvedProperty.Property && Object->HasAnyFlags(RF_ClassDefaultObject))
+	{
+		if (const UScriptStruct* SparseClassDataStruct = Object->GetClass()->GetSparseClassDataStruct())
+		{
+			if (const void* SparseClassData = Object->GetClass()->GetSparseClassData(EGetSparseClassDataMethod::ReturnIfNull))
+			{
+				ResolvedProperty = Resolve(SparseClassDataStruct, SparseClassData);
+			}
+		}
+	}
+	return ResolvedProperty;
 }
 
 FResolvedProperty FPropertySoftPath::Resolve(const UStruct* Struct, const void* StructData) const
