Modify PinValueInspector to support custom editor implementation
#rnx

diff --git a/Engine/Source/Editor/Kismet/Private/SPinValueInspector.cpp b/Engine/Source/Editor/Kismet/Private/SPinValueInspector.cpp
--- a/Engine/Source/Editor/Kismet/Private/SPinValueInspector.cpp
+++ b/Engine/Source/Editor/Kismet/Private/SPinValueInspector.cpp
@@ -178,6 +178,11 @@ void SPinValueInspector::OnExpansionChanged(FDebugTreeItemPtr InItem, bool bItem
 	ConstrainedBox->RequestResize();
 }
 
+void SPinValueInspector::AddTreeItemUnique(const FDebugTreeItemPtr& Item)
+{
+	TreeViewWidget->AddTreeItemUnique(Item);
+}
+
 void SPinValueInspector::PopulateTreeView()
 {
 	// Locate the class property associated with the source pin and set it to the root node.
@@ -245,6 +250,16 @@ void SPinValueInspector::AddPinToTreeView(const UEdGraphPin* Pin, UBlueprint* Bl
 	}
 }
 
+void SPinValueInspector::RequestUpdateFilteredItems()
+{
+	TreeViewWidget->RequestUpdateFilteredItems();
+}
+
+const FEdGraphPinReference& SPinValueInspector::GetPinRef() const
+{
+	return PinRef;
+}
+
 void SPinValueInspector::SetPinRef(const FEdGraphPinReference& InPinRef)
 {
 	TreeViewWidget->ClearTreeItems();
@@ -253,7 +268,7 @@ void SPinValueInspector::SetPinRef(const FEdGraphPinReference& InPinRef)
 	ConstrainedBox->RequestResize();
 }
 
-TWeakPtr<FPinValueInspectorTooltip> FPinValueInspectorTooltip::SummonTooltip(FEdGraphPinReference InPinRef)
+TWeakPtr<FPinValueInspectorTooltip> FPinValueInspectorTooltip::SummonTooltip(FEdGraphPinReference InPinRef, TSharedPtr<SPinValueInspector> InContentWidgetOverride)
 {
 	if (!FSlateApplication::Get().GetVisibleMenuWindow().Get())
 	{
@@ -269,7 +284,17 @@ TWeakPtr<FPinValueInspectorTooltip> FPinValueInspectorTooltip::SummonTooltip(FEd
 			}
 
 			Instance = MakeShared<FPinValueInspectorTooltip>();
-			ValueInspectorWidget->SetPinRef(InPinRef);
+			if (TooltipWidget.IsValid() && InContentWidgetOverride)
+			{
+				TooltipWidget->SetContentWidget(InContentWidgetOverride->AsShared());
+				InContentWidgetOverride->SetPinRef(InPinRef);
+				ValueInspectorWidget->SetPinRef({ });
+			}
+			else
+			{
+				ValueInspectorWidget->SetPinRef(InPinRef);
+			}
+
 			TooltipWindow->ShowWindow();
 			TooltipWindow->SetAllowFastUpdate(true);
 			return Instance;
@@ -312,9 +337,16 @@ void FPinValueInspectorTooltip::DismissTooltip()
 	{
 		FSlateApplication::Get().DismissAllMenus();
 	}
+
 	ValueInspectorWidget->SetPinRef(FEdGraphPinReference());
+
 	TooltipWindow->HideWindow();
 	TooltipWindow->SetAllowFastUpdate(false);
+	if (TooltipWidget->GetContentWidget() != ValueInspectorWidget)
+	{
+		TooltipWidget->SetContentWidget(ValueInspectorWidget->AsShared());
+	}
+
 	Instance.Reset();
 }
 
@@ -335,12 +367,12 @@ void FPinValueInspectorTooltip::CreatePinValueTooltipWindow()
 	TooltipWidget->SetContentWidget(SAssignNew(ValueInspectorWidget, SPinValueInspector));
 }
 
-bool FPinValueInspectorTooltip::TooltipCanClose()
+bool FPinValueInspectorTooltip::TooltipCanClose() const
 {
 	return !TooltipWindow->IsHovered() && !TooltipHostsMenu();
 }
 
-bool FPinValueInspectorTooltip::TooltipHostsMenu()
+bool FPinValueInspectorTooltip::TooltipHostsMenu() const
 {
 	TSharedPtr<SWindow> MenuHostWindow = FSlateApplication::Get().GetVisibleMenuWindow();
 	return MenuHostWindow.IsValid() && (MenuHostWindow == TooltipWindow);
diff --git a/Engine/Source/Editor/Kismet/Public/Debugging/SKismetDebugTreeView.h b/Engine/Source/Editor/Kismet/Public/Debugging/SKismetDebugTreeView.h
--- a/Engine/Source/Editor/Kismet/Public/Debugging/SKismetDebugTreeView.h
+++ b/Engine/Source/Editor/Kismet/Public/Debugging/SKismetDebugTreeView.h
@@ -28,7 +28,7 @@ class SCheckBox;
 typedef TSharedPtr<class FDebugLineItem> FDebugTreeItemPtr;
 
 // The base class for a line entry in the debugging tree view
-class FDebugLineItem : public TSharedFromThis<FDebugLineItem>
+class KISMET_API FDebugLineItem : public TSharedFromThis<FDebugLineItem>
 {
 public:
 	friend class FLineItemWithChildren; // used by FLineItemWithChildren::EnsureChildIsAdded
diff --git a/Engine/Source/Editor/Kismet/Public/SPinValueInspector.h b/Engine/Source/Editor/Kismet/Public/SPinValueInspector.h
--- a/Engine/Source/Editor/Kismet/Public/SPinValueInspector.h
+++ b/Engine/Source/Editor/Kismet/Public/SPinValueInspector.h
@@ -22,29 +22,38 @@ public:
 	{}
 	SLATE_END_ARGS()
 
-		void Construct(const FArguments& InArgs);
+	void Construct(const FArguments& InArgs);
 
 	/** Whether the search filter UI should be visible. */
 	bool ShouldShowSearchFilter() const;
 
+	/** Gets the current watched pin reference */
+	const FEdGraphPinReference& GetPinRef() const;
+
 	friend class FPinValueInspectorTooltip;
 
 protected:
 	/** @return Visibility of the search box filter widget. */
-	EVisibility GetSearchFilterVisibility() const;
+	virtual EVisibility GetSearchFilterVisibility() const;
 
 	/** Passes SearchText through to tree view */
-	void OnSearchTextChanged(const FText& InSearchText);
+	virtual void OnSearchTextChanged(const FText& InSearchText);
 
 	/** requests the constrained box be resized */
-	void OnExpansionChanged(FDebugTreeItemPtr InItem, bool bItemIsExpanded);
+	virtual void OnExpansionChanged(FDebugTreeItemPtr InItem, bool bItemIsExpanded);
+
+	/** Adds a unique tree item to the TreeView. */
+	virtual void AddTreeItemUnique(const FDebugTreeItemPtr& Item);
 
 	/** Adds relevant pins to the tree view */
-	void PopulateTreeView();
+	virtual void PopulateTreeView();
 
 	/** Adds a single pin to the tree view */
 	void AddPinToTreeView(const UEdGraphPin* Pin, UBlueprint* Blueprint);
 
+	/** Marks tree view for updating currently filtered items. */
+	void RequestUpdateFilteredItems();
+
 	/** Sets the current watched pin */
 	void SetPinRef(const FEdGraphPinReference& InPinRef);
 
@@ -77,24 +86,24 @@ public:
 	 */
 	void TryDismissTooltip(bool bForceDismiss = false);
 
+	/** @returns whether or not the tooltip can close */
+	bool TooltipCanClose() const;
+
+	/** @returns whether this tooltip is the host for a context menu */
+	bool TooltipHostsMenu() const;
+
 private:
 	/** Dismisses the current tooltip (internal implementation) */
 	void DismissTooltip();
-	
-	/** @returns whether this tooltip is the host for a context menu */
-	bool TooltipHostsMenu();
-
-	/** @returns whether or not the tooltip can close */
-	bool TooltipCanClose();
 
 public:
-	/** Summons a new tooltip in the shared window */
-	static TWeakPtr<FPinValueInspectorTooltip> SummonTooltip(FEdGraphPinReference InPinRef);
+	/** Summons a new tooltip in the shared window. If provided, uses a custom implementation of a pin value inspector widget. */
+	static TWeakPtr<FPinValueInspectorTooltip> SummonTooltip(FEdGraphPinReference InPinRef, TSharedPtr<SPinValueInspector> InContentWidgetOverride = { });
 
 	/** Release the references to the static widget shared pointers */
 	static void ShutdownTooltip();
 
-	/** Inspector widget in the tooltip */
+	/** Default inspector widget in the tooltip */
 	static TSharedPtr<SPinValueInspector> ValueInspectorWidget;
 
 private:
