Collision visibility option in the BP editor

#3731

diff --git a/Engine/Source/Editor/Kismet/Private/SSCSEditorViewport.cpp b/Engine/Source/Editor/Kismet/Private/SSCSEditorViewport.cpp
--- a/Engine/Source/Editor/Kismet/Private/SSCSEditorViewport.cpp
+++ b/Engine/Source/Editor/Kismet/Private/SSCSEditorViewport.cpp
@@ -155,6 +155,10 @@ public:
 			case VMI_BrushWireframe:
 				Label = NSLOCTEXT("BlueprintEditor", "ViewMenuTitle_Wireframe", "Wireframe");
 				break;
+
+			case VMI_CollisionVisibility:
+				Label = NSLOCTEXT("BlueprintEditor", "ViewMenuTitle_CollisionVisibility", "Collision Visibility");
+				break;
 			}
 		}
 
@@ -171,6 +175,7 @@ public:
 		ViewMenuBuilder.AddMenuEntry(FEditorViewportCommands::Get().LitMode, NAME_None, NSLOCTEXT("BlueprintEditor", "LitModeMenuOption", "Lit"));
 		ViewMenuBuilder.AddMenuEntry(FEditorViewportCommands::Get().UnlitMode, NAME_None, NSLOCTEXT("BlueprintEditor", "UnlitModeMenuOption", "Unlit"));
 		ViewMenuBuilder.AddMenuEntry(FEditorViewportCommands::Get().WireframeMode, NAME_None, NSLOCTEXT("BlueprintEditor", "WireframeModeMenuOption", "Wireframe"));
+		ViewMenuBuilder.AddMenuEntry(FEditorViewportCommands::Get().CollisionVisibility, NAME_None, NSLOCTEXT("BlueprintEditor", "CollisionVisibilityMenuOption", "Visibility Collision"));
 
 		return ViewMenuBuilder.MakeWidget();
 	}
@@ -320,7 +325,7 @@ void SSCSEditorViewport::BindCommands()
 	// Toggle camera lock on/off
 	CommandList->MapAction(
 		Commands.ResetCamera,
-		FExecuteAction::CreateSP(ViewportClient.Get(), &FSCSEditorViewportClient::ResetCamera) );
+		FExecuteAction::CreateSP(ViewportClient.Get(), &FSCSEditorViewportClient::ResetCamera));
 
 	CommandList->MapAction(
 		Commands.ShowFloor,
