[Bug Fix] Misleading browse to asset tooltip in blueprint editor
- looks like this bug may have been introduced in Steven Dao's CL 17427619
- fixed an additional place where this bug existed

diff --git a/Engine/Source/Editor/Kismet/Private/SBlueprintEditorSelectedDebugObjectWidget.cpp b/Engine/Source/Editor/Kismet/Private/SBlueprintEditorSelectedDebugObjectWidget.cpp
--- a/Engine/Source/Editor/Kismet/Private/SBlueprintEditorSelectedDebugObjectWidget.cpp
+++ b/Engine/Source/Editor/Kismet/Private/SBlueprintEditorSelectedDebugObjectWidget.cpp
@@ -114,9 +114,11 @@ const FString& SBlueprintEditorSelectedDebugObjectWidget::GetDebugAllWorldsStrin
 
 TSharedRef<SWidget> SBlueprintEditorSelectedDebugObjectWidget::OnGetActiveDetailSlotContent(bool bChangedToHighDetail)
 {
-	const TSharedRef<SWidget> BrowseButton = PropertyCustomizationHelpers::MakeBrowseButton(FSimpleDelegate::CreateSP(this, &SBlueprintEditorSelectedDebugObjectWidget::SelectedDebugObject_OnClicked));
+	const TSharedRef<SWidget> BrowseButton = PropertyCustomizationHelpers::MakeBrowseButton(
+		FSimpleDelegate::CreateSP(this, &SBlueprintEditorSelectedDebugObjectWidget::SelectedDebugObject_OnClicked),
+		LOCTEXT("DebugSelectActor", "Select this Actor in level")
+	);
 	BrowseButton->SetVisibility(TAttribute<EVisibility>::Create(TAttribute<EVisibility>::FGetter::CreateSP(this, &SBlueprintEditorSelectedDebugObjectWidget::IsSelectDebugObjectButtonVisible)));
-	BrowseButton->SetToolTipText(LOCTEXT("DebugSelectActor", "Select this Actor in level"));
 
 	TSharedRef<SWidget> DebugObjectSelectionWidget =
 		SNew(SHorizontalBox)
