Fix for a regression that leads to an editor crash when right-clicking a Blueprint delegate variable without an active function graph tab.

diff --git a/Engine/Source/Editor/Kismet/Private/SMyBlueprint.cpp b/Engine/Source/Editor/Kismet/Private/SMyBlueprint.cpp
--- a/Engine/Source/Editor/Kismet/Private/SMyBlueprint.cpp
+++ b/Engine/Source/Editor/Kismet/Private/SMyBlueprint.cpp
@@ -3496,16 +3496,21 @@ void SMyBlueprint::OnCopy()
 	}
 	else if (FEdGraphSchemaAction_BlueprintVariableBase* BPVariable = SelectionAsBlueprintVariable())
 	{
-		const UEdGraphSchema* Schema = GetFocusedGraph()->GetSchema();
-		TArray<FBPVariableDescription> LocalVariables;
-		Schema->GetLocalVariables(GetFocusedGraph(), LocalVariables);
-		for (const FBPVariableDescription& VariableDescription : LocalVariables)
+		if (const UEdGraph* FocusedGraph = Cast<UEdGraph>(BPVariable->GetVariableScope()))
 		{
-			if (VariableDescription.VarName == BPVariable->GetVariableName())
+			if (const UEdGraphSchema* Schema = FocusedGraph->GetSchema())
 			{
-				FBPVariableDescription::StaticStruct()->ExportText(OutputString, &VariableDescription, &VariableDescription, nullptr, 0, nullptr, false);
-				OutputString = VAR_PREFIX + OutputString;
-				break;
+				TArray<FBPVariableDescription> LocalVariables;
+				Schema->GetLocalVariables(FocusedGraph, LocalVariables);
+				for (const FBPVariableDescription& VariableDescription : LocalVariables)
+				{
+					if (VariableDescription.VarName == BPVariable->GetVariableName())
+					{
+						FBPVariableDescription::StaticStruct()->ExportText(OutputString, &VariableDescription, &VariableDescription, nullptr, 0, nullptr, false);
+						OutputString = VAR_PREFIX + OutputString;
+						break;
+					}
+				}
 			}
 		}
 	}
@@ -3537,15 +3542,20 @@ bool SMyBlueprint::CanCopy() const
 	}
 	else if (FEdGraphSchemaAction_BlueprintVariableBase* BPVariable = SelectionAsBlueprintVariable())
 	{
-		TArray<FBPVariableDescription> LocalVariables;
-		const UEdGraphSchema* Schema = GetFocusedGraph()->GetSchema();
-		Schema->GetLocalVariables(GetFocusedGraph(), LocalVariables);
-
-		for (const FBPVariableDescription& VariableDescription : LocalVariables)
+		if (const UEdGraph* FocusedGraph = Cast<UEdGraph>(BPVariable->GetVariableScope()))
 		{
-			if (VariableDescription.VarName == BPVariable->GetVariableName())
+			if (const UEdGraphSchema* Schema = FocusedGraph->GetSchema())
 			{
-				return true;
+				TArray<FBPVariableDescription> LocalVariables;
+				Schema->GetLocalVariables(FocusedGraph, LocalVariables);
+
+				for (const FBPVariableDescription& VariableDescription : LocalVariables)
+				{
+					if (VariableDescription.VarName == BPVariable->GetVariableName())
+					{
+						return true;
+					}
+				}
 			}
 		}
 		return false;		
