Fix crash reloading a dirty, compiled and not saved anim BP that has child BPs

Clear sparse class data so stale pointers to parent data are not maintained. Clearing at this point causes SetSuperClass to correctly apply the archetype sparse class data.

diff --git a/Engine/Source/Editor/Kismet/Private/BlueprintCompilationManager.cpp b/Engine/Source/Editor/Kismet/Private/BlueprintCompilationManager.cpp
--- a/Engine/Source/Editor/Kismet/Private/BlueprintCompilationManager.cpp
+++ b/Engine/Source/Editor/Kismet/Private/BlueprintCompilationManager.cpp
@@ -1941,6 +1941,7 @@ void FBlueprintCompilationManagerImpl::ReparentHierarchies(const TMap<UClass*, U
 			UClass* const* NewParent = OldToNewClasses.Find(ClassToReinstance->GetSuperClass());
 			if(NewParent)
 			{
+				ClassToReinstance->ClearSparseClassDataStruct(false);
 				ClassToReinstance->SetSuperStruct(*NewParent);
 			}
 
