Fixing issue where private functions could be overriden in the "My Blueprint" tab

diff --git a/Engine/Source/Editor/BlueprintGraph/Classes/EdGraphSchema_K2_Actions.h b/Engine/Source/Editor/BlueprintGraph/Classes/EdGraphSchema_K2_Actions.h
--- a/Engine/Source/Editor/BlueprintGraph/Classes/EdGraphSchema_K2_Actions.h
+++ b/Engine/Source/Editor/BlueprintGraph/Classes/EdGraphSchema_K2_Actions.h
@@ -664,7 +664,6 @@ struct BLUEPRINTGRAPH_API FEdGraphSchemaAction_K2Graph : public FEdGraphSchemaAc
 	virtual FEdGraphSchemaActionDefiningObject GetPersistentItemDefiningObject() const override;
 	// End of FEdGraphSchemaAction interface
 
-protected:
 	UFunction* GetFunction() const;
 	UBlueprint* GetSourceBlueprint() const;
 };
diff --git a/Engine/Source/Editor/Kismet/Private/SMyBlueprint.cpp b/Engine/Source/Editor/Kismet/Private/SMyBlueprint.cpp
--- a/Engine/Source/Editor/Kismet/Private/SMyBlueprint.cpp
+++ b/Engine/Source/Editor/Kismet/Private/SMyBlueprint.cpp
@@ -972,6 +972,18 @@ void SMyBlueprint::BuildOverridableFunctionsMenu(FMenuBuilder& MenuBuilder)
 	{
 		for (TSharedPtr<FEdGraphSchemaAction_K2Graph>& OverrideAction : OverridableFunctionActions)
 		{
+			// Check if function data is valid and skip this entry if it's a private function
+			if (OverrideAction->FuncName != NAME_None)
+			{
+				if (const UFunction* OverrideActionFunction = FindUField<UFunction>(Blueprint->SkeletonGeneratedClass, OverrideAction->FuncName))
+				{
+					if (OverrideActionFunction->HasAnyFunctionFlags(FUNC_Private))
+					{
+						continue;
+					}
+				}
+			}
+			
 			UClass* const OverrideFuncClass = FBlueprintEditorUtils::GetOverrideFunctionClass(GetBlueprintObj(), OverrideAction->FuncName);
 			
 			// Add the function name and tooltip 
