Fixing incorrect Tooltip on "Move To Parent Class" context option in My Blueprint

diff --git a/Engine/Source/Editor/Kismet/Private/SMyBlueprint.cpp b/Engine/Source/Editor/Kismet/Private/SMyBlueprint.cpp
--- a/Engine/Source/Editor/Kismet/Private/SMyBlueprint.cpp
+++ b/Engine/Source/Editor/Kismet/Private/SMyBlueprint.cpp
@@ -78,7 +78,8 @@ void FMyBlueprintCommands::RegisterCommands()
 	UI_COMMAND( PasteFunction, "Paste Function", "Pastes the function to this blueprint.", EUserInterfaceActionType::Button, FInputChord());
 	UI_COMMAND( PasteMacro, "Paste Macro", "Pastes the macro to this blueprint.", EUserInterfaceActionType::Button, FInputChord());
 	UI_COMMAND( GotoNativeVarDefinition, "Goto Code Definition", "Goto the native code definition of this variable", EUserInterfaceActionType::Button, FInputChord() );
-	UI_COMMAND( MoveToParent, "Move to Parent Class", "Moves the variable to its parent class", EUserInterfaceActionType::Button, FInputChord() );
+	UI_COMMAND( MoveVariableToParent, "Move to Parent Class", "Moves the variable to its parent class", EUserInterfaceActionType::Button, FInputChord() );
+	UI_COMMAND( MoveFunctionToParent, "Move to Parent Class", "Moves the function to its parent class", EUserInterfaceActionType::Button, FInputChord() );
 }
 
 //////////////////////////////////////////////////////////////////////////
@@ -312,11 +313,17 @@ void SMyBlueprint::Construct(const FArguments& InArgs, TWeakPtr<FBlueprintEditor
 			FIsActionChecked(),
 			FIsActionButtonVisible::CreateSP(this, &SMyBlueprint::IsDuplicateActionVisible) );
 
-		CommandList->MapAction( FMyBlueprintCommands::Get().MoveToParent,
+		CommandList->MapAction( FMyBlueprintCommands::Get().MoveVariableToParent,
 			FExecuteAction::CreateSP(this, &SMyBlueprint::OnMoveToParent),
 			FCanExecuteAction(),
 			FIsActionChecked(),
-			FIsActionButtonVisible::CreateSP(this, &SMyBlueprint::CanMoveToParent) );
+			FIsActionButtonVisible::CreateSP(this, &SMyBlueprint::CanMoveVariableToParent) );
+
+		CommandList->MapAction( FMyBlueprintCommands::Get().MoveFunctionToParent,
+			 FExecuteAction::CreateSP(this, &SMyBlueprint::OnMoveToParent),
+			 FCanExecuteAction(),
+			 FIsActionChecked(),
+			 FIsActionButtonVisible::CreateSP(this, &SMyBlueprint::CanMoveFunctionToParent) );
 
 		CommandList->MapAction( FMyBlueprintCommands::Get().GotoNativeVarDefinition,
 			FExecuteAction::CreateSP(this, &SMyBlueprint::GotoNativeCodeVarDefinition),
@@ -2342,7 +2349,8 @@ TSharedPtr<SWidget> SMyBlueprint::OnContextMenuOpening()
 			MenuBuilder.AddMenuEntry(FGenericCommands::Get().Cut);
 			MenuBuilder.AddMenuEntry(FGenericCommands::Get().Copy);
 			MenuBuilder.AddMenuEntry(FGenericCommands::Get().Duplicate);
-			MenuBuilder.AddMenuEntry(FMyBlueprintCommands::Get().MoveToParent);
+			MenuBuilder.AddMenuEntry(FMyBlueprintCommands::Get().MoveVariableToParent);
+			MenuBuilder.AddMenuEntry(FMyBlueprintCommands::Get().MoveFunctionToParent);
 			MenuBuilder.AddMenuEntry(FMyBlueprintCommands::Get().DeleteEntry);
 		}
 		MenuBuilder.EndSection();
@@ -3438,7 +3446,7 @@ void SMyBlueprint::OnMoveToParent()
 	}
 }
 
-bool SMyBlueprint::CanMoveToParent() const
+bool SMyBlueprint::CanMoveVariableToParent() const
 {
 	bool bCanMove = false;
 
@@ -3452,7 +3460,19 @@ bool SMyBlueprint::CanMoveToParent() const
 			int32 VarIndex = FBlueprintEditorUtils::FindNewVariableIndexAndBlueprint(Blueprint, VarAction->GetVariableName(), SourceBlueprint);
 			bCanMove = (VarIndex != INDEX_NONE) && (SourceBlueprint == Blueprint);
 		}
-		else if (FEdGraphSchemaAction_K2Graph* GraphAction = SelectionAsGraph())
+	}
+
+	return bCanMove;
+}
+
+bool SMyBlueprint::CanMoveFunctionToParent() const
+{
+	bool bCanMove = false;
+
+	TSharedPtr<FBlueprintEditor> PinnedEditor = BlueprintEditorPtr.Pin();
+	if (PinnedEditor.IsValid() && PinnedEditor->IsParentClassABlueprint())
+	{
+		if (FEdGraphSchemaAction_K2Graph* GraphAction = SelectionAsGraph())
 		{
 			if (CanDeleteEntry())
 			{
diff --git a/Engine/Source/Editor/Kismet/Public/SMyBlueprint.h b/Engine/Source/Editor/Kismet/Public/SMyBlueprint.h
--- a/Engine/Source/Editor/Kismet/Public/SMyBlueprint.h
+++ b/Engine/Source/Editor/Kismet/Public/SMyBlueprint.h
@@ -49,7 +49,8 @@ public:
 	TSharedPtr<FUICommandInfo> PasteFunction;
 	TSharedPtr<FUICommandInfo> PasteMacro;
 	TSharedPtr<FUICommandInfo> GotoNativeVarDefinition;
-	TSharedPtr<FUICommandInfo> MoveToParent;
+	TSharedPtr<FUICommandInfo> MoveVariableToParent;
+	TSharedPtr<FUICommandInfo> MoveFunctionToParent;
 	// Add New Item
 	/** Initialize commands */
 	virtual void RegisterCommands() override;
@@ -254,7 +255,8 @@ private:
 	void GotoNativeCodeVarDefinition();
 	bool IsNativeVariable() const;
 	void OnMoveToParent();
-	bool CanMoveToParent() const;
+	bool CanMoveVariableToParent() const;
+	bool CanMoveFunctionToParent() const;
 	void OnCopy();
 	bool CanCopy() const;
 	void OnCut();
