[BugFix] Bitmask enum types use short names in the BP editor details customization as well as in serialized BP variable descriptions.

diff --git a/Engine/Source/Editor/Kismet/Private/BlueprintDetailsCustomization.cpp b/Engine/Source/Editor/Kismet/Private/BlueprintDetailsCustomization.cpp
--- a/Engine/Source/Editor/Kismet/Private/BlueprintDetailsCustomization.cpp
+++ b/Engine/Source/Editor/Kismet/Private/BlueprintDetailsCustomization.cpp
@@ -736,14 +736,14 @@ void FBlueprintVarActionDetails::CustomizeDetails( IDetailLayoutBuilder& DetailL
 		.ToolTip(BitmaskTooltip)
 	];
 
-	BitmaskEnumTypeNames.Empty();
-	BitmaskEnumTypeNames.Add(MakeShareable(new FString(LOCTEXT("BitmaskEnumTypeName_None", "None").ToString())));
+	BitmaskEnumTypePaths.Empty();
+	BitmaskEnumTypePaths.Add(MakeShareable(new FTopLevelAssetPath())); // option to set the bitmask to None
 	for (TObjectIterator<UEnum> EnumIt; EnumIt; ++EnumIt)
 	{
 		UEnum* CurrentEnum = *EnumIt;
 		if (UEdGraphSchema_K2::IsAllowableBlueprintVariableType(CurrentEnum) && CurrentEnum->HasMetaData(TEXT("Bitflags")))
 		{
-			BitmaskEnumTypeNames.Add(MakeShareable(new FString(CurrentEnum->GetFName().ToString())));
+			BitmaskEnumTypePaths.Add(MakeShareable(new FTopLevelAssetPath(CurrentEnum->GetPackage()->GetFName(), CurrentEnum->GetFName())));
 		}
 	}
 
@@ -760,11 +760,16 @@ void FBlueprintVarActionDetails::CustomizeDetails( IDetailLayoutBuilder& DetailL
 	]
 	.ValueContent()
 	[
-		SNew(STextComboBox)
-		.OptionsSource(&BitmaskEnumTypeNames)
-		.InitiallySelectedItem(GetBitmaskEnumTypeName())
+		SNew(SComboBox<TSharedPtr<FTopLevelAssetPath>>)
+		.OptionsSource(&BitmaskEnumTypePaths)
+		.InitiallySelectedItem(GetBitmaskEnumTypePath())
 		.OnSelectionChanged(this, &FBlueprintVarActionDetails::OnBitmaskEnumTypeChanged)
+		.OnGenerateWidget(this, &FBlueprintVarActionDetails::GenerateBitmaskEnumTypeWidget)
 		.IsEnabled(IsVariableInBlueprint() && OnBitmaskCheckboxState() == ECheckBoxState::Checked)
+		[
+			SNew(STextBlock)
+				.Text(this, &FBlueprintVarActionDetails::GetBitmaskEnumTypeName)
+		]
 	];
 
 	ReplicationOptions.Empty();
@@ -2261,23 +2266,23 @@ void FBlueprintVarActionDetails::OnBitmaskChanged(ECheckBoxState InNewState)
 	}
 }
 
-TSharedPtr<FString> FBlueprintVarActionDetails::GetBitmaskEnumTypeName() const
+TSharedPtr<FTopLevelAssetPath> FBlueprintVarActionDetails::GetBitmaskEnumTypePath() const
 {
-	TSharedPtr<FString> Result;
+	TSharedPtr<FTopLevelAssetPath> Result;
 	const FName VarName = CachedVariableName;
 
-	if (BitmaskEnumTypeNames.Num() > 0 && VarName != NAME_None)
+	if (BitmaskEnumTypePaths.Num() > 0 && VarName != NAME_None)
 	{
-		Result = BitmaskEnumTypeNames[0];
+		Result = BitmaskEnumTypePaths[0];
 
 		FString OutValue;
 		FBlueprintEditorUtils::GetBlueprintVariableMetaData(GetBlueprintObj(), VarName, nullptr, FBlueprintMetadata::MD_BitmaskEnum, OutValue);
-
-		for (int32 i = 1; i < BitmaskEnumTypeNames.Num(); ++i)
+		
+		for (int32 i = 1; i < BitmaskEnumTypePaths.Num(); ++i)
 		{
-			if (OutValue == *BitmaskEnumTypeNames[i])
+			if (OutValue == BitmaskEnumTypePaths[i]->ToString())
 			{
-				Result = BitmaskEnumTypeNames[i];
+				Result = BitmaskEnumTypePaths[i];
 				break;
 			}
 		}
@@ -2286,20 +2291,20 @@ TSharedPtr<FString> FBlueprintVarActionDetails::GetBitmaskEnumTypeName() const
 	return Result;
 }
 
-void FBlueprintVarActionDetails::OnBitmaskEnumTypeChanged(TSharedPtr<FString> ItemSelected, ESelectInfo::Type SelectInfo)
+void FBlueprintVarActionDetails::OnBitmaskEnumTypeChanged(TSharedPtr<FTopLevelAssetPath> ItemSelected, ESelectInfo::Type SelectInfo)
 {
 	const FName VarName = CachedVariableName;
 	if (VarName != NAME_None)
 	{
 		UBlueprint* LocalBlueprint = GetBlueprintObj();
 
-		if (ItemSelected == BitmaskEnumTypeNames[0])
+		if (ItemSelected == BitmaskEnumTypePaths[0])
 		{
 			FBlueprintEditorUtils::RemoveBlueprintVariableMetaData(LocalBlueprint, VarName, nullptr, FBlueprintMetadata::MD_BitmaskEnum);
 		}
 		else if(ItemSelected.IsValid())
 		{
-			FBlueprintEditorUtils::SetBlueprintVariableMetaData(LocalBlueprint, VarName, nullptr, FBlueprintMetadata::MD_BitmaskEnum, *ItemSelected);
+			FBlueprintEditorUtils::SetBlueprintVariableMetaData(LocalBlueprint, VarName, nullptr, FBlueprintMetadata::MD_BitmaskEnum, ItemSelected->ToString());
 		}
 
 		// Reset default value
@@ -2328,6 +2333,19 @@ void FBlueprintVarActionDetails::OnBitmaskEnumTypeChanged(TSharedPtr<FString> It
 	}
 }
 
+TSharedRef<SWidget> FBlueprintVarActionDetails::GenerateBitmaskEnumTypeWidget(TSharedPtr<FTopLevelAssetPath> Item)
+{
+	check(Item.IsValid());
+
+	return SNew(STextBlock)
+		.Text(FText::FromName(Item->GetAssetName()));
+}
+
+FText FBlueprintVarActionDetails::GetBitmaskEnumTypeName() const
+{
+	return FText::FromName(GetBitmaskEnumTypePath()->GetAssetName());
+}
+
 TSharedPtr<FString> FBlueprintVarActionDetails::GetVariableReplicationType() const
 {
 	EVariableReplication::Type VariableReplication = EVariableReplication::None;
diff --git a/Engine/Source/Editor/Kismet/Private/BlueprintDetailsCustomization.h b/Engine/Source/Editor/Kismet/Private/BlueprintDetailsCustomization.h
--- a/Engine/Source/Editor/Kismet/Private/BlueprintDetailsCustomization.h
+++ b/Engine/Source/Editor/Kismet/Private/BlueprintDetailsCustomization.h
@@ -219,8 +219,10 @@ private:
 	EVisibility BitmaskVisibility() const;
 	void OnBitmaskChanged(ECheckBoxState InNewState);
 
-	TSharedPtr<FString> GetBitmaskEnumTypeName() const;
-	void OnBitmaskEnumTypeChanged(TSharedPtr<FString> ItemSelected, ESelectInfo::Type SelectInfo);
+	TSharedPtr<FTopLevelAssetPath> GetBitmaskEnumTypePath() const;
+	void OnBitmaskEnumTypeChanged(TSharedPtr<FTopLevelAssetPath> ItemSelected, ESelectInfo::Type SelectInfo);
+	TSharedRef<SWidget> GenerateBitmaskEnumTypeWidget(TSharedPtr<FTopLevelAssetPath> Item);
+	FText GetBitmaskEnumTypeName() const;
 	
 	TSharedPtr<FString> GetVariableReplicationType() const;
 	void OnChangeReplication(TSharedPtr<FString> ItemSelected, ESelectInfo::Type SelectInfo);
@@ -289,7 +291,7 @@ private:
 	TArray<TSharedPtr<FString>> ReplicationOptions;
 
 	/** Array of enum type names for integers used as bitmasks */
-	TArray<TSharedPtr<FString>> BitmaskEnumTypeNames;
+	TArray<TSharedPtr<FTopLevelAssetPath>> BitmaskEnumTypePaths;
 
 	/** The widget used when in variable name editing mode */ 
 	TSharedPtr<SEditableTextBox> VarNameEditableTextBox;
diff --git a/Engine/Source/Editor/UnrealEd/Private/Kismet2/BlueprintEditorUtils.cpp b/Engine/Source/Editor/UnrealEd/Private/Kismet2/BlueprintEditorUtils.cpp
--- a/Engine/Source/Editor/UnrealEd/Private/Kismet2/BlueprintEditorUtils.cpp
+++ b/Engine/Source/Editor/UnrealEd/Private/Kismet2/BlueprintEditorUtils.cpp
@@ -6225,11 +6225,24 @@ void FBlueprintEditorUtils::FixupVariableDescription(UBlueprint* Blueprint, FBPV
 	// Remove bitflag enum type metadata if the enum type name is missing or if the enum type is no longer a bitflags type.
 	if (VarDesc.HasMetaData(FBlueprintMetadata::MD_BitmaskEnum))
 	{
-		FString BitmaskEnumTypeName = VarDesc.GetMetaData(FBlueprintMetadata::MD_BitmaskEnum);
-		if (!BitmaskEnumTypeName.IsEmpty())
+		FString BitmaskEnumTypePath = VarDesc.GetMetaData(FBlueprintMetadata::MD_BitmaskEnum);
+		if (!BitmaskEnumTypePath.IsEmpty())
 		{
-			UEnum* BitflagsEnum = UClass::TryFindTypeSlow<UEnum>(BitmaskEnumTypeName);
-			if (BitflagsEnum == nullptr || !BitflagsEnum->HasMetaData(*FBlueprintMetadata::MD_Bitflags.ToString()))
+			const UEnum* BitflagsEnum = nullptr;
+			
+			// if the enum is saved by name (deprecated), find the associated enum and reserialize it as a long asset path
+			if (FPackageName::IsShortPackageName(BitmaskEnumTypePath))
+			{
+				BitflagsEnum = FindFirstObject<UEnum>(GetData(BitmaskEnumTypePath));
+				BitmaskEnumTypePath = FTopLevelAssetPath(BitflagsEnum->GetPackage()->GetFName(), BitflagsEnum->GetFName()).ToString();
+				VarDesc.SetMetaData(FBlueprintMetadata::MD_BitmaskEnum, BitmaskEnumTypePath);
+			}
+			else
+			{
+				BitflagsEnum = FindObject<UEnum>(nullptr, GetData(BitmaskEnumTypePath));
+			}
+			
+			if (BitflagsEnum == nullptr || !BitflagsEnum->HasMetaData(*FBlueprintMetadata::MD_Bitflags.ToString()) || !UEdGraphSchema_K2::IsAllowableBlueprintVariableType(BitflagsEnum))
 			{
 				VarDesc.RemoveMetaData(FBlueprintMetadata::MD_BitmaskEnum);
 			}
