[BugFix] BP Bitmask null check fix

diff --git a/Engine/Source/Editor/Kismet/Private/BlueprintDetailsCustomization.cpp b/Engine/Source/Editor/Kismet/Private/BlueprintDetailsCustomization.cpp
--- a/Engine/Source/Editor/Kismet/Private/BlueprintDetailsCustomization.cpp
+++ b/Engine/Source/Editor/Kismet/Private/BlueprintDetailsCustomization.cpp
@@ -2343,7 +2343,8 @@ TSharedRef<SWidget> FBlueprintVarActionDetails::GenerateBitmaskEnumTypeWidget(TS
 
 FText FBlueprintVarActionDetails::GetBitmaskEnumTypeName() const
 {
-	return FText::FromName(GetBitmaskEnumTypePath()->GetAssetName());
+	const TSharedPtr<FTopLevelAssetPath> BitmaskEnumTypePath = GetBitmaskEnumTypePath();
+	return BitmaskEnumTypePath? FText::FromName(BitmaskEnumTypePath->GetAssetName()) : FText();
 }
 
 TSharedPtr<FString> FBlueprintVarActionDetails::GetVariableReplicationType() const
