UE-153955 - Event Dispatcher Inputs do not inform/warn the end-user when there is an existing matching input name.

Event dispatchers don't have a corresponding UFunction, so the FoundFunction result would always be null, which caused the pin rename to always succeed during verification. However, we do have terminator nodes that we can use instead for pin naming conflicts. This also matches what OnPinRenamed already does.

This change also happens to fix the missing warning when macro pins have a name conflict.

diff --git a/Engine/Source/Editor/Kismet/Private/BlueprintDetailsCustomization.cpp b/Engine/Source/Editor/Kismet/Private/BlueprintDetailsCustomization.cpp
--- a/Engine/Source/Editor/Kismet/Private/BlueprintDetailsCustomization.cpp
+++ b/Engine/Source/Editor/Kismet/Private/BlueprintDetailsCustomization.cpp
@@ -4650,10 +4650,26 @@ bool FBaseBlueprintGraphActionDetails::OnVerifyPinRename(UK2Node_EditablePinBase
 
 	if (InTargetNode)
 	{
-		// Check if the name conflicts with any of the other internal UFunction's property names (local variables and parameters).
-		const UFunction* FoundFunction = FFunctionFromNodeHelper::FunctionFromNode(InTargetNode);
-		const FProperty* ExistingProperty = FindFProperty<const FProperty>(FoundFunction, *InNewName);
-		if (ExistingProperty)
+		UK2Node_EditablePinBase* EntryNode = FunctionEntryNodePtr.Get();
+		UK2Node_EditablePinBase* ResultNode = FunctionResultNodePtr.Get();
+		const FName NewFName = *InNewName;
+
+		ERenamePinResult RenameResult = InTargetNode->RenameUserDefinedPin(InOldName, NewFName, true);
+
+		if (RenameResult == ERenamePinResult_Success)
+		{
+			UK2Node_EditablePinBase* OtherNode = (InTargetNode == EntryNode) ? ResultNode : EntryNode;
+
+			// OtherNode can be null if the function, macro, etc. doesn't return a value.
+			if (OtherNode)
+			{
+				RenameResult = OtherNode->RenameUserDefinedPin(InOldName, NewFName, true);
+			}
+		}
+
+		check(RenameResult != ERenamePinResult_NoSuchPin);
+
+		if (RenameResult == ERenamePinResult_NameCollision)
 		{
 			OutErrorMessage = LOCTEXT("ConflictsWithProperty", "Conflicts with another local variable or function parameter!");
 			return false;
