Bugfix: don't assert on ERenamePinResult_NoSuchPin

It's perfectly valid for renaming a pin to return NoSuchPin. The logic in OnVerifyPinRename is specifically looking for name collisions, which we don't allow.

diff --git a/Engine/Source/Editor/Kismet/Private/BlueprintDetailsCustomization.cpp b/Engine/Source/Editor/Kismet/Private/BlueprintDetailsCustomization.cpp
--- a/Engine/Source/Editor/Kismet/Private/BlueprintDetailsCustomization.cpp
+++ b/Engine/Source/Editor/Kismet/Private/BlueprintDetailsCustomization.cpp
@@ -4656,7 +4656,7 @@ bool FBaseBlueprintGraphActionDetails::OnVerifyPinRename(UK2Node_EditablePinBase
 
 		ERenamePinResult RenameResult = InTargetNode->RenameUserDefinedPin(InOldName, NewFName, true);
 
-		if (RenameResult == ERenamePinResult_Success)
+		if (RenameResult != ERenamePinResult_NameCollision)
 		{
 			UK2Node_EditablePinBase* OtherNode = (InTargetNode == EntryNode) ? ResultNode : EntryNode;
 
@@ -4667,8 +4667,6 @@ bool FBaseBlueprintGraphActionDetails::OnVerifyPinRename(UK2Node_EditablePinBase
 			}
 		}
 
-		check(RenameResult != ERenamePinResult_NoSuchPin);
-
 		if (RenameResult == ERenamePinResult_NameCollision)
 		{
 			OutErrorMessage = LOCTEXT("ConflictsWithProperty", "Conflicts with another local variable or function parameter!");
