[BugFix] BP_Diff data flow changes

- because the graphs were being diffed in two different ways - once to generate the list and once to determine how to render the nodes, there were some differences in the diff behavior that caused certain nodes to not render as diffed even though they showed up in the list. this change uses the same diff list for both now so there shouldn't be any differences anymore

- This also solves an issue that caused highlighted pins from Blueprint Diff to persist to the blueprint itself. Now all visual modifications made by the diff viewer only appear in the diff panels rather than the main graph

This change sets up for future style changes allowing for us to render nodes in the diff graph based on their diff type. There are no style changes in this CL however.

diff --git a/Engine/Plugins/Experimental/CommonConversation/Source/CommonConversationEditor/Private/SConversationDiff.cpp b/Engine/Plugins/Experimental/CommonConversation/Source/CommonConversationEditor/Private/SConversationDiff.cpp
--- a/Engine/Plugins/Experimental/CommonConversation/Source/CommonConversationEditor/Private/SConversationDiff.cpp
+++ b/Engine/Plugins/Experimental/CommonConversation/Source/CommonConversationEditor/Private/SConversationDiff.cpp
@@ -96,7 +96,6 @@ public:
 BEGIN_SLATE_FUNCTION_BUILD_OPTIMIZATION
 void SConversationDiff::Construct( const FArguments& InArgs )
 {
-	LastPinTarget = nullptr;
 	LastOtherPinTarget = nullptr;
 
 	FDiffListCommands::Register();
@@ -329,8 +328,6 @@ TSharedRef<ITableRow> SConversationDiff::OnGenerateRow(FSharedDiffOnGraph Item,
 
 void SConversationDiff::OnSelectionChanged(FSharedDiffOnGraph Item, ESelectInfo::Type SelectionType)
 {
-	DisablePinDiffFocus();
-
 	if(!Item.IsValid())
 	{
 		return;
@@ -347,9 +344,6 @@ void SConversationDiff::OnSelectionChanged(FSharedDiffOnGraph Item, ESelectInfo:
 		{
 			if (InPin)
 			{
-				LastPinTarget = InPin;
-				InPin->bIsDiffing = true;
-
 				UEdGraph* NodeGraph = InPin->GetOwningNode()->GetGraph();
 				SGraphEditor* NodeGraphEditor = GetGraphEditorForGraph(NodeGraph);
 				NodeGraphEditor->JumpToPin(InPin);
@@ -417,19 +411,6 @@ SGraphEditor* SConversationDiff::GetGraphEditorForGraph(UEdGraph* Graph) const
 	return nullptr;
 }
 
-void SConversationDiff::DisablePinDiffFocus()
-{
-	if(LastPinTarget)
-	{
-		LastPinTarget->bIsDiffing = false;
-	}
-	if(LastOtherPinTarget)
-	{
-		LastOtherPinTarget->bIsDiffing = false;
-	}
-}
-
-
 //////////////////////////////////////////////////////////////////////////
 // FConversationDiffPanel
 //////////////////////////////////////////////////////////////////////////
diff --git a/Engine/Plugins/Experimental/CommonConversation/Source/CommonConversationEditor/Private/SConversationDiff.h b/Engine/Plugins/Experimental/CommonConversation/Source/CommonConversationEditor/Private/SConversationDiff.h
--- a/Engine/Plugins/Experimental/CommonConversation/Source/CommonConversationEditor/Private/SConversationDiff.h
+++ b/Engine/Plugins/Experimental/CommonConversation/Source/CommonConversationEditor/Private/SConversationDiff.h
@@ -169,10 +169,6 @@ private:
 	 */
 	SGraphEditor* GetGraphEditorForGraph(UEdGraph* Graph) const;
 
-	/**
-	 * Removes focus of the pin in the graph viewer
-	 */
-	void DisablePinDiffFocus();
 
 private:
 
@@ -191,9 +187,6 @@ private:
 	// ListView of differences
 	TSharedPtr<SListViewType> DiffList;
 
-	// The last pin the user clicked on
-	UEdGraphPin* LastPinTarget;
-
 	// The last other pin the user clicked on
 	UEdGraphPin* LastOtherPinTarget;
 };
diff --git a/Engine/Plugins/Experimental/CommonConversation/Source/CommonConversationGraph/Private/SConversationGraphNode.cpp b/Engine/Plugins/Experimental/CommonConversation/Source/CommonConversationGraph/Private/SConversationGraphNode.cpp
--- a/Engine/Plugins/Experimental/CommonConversation/Source/CommonConversationGraph/Private/SConversationGraphNode.cpp
+++ b/Engine/Plugins/Experimental/CommonConversation/Source/CommonConversationGraph/Private/SConversationGraphNode.cpp
@@ -66,7 +66,7 @@ void SConversationPin::Construct(const FArguments& InArgs, UEdGraphPin* InPin)
 FSlateColor SConversationPin::GetPinColor() const
 {
 	return 
-		GraphPinObj->bIsDiffing ? ConversationEditorColors::Pin::Diff :
+		bIsDiffHighlighted ? ConversationEditorColors::Pin::Diff :
 		IsHovered() ? ConversationEditorColors::Pin::Hover :
 		(GraphPinObj->PinType.PinCategory == UConversationGraphTypes::PinCategory_SingleComposite) ? ConversationEditorColors::Pin::CompositeOnly :
 		(GraphPinObj->PinType.PinCategory == UConversationGraphTypes::PinCategory_SingleTask) ? ConversationEditorColors::Pin::TaskOnly :
diff --git a/Engine/Source/Editor/BehaviorTreeEditor/Private/SBehaviorTreeDiff.cpp b/Engine/Source/Editor/BehaviorTreeEditor/Private/SBehaviorTreeDiff.cpp
--- a/Engine/Source/Editor/BehaviorTreeEditor/Private/SBehaviorTreeDiff.cpp
+++ b/Engine/Source/Editor/BehaviorTreeEditor/Private/SBehaviorTreeDiff.cpp
@@ -95,6 +95,7 @@ void SBehaviorTreeDiff::Construct( const FArguments& InArgs )
 {
 	LastPinTarget = NULL;
 	LastOtherPinTarget = NULL;
+	FoundDiffs = MakeShared<TArray<FDiffSingleResult>>();
 
 	FDiffListCommands::Register();
 
@@ -179,8 +180,8 @@ void SBehaviorTreeDiff::Construct( const FArguments& InArgs )
 		]
 	];
 
-	PanelOld.GeneratePanel(PanelOld.BehaviorTree->BTGraph, PanelNew.BehaviorTree->BTGraph);
-	PanelNew.GeneratePanel(PanelNew.BehaviorTree->BTGraph, PanelOld.BehaviorTree->BTGraph);
+	PanelOld.GeneratePanel(PanelOld.BehaviorTree->BTGraph, FoundDiffs);
+	PanelNew.GeneratePanel(PanelNew.BehaviorTree->BTGraph, FoundDiffs);
 }
 END_SLATE_FUNCTION_BUILD_OPTIMIZATION
 
@@ -264,11 +265,11 @@ TSharedRef<SWidget> SBehaviorTreeDiff::GenerateDiffListWidget()
 
 void SBehaviorTreeDiff::BuildDiffSourceArray()
 {
-	TArray<FDiffSingleResult> FoundDiffs;
-	FGraphDiffControl::DiffGraphs(PanelOld.BehaviorTree->BTGraph, PanelNew.BehaviorTree->BTGraph, FoundDiffs);
+	FoundDiffs->Empty();
+	FGraphDiffControl::DiffGraphs(PanelOld.BehaviorTree->BTGraph, PanelNew.BehaviorTree->BTGraph, *FoundDiffs);
 
 	DiffListSource.Empty();
-	for (auto DiffIt(FoundDiffs.CreateConstIterator()); DiffIt; ++DiffIt)
+	for (auto DiffIt(FoundDiffs->CreateConstIterator()); DiffIt; ++DiffIt)
 	{
 		DiffListSource.Add(FSharedDiffOnGraph(new FTreeDiffResultItem(*DiffIt)));
 	}
@@ -322,8 +323,6 @@ TSharedRef<ITableRow> SBehaviorTreeDiff::OnGenerateRow(FSharedDiffOnGraph Item,
 
 void SBehaviorTreeDiff::OnSelectionChanged(FSharedDiffOnGraph Item, ESelectInfo::Type SelectionType)
 {
-	DisablePinDiffFocus();
-
 	if(!Item.IsValid())
 	{
 		return;
@@ -341,7 +340,6 @@ void SBehaviorTreeDiff::OnSelectionChanged(FSharedDiffOnGraph Item, ESelectInfo:
 			if (InPin)
 			{
 				LastPinTarget = InPin;
-				InPin->bIsDiffing = true;
 
 				UEdGraph* NodeGraph = InPin->GetOwningNode()->GetGraph();
 				SGraphEditor* NodeGraphEditor = GetGraphEditorForGraph(NodeGraph);
@@ -410,18 +408,6 @@ SGraphEditor* SBehaviorTreeDiff::GetGraphEditorForGraph(UEdGraph* Graph) const
 	return NULL;
 }
 
-void SBehaviorTreeDiff::DisablePinDiffFocus()
-{
-	if(LastPinTarget)
-	{
-		LastPinTarget->bIsDiffing = false;
-	}
-	if(LastOtherPinTarget)
-	{
-		LastOtherPinTarget->bIsDiffing = false;
-	}
-}
-
 
 //////////////////////////////////////////////////////////////////////////
 // FBehaviorTreeDiffPanel
@@ -432,7 +418,7 @@ SBehaviorTreeDiff::FBehaviorTreeDiffPanel::FBehaviorTreeDiffPanel()
 	BehaviorTree = NULL;
 }
 
-void SBehaviorTreeDiff::FBehaviorTreeDiffPanel::GeneratePanel(UEdGraph* Graph, UEdGraph* GraphToDiff)
+void SBehaviorTreeDiff::FBehaviorTreeDiffPanel::GeneratePanel(UEdGraph* Graph, TSharedPtr<TArray<FDiffSingleResult>> DiffResults)
 {
 	TSharedPtr<SWidget> Widget = SNew(SBorder)
 		.HAlign(HAlign_Center)
@@ -469,7 +455,7 @@ void SBehaviorTreeDiff::FBehaviorTreeDiffPanel::GeneratePanel(UEdGraph* Graph, U
 		auto Editor = SNew(SGraphEditor)
 			.AdditionalCommands(GraphEditorCommands)
 			.GraphToEdit(Graph)
-			.GraphToDiff(GraphToDiff)
+			.DiffResults(DiffResults)
 			.IsEditable(false)
 			.TitleBar(SNew(SBorder).HAlign(HAlign_Center)
 			[
diff --git a/Engine/Source/Editor/BehaviorTreeEditor/Private/SBehaviorTreeDiff.h b/Engine/Source/Editor/BehaviorTreeEditor/Private/SBehaviorTreeDiff.h
--- a/Engine/Source/Editor/BehaviorTreeEditor/Private/SBehaviorTreeDiff.h
+++ b/Engine/Source/Editor/BehaviorTreeEditor/Private/SBehaviorTreeDiff.h
@@ -43,11 +43,11 @@ private:
 		FBehaviorTreeDiffPanel();
 
 		/** 
-		 * Generates the slate for this panel
-		 * @param Graph        The Left side graph.
-		 * @param GraphToDiff  The right side graph to diff against
+		 * Generate a panel that displays the Graph and reflects the items in the DiffResults
+		 * @param Graph        Graph Being Displayed
+		 * @param DiffResults  Attribute for obtaining the list of differences
 		 */
-		void GeneratePanel(UEdGraph* Graph, UEdGraph* GraphToDiff);
+		void GeneratePanel(UEdGraph* Graph, TSharedPtr<TArray<FDiffSingleResult>> DiffResults);
 
 		/**
 		 * Returns title for this panel
@@ -182,6 +182,9 @@ private:
 	// Source for list view 
 	TArray<FSharedDiffOnGraph> DiffListSource;
 
+	// result from FGraphDiffControl::DiffGraphs
+	TSharedPtr<TArray<FDiffSingleResult>> FoundDiffs;
+
 	// Key commands processed by this widget
 	TSharedPtr<FUICommandList> KeyCommands;
 
diff --git a/Engine/Source/Editor/BehaviorTreeEditor/Private/SGraphNode_BehaviorTree.cpp b/Engine/Source/Editor/BehaviorTreeEditor/Private/SGraphNode_BehaviorTree.cpp
--- a/Engine/Source/Editor/BehaviorTreeEditor/Private/SGraphNode_BehaviorTree.cpp
+++ b/Engine/Source/Editor/BehaviorTreeEditor/Private/SGraphNode_BehaviorTree.cpp
@@ -67,7 +67,7 @@ void SBehaviorTreePin::Construct(const FArguments& InArgs, UEdGraphPin* InPin)
 FSlateColor SBehaviorTreePin::GetPinColor() const
 {
 	return 
-		GraphPinObj->bIsDiffing ? BehaviorTreeColors::Pin::Diff :
+		bIsDiffHighlighted ? BehaviorTreeColors::Pin::Diff :
 		IsHovered() ? BehaviorTreeColors::Pin::Hover :
 		(GraphPinObj->PinType.PinCategory == UBehaviorTreeEditorTypes::PinCategory_SingleComposite) ? BehaviorTreeColors::Pin::CompositeOnly :
 		(GraphPinObj->PinType.PinCategory == UBehaviorTreeEditorTypes::PinCategory_SingleTask) ? BehaviorTreeColors::Pin::TaskOnly :
diff --git a/Engine/Source/Editor/GraphEditor/Private/GraphEditorModule.cpp b/Engine/Source/Editor/GraphEditor/Private/GraphEditorModule.cpp
--- a/Engine/Source/Editor/GraphEditor/Private/GraphEditorModule.cpp
+++ b/Engine/Source/Editor/GraphEditor/Private/GraphEditorModule.cpp
@@ -64,7 +64,8 @@ TSharedRef<SGraphEditor> FGraphEditorModule::PRIVATE_MakeGraphEditor(
 	UEdGraph* InGraphToEdit,
 	SGraphEditor::FGraphEditorEvents InGraphEvents,
 	bool InAutoExpandActionMenu,
-	UEdGraph* InGraphToDiff,
+	TSharedPtr<TArray<FDiffSingleResult>> DiffResults,
+	TAttribute<int32> FocusedDiffResult,
 	FSimpleDelegate InOnNavigateHistoryBack,
 	FSimpleDelegate InOnNavigateHistoryForward,
 	TAttribute<bool> ShowGraphStateOverlay)
@@ -79,7 +80,8 @@ TSharedRef<SGraphEditor> FGraphEditorModule::PRIVATE_MakeGraphEditor(
 		.GraphToEdit(InGraphToEdit)
 		.GraphEvents(InGraphEvents)
 		.AutoExpandActionMenu(InAutoExpandActionMenu)
-		.GraphToDiff(InGraphToDiff)
+		.DiffResults(DiffResults)
+		.FocusedDiffResult(FocusedDiffResult)
 		.OnNavigateHistoryBack(InOnNavigateHistoryBack)
 		.OnNavigateHistoryForward(InOnNavigateHistoryForward)
 		.ShowGraphStateOverlay(ShowGraphStateOverlay);
diff --git a/Engine/Source/Editor/GraphEditor/Private/SGraphEditorImpl.cpp b/Engine/Source/Editor/GraphEditor/Private/SGraphEditorImpl.cpp
--- a/Engine/Source/Editor/GraphEditor/Private/SGraphEditorImpl.cpp
+++ b/Engine/Source/Editor/GraphEditor/Private/SGraphEditorImpl.cpp
@@ -694,7 +694,8 @@ void SGraphEditorImpl::Construct( const FArguments& InArgs )
 		[
 			SAssignNew(GraphPanel, SGraphPanel)
 			.GraphObj( EdGraphObj )
-			.GraphObjToDiff( InArgs._GraphToDiff)
+			.DiffResults( InArgs._DiffResults)
+			.FocusedDiffResult( InArgs._FocusedDiffResult)
 			.OnGetContextMenuFor( this, &SGraphEditorImpl::GraphEd_OnGetContextMenuFor )
 			.OnSelectionChanged( InArgs._GraphEvents.OnSelectionChanged )
 			.OnNodeDoubleClicked( InArgs._GraphEvents.OnNodeDoubleClicked )
diff --git a/Engine/Source/Editor/GraphEditor/Private/SGraphEditorImpl.h b/Engine/Source/Editor/GraphEditor/Private/SGraphEditorImpl.h
--- a/Engine/Source/Editor/GraphEditor/Private/SGraphEditorImpl.h
+++ b/Engine/Source/Editor/GraphEditor/Private/SGraphEditorImpl.h
@@ -192,7 +192,8 @@ public:
 		SLATE_ARGUMENT( TSharedPtr<SWidget>, TitleBar )
 		SLATE_ATTRIBUTE( FGraphAppearanceInfo, Appearance )
 		SLATE_ARGUMENT( UEdGraph*, GraphToEdit )
-		SLATE_ARGUMENT( UEdGraph*, GraphToDiff )
+		SLATE_ARGUMENT(TSharedPtr<TArray<FDiffSingleResult>>, DiffResults )
+		SLATE_ATTRIBUTE(int32, FocusedDiffResult )
 		SLATE_ARGUMENT(SGraphEditor::FGraphEditorEvents, GraphEvents)
 		SLATE_ARGUMENT(bool, AutoExpandActionMenu)
 		SLATE_EVENT(FSimpleDelegate, OnNavigateHistoryBack)
diff --git a/Engine/Source/Editor/GraphEditor/Private/SGraphPanel.cpp b/Engine/Source/Editor/GraphEditor/Private/SGraphPanel.cpp
--- a/Engine/Source/Editor/GraphEditor/Private/SGraphPanel.cpp
+++ b/Engine/Source/Editor/GraphEditor/Private/SGraphPanel.cpp
@@ -46,7 +46,8 @@ void SGraphPanel::Construct( const SGraphPanel::FArguments& InArgs )
 
 	this->OnGetContextMenuFor = InArgs._OnGetContextMenuFor;
 	this->GraphObj = InArgs._GraphObj;
-	this->GraphObjToDiff = InArgs._GraphObjToDiff;
+	this->DiffResults = InArgs._DiffResults;
+	this->FocusedDiffResult = InArgs._FocusedDiffResult;
 	this->SelectionManager.OnSelectionChanged = InArgs._OnSelectionChanged;
 	this->IsEditable = InArgs._IsEditable;
 	this->DisplayAsReadOnly = InArgs._DisplayAsReadOnly;
@@ -141,6 +142,41 @@ int32 SGraphPanel::OnPaint( const FPaintArgs& Args, const FGeometry& AllottedGeo
 	const FVector2D NodeShadowSize = GetDefault<UGraphEditorSettings>()->GetShadowDeltaSize();
 	const UEdGraphSchema* Schema = GraphObj->GetSchema();
 
+	
+	// If we were provided diff results, organize those by owner
+	TMap<UEdGraphNode*, FDiffSingleResult> NodeDiffResults;
+	TMap<UEdGraphPin*, FDiffSingleResult> PinDiffResults;
+	if (DiffResults.IsValid())
+	{
+		// diffs with Node1/Pin1 get precedence so set those first
+		for (const FDiffSingleResult& Result : *DiffResults)
+		{
+			if (Result.Pin1)
+			{
+				PinDiffResults.FindOrAdd(Result.Pin1, Result);
+				NodeDiffResults.FindOrAdd(Result.Pin1->GetOwningNode(), Result);
+			}
+			else if (Result.Node1)
+			{
+				NodeDiffResults.FindOrAdd(Result.Node1, Result);
+			}
+		}
+
+		// only diffs with Node2/Pin2 if those nodes don't already have a diff result
+		for (const FDiffSingleResult& Result : *DiffResults)
+		{
+			if (Result.Pin2)
+			{
+				PinDiffResults.FindOrAdd(Result.Pin2, Result);
+				NodeDiffResults.FindOrAdd(Result.Pin2->GetOwningNode(), Result);
+			}
+			else if (Result.Node2)
+			{
+				NodeDiffResults.FindOrAdd(Result.Node2, Result);
+			}
+		}
+	}
+
 	// Draw the child nodes
 	{
 		// When drawing a marquee, need a preview of what the selection will be.
@@ -155,7 +191,6 @@ int32 SGraphPanel::OnPaint( const FPaintArgs& Args, const FGeometry& AllottedGeo
 		// Context for rendering node infos
 		FKismetNodeInfoContext Context(GraphObj);
 
-		TArray<FGraphDiffControl::FNodeMatch> NodeMatches;
 		for (int32 ChildIndex = 0; ChildIndex < ArrangedChildren.Num(); ++ChildIndex)
 		{
 			FArrangedWidget& CurWidget = ArrangedChildren[ChildIndex];
@@ -195,6 +230,27 @@ int32 SGraphPanel::OnPaint( const FPaintArgs& Args, const FGeometry& AllottedGeo
 						ChildNode->ApplyRename();
 					}
 				}
+				
+				/** Highlight Pins that are focused by the diff view */
+				for (UEdGraphPin* Pin : NodeObj->Pins)
+				{
+					if (TSharedPtr<SGraphPin> PinWidget = ChildNode->FindWidgetForPin(Pin))
+					{
+						if (FDiffSingleResult* DiffResult = PinDiffResults.Find(Pin))
+						{
+							// if the diff result associated with this pin is focused, highlight the pin
+							if (DiffResults.IsValid() && FocusedDiffResult.IsSet())
+							{
+								const FDiffSingleResult& Focused = (*DiffResults)[FocusedDiffResult.Get()];
+								PinWidget->SetDiffHighlighted(*DiffResult == Focused);
+							}
+						}
+						else
+						{
+							PinWidget->SetDiffHighlighted(false);
+						}
+					}
+				}
 
 				// Draw the node's shadow.
 				if (bDrawShadowsThisFrame || bSelected)
@@ -230,23 +286,24 @@ int32 SGraphPanel::OnPaint( const FPaintArgs& Args, const FGeometry& AllottedGeo
 
 				int32 CurWidgetsMaxLayerId;
 				{
-					/** When diffing nodes, nodes that are different between revisions are opaque, nodes that have not changed are faded */
-					FGraphDiffControl::FNodeMatch NodeMatch = FGraphDiffControl::FindNodeMatch(GraphObjToDiff, NodeObj, NodeMatches);
-					if (NodeMatch.IsValid())
-					{
-						NodeMatches.Add(NodeMatch);
-					}
-					const bool bNodeIsDifferent = (!GraphObjToDiff || NodeMatch.Diff(FGraphDiffControl::FNodeDiffContext()));
-
 					/* When dragging off a pin, we want to duck the alpha of some nodes */
 					TSharedPtr< SGraphPin > OnlyStartPin = (1 == PreviewConnectorFromPins.Num()) ? PreviewConnectorFromPins[0].FindInGraphPanel(*this) : TSharedPtr< SGraphPin >();
 					const bool bNodeIsNotUsableInCurrentContext = Schema->FadeNodeWhenDraggingOffPin(NodeObj, OnlyStartPin.IsValid() ? OnlyStartPin.Get()->GetPinObj() : nullptr);
 					
-					const FWidgetStyle& NodeStyle = (bNodeIsDifferent && !bNodeIsNotUsableInCurrentContext)? InWidgetStyle : FadedStyle;
-					FWidgetStyle NodeStyleToUse = NodeStyle;
+					const bool bCleanDiff = DiffResults.IsValid() && !NodeDiffResults.Contains(NodeObj);
+					
+					FWidgetStyle NodeStyleToUse = InWidgetStyle;
+					if (bNodeIsNotUsableInCurrentContext)
+					{
+						NodeStyleToUse = FadedStyle;
+					}
+					else if (bCleanDiff)
+					{
+						NodeStyleToUse = FadedStyle;
+					}
 					NodeStyleToUse.BlendColorAndOpacityTint(FLinearColor(1.0f, 1.0f, 1.0f, Alpha));
 
-					// Draw the node.O
+					// Draw the node.
 					CurWidgetsMaxLayerId = CurWidget.Widget->Paint(NewArgs, CurWidget.Geometry, MyCullingRect, OutDrawElements, ChildLayerId, NodeStyleToUse, !DisplayAsReadOnly.Get() && ShouldBeEnabled( bParentEnabled ) );
 				}
 
@@ -1881,7 +1938,6 @@ void SGraphPanel::NotifyGraphChanged(const FEdGraphEditAction& EditAction)
 void SGraphPanel::AddReferencedObjects(FReferenceCollector& Collector)
 {
 	Collector.AddReferencedObject( GraphObj );
-	Collector.AddReferencedObject( GraphObjToDiff );
 }
 
 FString SGraphPanel::GetReferencerName() const
diff --git a/Engine/Source/Editor/GraphEditor/Private/SGraphPin.cpp b/Engine/Source/Editor/GraphEditor/Private/SGraphPin.cpp
--- a/Engine/Source/Editor/GraphEditor/Private/SGraphPin.cpp
+++ b/Engine/Source/Editor/GraphEditor/Private/SGraphPin.cpp
@@ -108,6 +108,8 @@ SGraphPin::SGraphPin()
 
 	static const FName NAME_Pin_Background("Graph.Pin.Background");
 	static const FName NAME_Pin_BackgroundHovered("Graph.Pin.BackgroundHovered");
+	
+	static const FName NAME_Pin_DiffOutline("Graph.Pin.DiffHighlight");
 
 	static const FName NAME_PosePin_Connected("Graph.PosePin.Connected");
 	static const FName NAME_PosePin_Disconnected("Graph.PosePin.Disconnected");
@@ -1059,7 +1061,7 @@ const FSlateBrush* SGraphPin::GetPinBorder() const
 		}
 	}
 	UEdGraphPin* GraphPin = GetPinObj();
-	return (IsHovered() || bIsMarkedPin || (GraphPin && GraphPin->bIsDiffing) || bOnlyShowDefaultValue) ? CachedImg_Pin_BackgroundHovered : CachedImg_Pin_Background;
+	return (IsHovered() || bIsMarkedPin || bIsDiffHighlighted || bOnlyShowDefaultValue) ? CachedImg_Pin_BackgroundHovered : CachedImg_Pin_Background;
 }
 
 
@@ -1068,7 +1070,7 @@ FSlateColor SGraphPin::GetPinColor() const
 	UEdGraphPin* GraphPin = GetPinObj();
 	if (GraphPin && !GraphPin->IsPendingKill())
 	{
-		if (GraphPin->bIsDiffing)
+		if (bIsDiffHighlighted)
 		{
 			return FSlateColor(FLinearColor(0.9f, 0.2f, 0.15f));
 		}
diff --git a/Engine/Source/Editor/GraphEditor/Public/GraphEditorModule.h b/Engine/Source/Editor/GraphEditor/Public/GraphEditorModule.h
--- a/Engine/Source/Editor/GraphEditor/Public/GraphEditorModule.h
+++ b/Engine/Source/Editor/GraphEditor/Public/GraphEditorModule.h
@@ -44,7 +44,8 @@ private:
 		UEdGraph* InGraphToEdit,
 		SGraphEditor::FGraphEditorEvents GraphEvents,
 		bool InAutoExpandActionMenu,
-		UEdGraph* InGraphToDiff,
+		TSharedPtr<TArray<FDiffSingleResult>> DiffResults,
+		TAttribute<int32> FocusedDiffResult,
 		FSimpleDelegate InOnNavigateHistoryBack,
 		FSimpleDelegate InOnNavigateHistoryForward,
 		TAttribute<bool> ShowGraphStateOverlay);
diff --git a/Engine/Source/Editor/GraphEditor/Public/SGraphPanel.h b/Engine/Source/Editor/GraphEditor/Public/SGraphPanel.h
--- a/Engine/Source/Editor/GraphEditor/Public/SGraphPanel.h
+++ b/Engine/Source/Editor/GraphEditor/Public/SGraphPanel.h
@@ -55,7 +55,6 @@ public:
 		, _OnSelectionChanged()
 		, _OnNodeDoubleClicked()
 		, _GraphObj( static_cast<UEdGraph*>(NULL) )
-		, _GraphObjToDiff( static_cast<UEdGraph*>(NULL) )
 		, _InitialZoomToFit( false )
 		, _IsEditable( true )
 		, _DisplayAsReadOnly( false )
@@ -71,7 +70,8 @@ public:
 		SLATE_EVENT( SGraphEditor::FOnDropActor, OnDropActor )
 		SLATE_EVENT( SGraphEditor::FOnDropStreamingLevel, OnDropStreamingLevel )
 		SLATE_ARGUMENT( class UEdGraph*, GraphObj )
-		SLATE_ARGUMENT( class UEdGraph*, GraphObjToDiff )
+		SLATE_ARGUMENT( TSharedPtr<TArray<FDiffSingleResult>>, DiffResults )
+		SLATE_ATTRIBUTE( int32, FocusedDiffResult )
 		SLATE_ARGUMENT( bool, InitialZoomToFit )
 		SLATE_ATTRIBUTE( bool, IsEditable )
 		SLATE_ATTRIBUTE( bool, DisplayAsReadOnly )
@@ -220,7 +220,11 @@ private:
 
 protected:
 	UEdGraph* GraphObj;
-	UEdGraph* GraphObjToDiff;//if it exists, this is 
+	
+	// if this graph is displaying the results of a diff, this will provide info
+	// on how to display the nodes
+	TSharedPtr<TArray<FDiffSingleResult>> DiffResults;
+	TAttribute<int32> FocusedDiffResult;
 
 	// Should we ignore the OnStopMakingConnection unless forced?
 	bool bPreservePinPreviewConnection;
diff --git a/Engine/Source/Editor/GraphEditor/Public/SGraphPin.h b/Engine/Source/Editor/GraphEditor/Public/SGraphPin.h
--- a/Engine/Source/Editor/GraphEditor/Public/SGraphPin.h
+++ b/Engine/Source/Editor/GraphEditor/Public/SGraphPin.h
@@ -170,6 +170,11 @@ public:
 		PinColorModifier = InColor;
 	}
 
+	void SetDiffHighlighted(bool bHighlighted)
+	{
+		bIsDiffHighlighted = bHighlighted;
+	}
+
 	/** Set this pin to only be used to display default value */
 	void SetOnlyShowDefaultValue(bool bNewOnlyShowDefaultValue);
 
@@ -354,4 +359,7 @@ protected:
 
 	/** TRUE if the pin should allow any drag and drop */
 	bool bDragAndDropEnabled;
+
+	/** True if this pin is being diffed and it's currently selected in the diff view. Highlights this pin */
+	bool bIsDiffHighlighted;
 };
diff --git a/Engine/Source/Editor/Kismet/Private/SBlueprintDiff.cpp b/Engine/Source/Editor/Kismet/Private/SBlueprintDiff.cpp
--- a/Engine/Source/Editor/Kismet/Private/SBlueprintDiff.cpp
+++ b/Engine/Source/Editor/Kismet/Private/SBlueprintDiff.cpp
@@ -646,6 +646,7 @@ struct FGraphToDiff	: public TSharedFromThis<FGraphToDiff>, IDiffControl
 
 	/** Source for list view */
 	TArray<TSharedPtr<FDiffResultItem>> DiffListSource;
+	TSharedPtr<TArray<FDiffSingleResult>> FoundDiffs;
 
 private:
 	/** Get tooltip for category */
@@ -677,7 +678,7 @@ private:
 };
 
 FGraphToDiff::FGraphToDiff(SBlueprintDiff* InDiffWidget, UEdGraph* InGraphOld, UEdGraph* InGraphNew, const FRevisionInfo& InRevisionOld, const FRevisionInfo& InRevisionNew)
-	: DiffWidget(InDiffWidget), GraphOld(InGraphOld), GraphNew(InGraphNew), RevisionOld(InRevisionOld), RevisionNew(InRevisionNew)
+	: FoundDiffs(MakeShared<TArray<FDiffSingleResult>>()), DiffWidget(InDiffWidget), GraphOld(InGraphOld), GraphNew(InGraphNew), RevisionOld(InRevisionOld), RevisionNew(InRevisionNew)
 {
 	check(InGraphOld || InGraphNew); //one of them needs to exist
 
@@ -793,24 +794,24 @@ void FGraphToDiff::GenerateTreeEntries(TArray< TSharedPtr<FBlueprintDifferenceTr
 
 void FGraphToDiff::BuildDiffSourceArray()
 {
-	TArray<FDiffSingleResult> FoundDiffs;
-	FGraphDiffControl::DiffGraphs(GraphOld, GraphNew, FoundDiffs);
-
-	DiffListSource.Empty();
-	for (const FDiffSingleResult& Diff : FoundDiffs)
-	{
-		DiffListSource.Add(MakeShared<FDiffResultItem>(Diff));
-	}
+	FoundDiffs->Empty();
+	FGraphDiffControl::DiffGraphs(GraphOld, GraphNew, *FoundDiffs);
 
 	struct SortDiff
 	{
-		bool operator () (const TSharedPtr<FDiffResultItem>& A, const TSharedPtr<FDiffResultItem>& B) const
+		bool operator () (const FDiffSingleResult& A, const FDiffSingleResult& B) const
 		{
-			return A->Result.Diff < B->Result.Diff;
+			return A.Diff < B.Diff;
 		}
 	};
 
-	Sort(DiffListSource.GetData(), DiffListSource.Num(), SortDiff());
+	Sort(FoundDiffs->GetData(), FoundDiffs->Num(), SortDiff());
+
+	DiffListSource.Empty();
+	for (const FDiffSingleResult& Diff : *FoundDiffs)
+	{
+		DiffListSource.Add(MakeShared<FDiffResultItem>(Diff));
+	}
 }
 
 void FGraphToDiff::OnGraphChanged( const FEdGraphEditAction& Action )
@@ -821,7 +822,6 @@ void FGraphToDiff::OnGraphChanged( const FEdGraphEditAction& Action )
 FDiffPanel::FDiffPanel()
 {
 	Blueprint = nullptr;
-	LastFocusedPin = nullptr;
 }
 
 void FDiffPanel::InitializeDiffPanel()
@@ -1314,19 +1314,20 @@ void SBlueprintDiff::ResetGraphEditors()
 	}
 }
 
-void FDiffPanel::GeneratePanel(UEdGraph* Graph, UEdGraph* GraphToDiff )
+void FDiffPanel::GeneratePanel(UEdGraph* NewGraph, UEdGraph* OldGraph )
+{
+	const TSharedPtr<TArray<FDiffSingleResult>> Diff = MakeShared<TArray<FDiffSingleResult>>();
+	FGraphDiffControl::DiffGraphs(OldGraph, NewGraph, *Diff);
+	GeneratePanel(NewGraph, Diff, {});
+}
+
+void FDiffPanel::GeneratePanel(UEdGraph* Graph, TSharedPtr<TArray<FDiffSingleResult>> DiffResults, TAttribute<int32> FocusedDiffResult)
 {
 	if( GraphEditor.IsValid() && GraphEditor.Pin()->GetCurrentGraph() == Graph )
 	{
 		return;
 	}
 
-	if( LastFocusedPin )
-	{
-		LastFocusedPin->bIsDiffing = false;
-	}
-	LastFocusedPin = nullptr;
-
 	TSharedPtr<SWidget> Widget = SNew(SBorder)
 								.HAlign(HAlign_Center)
 								.VAlign(VAlign_Center)
@@ -1369,7 +1370,9 @@ void FDiffPanel::GeneratePanel(UEdGraph* Graph, UEdGraph* GraphToDiff )
 		TSharedRef<SGraphEditor> Editor = SNew(SGraphEditor)
 			.AdditionalCommands(GraphEditorCommands)
 			.GraphToEdit(Graph)
-			.GraphToDiff(GraphToDiff)
+			.GraphToDiff(nullptr)
+			.DiffResults(DiffResults)
+			.FocusedDiffResult(FocusedDiffResult)
 			.IsEditable(false)
 			.GraphEvents(InEvents);
 
@@ -1423,24 +1426,11 @@ bool FDiffPanel::CanCopyNodes() const
 
 void FDiffPanel::FocusDiff(UEdGraphPin& Pin)
 {
-	if( LastFocusedPin )
-	{
-		LastFocusedPin->bIsDiffing = false;
-	}
-	Pin.bIsDiffing = true;
-	LastFocusedPin = &Pin;
-
 	GraphEditor.Pin()->JumpToPin(&Pin);
 }
 
 void FDiffPanel::FocusDiff(UEdGraphNode& Node)
 {
-	if (LastFocusedPin)
-	{
-		LastFocusedPin->bIsDiffing = false;
-	}
-	LastFocusedPin = nullptr;
-
 	if (GraphEditor.IsValid())
 	{
 		GraphEditor.Pin()->JumpToNode(&Node, false);
@@ -1468,32 +1458,30 @@ void SBlueprintDiff::HandleGraphChanged( const FString& GraphPath )
 {
 	SetCurrentMode(GraphMode);
 	
-	TArray<UEdGraph*> GraphsOld, GraphsNew;
-	PanelOld.Blueprint->GetAllGraphs(GraphsOld);
-	PanelNew.Blueprint->GetAllGraphs(GraphsNew);
-
 	UEdGraph* GraphOld = nullptr;
-	for (UEdGraph* OldGraph : GraphsOld)
-	{
-		if (GraphPath.Equals(FGraphDiffControl::GetGraphPath(OldGraph)))
-		{
-			GraphOld = OldGraph;
-			break;
-		}
-	}
-
 	UEdGraph* GraphNew = nullptr;
-	for (UEdGraph* NewGraph : GraphsNew)
+	TSharedPtr<TArray<FDiffSingleResult>> DiffResults;
+	TAttribute<int32> FocusedDiffResult;
+	for (const TSharedPtr<FGraphToDiff>& GraphToDiff : Graphs)
 	{
+		UEdGraph* NewGraph = GraphToDiff->GetGraphNew();
+		UEdGraph* OldGraph = GraphToDiff->GetGraphOld();
 		if (GraphPath.Equals(FGraphDiffControl::GetGraphPath(NewGraph)))
 		{
 			GraphNew = NewGraph;
+			GraphOld = OldGraph;
+			DiffResults = GraphToDiff->FoundDiffs;
+			FocusedDiffResult = TAttribute<int32>::CreateLambda(
+				[DifferencesTreeView = DifferencesTreeView.ToSharedRef(), &RealDifferences = RealDifferences]()
+				{
+					return DiffTreeView::CurrentDifference(DifferencesTreeView, RealDifferences);
+				});
 			break;
 		}
 	}
 
-	PanelOld.GeneratePanel(GraphOld, GraphNew);
-	PanelNew.GeneratePanel(GraphNew, GraphOld);
+	PanelOld.GeneratePanel(GraphOld, DiffResults, FocusedDiffResult);
+	PanelNew.GeneratePanel(GraphNew, DiffResults, FocusedDiffResult);
 }
 
 void SBlueprintDiff::GenerateDifferencesList()
diff --git a/Engine/Source/Editor/Kismet/Public/SBlueprintDiff.h b/Engine/Source/Editor/Kismet/Public/SBlueprintDiff.h
--- a/Engine/Source/Editor/Kismet/Public/SBlueprintDiff.h
+++ b/Engine/Source/Editor/Kismet/Public/SBlueprintDiff.h
@@ -51,8 +51,11 @@ struct KISMET_API FDiffPanel
 	/** Initializes the panel, can be moved into constructor if diff and merge clients are made more uniform: */
 	void InitializeDiffPanel();
 
-	/** Generate this panel based on the specified graph */
-	void GeneratePanel(UEdGraph* Graph, UEdGraph* GraphToDiff);
+	/** Generate a panel for NewGraph diffed against OldGraph */
+	void GeneratePanel(UEdGraph* NewGraph, UEdGraph* OldGraph);
+	
+	/** Generate a panel that displays the Graph and reflects the items in the DiffResults */
+	void GeneratePanel(UEdGraph* Graph, TSharedPtr<TArray<FDiffSingleResult>> DiffResults, TAttribute<int32> FocusedDiffResult);
 
 	/** Generate the 'MyBlueprint' widget, which is private to this module */
 	TSharedRef<class SWidget> GenerateMyBlueprintWidget();
@@ -91,9 +94,6 @@ struct KISMET_API FDiffPanel
 	/** True if we should show a name identifying which asset this panel is displaying */
 	bool							bShowAssetName;
 
-	/** The panel stores the last pin that was focused on by the user, so that it can clear the visual style when selection changes */
-	UEdGraphPin*					LastFocusedPin;
-
 	/** The widget that contains the revision info in graph mode */
 	TSharedPtr<SWidget>				OverlayGraphRevisionInfo;
 private:
diff --git a/Engine/Source/Editor/UnrealEd/Private/GraphEditor.cpp b/Engine/Source/Editor/UnrealEd/Private/GraphEditor.cpp
--- a/Engine/Source/Editor/UnrealEd/Private/GraphEditor.cpp
+++ b/Engine/Source/Editor/UnrealEd/Private/GraphEditor.cpp
@@ -28,7 +28,8 @@ void SGraphEditor::ConstructImplementation( const FArguments& InArgs )
 		InArgs._GraphToEdit,
 		InArgs._GraphEvents,
 		InArgs._AutoExpandActionMenu,
-		InArgs._GraphToDiff,
+		InArgs._DiffResults,
+		InArgs._FocusedDiffResult,
 		InArgs._OnNavigateHistoryBack,
 		InArgs._OnNavigateHistoryForward,
 		InArgs._ShowGraphStateOverlay
diff --git a/Engine/Source/Editor/UnrealEd/Public/GraphEditor.h b/Engine/Source/Editor/UnrealEd/Public/GraphEditor.h
--- a/Engine/Source/Editor/UnrealEd/Public/GraphEditor.h
+++ b/Engine/Source/Editor/UnrealEd/Public/GraphEditor.h
@@ -137,7 +137,6 @@ public:
 		, _DisplayAsReadOnly(false)
 		, _IsEmpty(false)
 		, _GraphToEdit(NULL)
-		, _GraphToDiff(NULL)
 		, _AutoExpandActionMenu(false)
 		, _ShowGraphStateOverlay(true)
 		{}
@@ -150,7 +149,15 @@ public:
 		SLATE_ATTRIBUTE( FGraphAppearanceInfo, Appearance )
 		SLATE_EVENT( FEdGraphEvent, OnGraphModuleReloaded )
 		SLATE_ARGUMENT( UEdGraph*, GraphToEdit )
+	
+		PRAGMA_DISABLE_DEPRECATION_WARNINGS
+		UE_DEPRECATED(5.1, "GraphToDiff is no longer supported. Use DiffResults instead")
 		SLATE_ARGUMENT( UEdGraph*, GraphToDiff )
+		PRAGMA_ENABLE_DEPRECATION_WARNINGS
+	
+		SLATE_ARGUMENT( TSharedPtr<TArray<FDiffSingleResult>>, DiffResults )
+		SLATE_ATTRIBUTE( int32, FocusedDiffResult )
+	
 		SLATE_ARGUMENT( FGraphEditorEvents, GraphEvents)
 		SLATE_ARGUMENT( bool, AutoExpandActionMenu )
 		SLATE_ARGUMENT( TWeakPtr<FAssetEditorToolkit>, AssetEditorToolkit)
diff --git a/Engine/Source/Runtime/Engine/Classes/EdGraph/EdGraphPin.h b/Engine/Source/Runtime/Engine/Classes/EdGraph/EdGraphPin.h
--- a/Engine/Source/Runtime/Engine/Classes/EdGraph/EdGraphPin.h
+++ b/Engine/Source/Runtime/Engine/Classes/EdGraph/EdGraphPin.h
@@ -331,9 +331,9 @@ public:
 	/** If true, the default value on this pin is ignored and should not be set. */
 	uint8 bDefaultValueIsIgnored:1;
 
-	/** If true, this pin is the focus of a diff. This is transient. */
-	uint8 bIsDiffing:1;
-
+	UE_DEPRECATED(5.0, "bIsDiffing is deprecated. Please use SGraphPanel::DiffResults and SGraphPanel::FocusedDiffResult instead")
+	uint32 bIsDiffing:1;
+	
 	/** If true, the pin may be hidden by user. */
 	uint8 bAdvancedView:1;
 
@@ -670,10 +670,6 @@ class UEdGraphPin_Deprecated : public UObject
 	UPROPERTY()
 	uint32 bDefaultValueIsIgnored:1;
 
-	/** If true, this pin is the focus of a diff */
-	UPROPERTY(transient)
-	uint32 bIsDiffing:1;
-
 	/** If true, the pin may be hidden by user */
 	UPROPERTY()
 	uint32 bAdvancedView:1;
diff --git a/Engine/Source/Runtime/Engine/Private/EdGraph/EdGraphPin.cpp b/Engine/Source/Runtime/Engine/Private/EdGraph/EdGraphPin.cpp
--- a/Engine/Source/Runtime/Engine/Private/EdGraph/EdGraphPin.cpp
+++ b/Engine/Source/Runtime/Engine/Private/EdGraph/EdGraphPin.cpp
@@ -1410,7 +1410,6 @@ UEdGraphPin::UEdGraphPin(UEdGraphNode* InOwningNode, const FGuid& PinIdGuid)
 	, bNotConnectable(false)
 	, bDefaultValueIsReadOnly(false)
 	, bDefaultValueIsIgnored(false)
-	, bIsDiffing(false)
 	, bAdvancedView(false)
 	, bDisplayAsMutableRef(false)
 	, bAllowFriendlyName(true)
@@ -1570,7 +1569,6 @@ void UEdGraphPin::InitFromDeprecatedPin(class UEdGraphPin_Deprecated* Deprecated
 	bNotConnectable = DeprecatedPin->bNotConnectable;
 	bDefaultValueIsReadOnly = DeprecatedPin->bDefaultValueIsReadOnly;
 	bDefaultValueIsIgnored = DeprecatedPin->bDefaultValueIsIgnored;
-	bIsDiffing = DeprecatedPin->bIsDiffing;
 	bAdvancedView = DeprecatedPin->bAdvancedView;
 	bDisplayAsMutableRef = DeprecatedPin->bDisplayAsMutableRef;
 	bUseBackwardsCompatForEmptyAutogeneratedValue = true; // If it is from deprecated pin format then it is definitely old enough to need this
@@ -1736,9 +1734,8 @@ bool UEdGraphPin::Serialize(FArchive& Ar)
 		if (Ar.IsTransacting())
 		{
 			BitField |= bDisplayAsMutableRef << (PersistentBits + 0);
-			BitField |= bIsDiffing << (PersistentBits + 1);
-			BitField |= bSavePinIfOrphaned << (PersistentBits + 2);
-			BitField |= bWasTrashed << (PersistentBits + 3);
+			BitField |= bSavePinIfOrphaned << (PersistentBits + 1);
+			BitField |= bWasTrashed << (PersistentBits + 2);
 		}
 
 		Ar << BitField;
@@ -1756,9 +1753,8 @@ bool UEdGraphPin::Serialize(FArchive& Ar)
 		if (Ar.IsTransacting())
 		{
 			bDisplayAsMutableRef = !!(BitField & (1 << (PersistentBits + 0)));
-			bIsDiffing = !!(BitField & (1 << (PersistentBits + 1)));
-			bSavePinIfOrphaned = !!(BitField & (1 << (PersistentBits + 2)));
-			bWasTrashed = !!(BitField & (1 << (PersistentBits + 3)));
+			bSavePinIfOrphaned = !!(BitField & (1 << (PersistentBits + 1)));
+			bWasTrashed = !!(BitField & (1 << (PersistentBits + 2)));
 		}
 
 		if (Ar.IsLoading() && Ar.CustomVer(FFrameworkObjectVersion::GUID) < FFrameworkObjectVersion::ChangeAssetPinsToString)
