Always clear the debugger actor set if there is no play world

diff --git a/Engine/Source/Editor/Kismet/Private/Debugging/SKismetDebuggingView.cpp b/Engine/Source/Editor/Kismet/Private/Debugging/SKismetDebuggingView.cpp
--- a/Engine/Source/Editor/Kismet/Private/Debugging/SKismetDebuggingView.cpp
+++ b/Engine/Source/Editor/Kismet/Private/Debugging/SKismetDebuggingView.cpp
@@ -378,6 +378,14 @@ void SKismetDebuggingView::Tick(const FGeometry& AllottedGeometry, const double
 		return;
 	}
 
+	// If there is no play world, immediately clear the list. this avoids showing a phantom 'nullptr' item when ending PIE
+	const bool bIsDebugging = GEditor->PlayWorld != nullptr;
+	if (!bIsDebugging && DebugTreeView->GetRootTreeItems().Num() > 0)
+	{
+		DebugTreeView->ClearTreeItems();
+		return;
+	}
+
 	// update less often to avoid lag
 	TreeUpdateTimer += InDeltaTime;
 	if (TreeUpdateTimer < UpdateInterval)
@@ -397,8 +405,6 @@ void SKismetDebuggingView::Tick(const FGeometry& AllottedGeometry, const double
 	}
 
 	// Gather what we'd like to be the new root set
-	const bool bIsDebugging = GEditor->PlayWorld != nullptr;
-
 	TSet<UObject*> NewRootSet;
 
 	const auto TryAddBlueprintToNewRootSet = [&NewRootSet](UBlueprint* InBlueprint)
