[BugFix] EXCEPTION_ACCESS_VIOLATION Crash when using Blueprint Diff when there's recently deleted or created graphs in the diff

diff --git a/Engine/Source/Editor/Kismet/Private/SBlueprintDiff.cpp b/Engine/Source/Editor/Kismet/Private/SBlueprintDiff.cpp
--- a/Engine/Source/Editor/Kismet/Private/SBlueprintDiff.cpp
+++ b/Engine/Source/Editor/Kismet/Private/SBlueprintDiff.cpp
@@ -1465,7 +1465,8 @@ void SBlueprintDiff::HandleGraphChanged( const FString& GraphPath )
 	{
 		UEdGraph* NewGraph = GraphToDiff->GetGraphNew();
 		UEdGraph* OldGraph = GraphToDiff->GetGraphOld();
-		if (GraphPath.Equals(FGraphDiffControl::GetGraphPath(NewGraph)))
+		const FString OtherGraphPath = NewGraph ? FGraphDiffControl::GetGraphPath(NewGraph) : FGraphDiffControl::GetGraphPath(OldGraph);
+		if (GraphPath.Equals(OtherGraphPath))
 		{
 			GraphNew = NewGraph;
 			GraphOld = OldGraph;
@@ -1482,7 +1483,10 @@ void SBlueprintDiff::HandleGraphChanged( const FString& GraphPath )
         	int32 startIndex = 0;
 			for (const TSharedPtr<FGraphToDiff>& GraphToDiff : Graphs)
 			{
-				if (GraphPath.Equals(FGraphDiffControl::GetGraphPath(GraphToDiff->GetGraphNew())))
+				UEdGraph* NewGraph = GraphToDiff->GetGraphNew();
+				UEdGraph* OldGraph = GraphToDiff->GetGraphOld();
+				const FString OtherGraphPath = NewGraph ? FGraphDiffControl::GetGraphPath(NewGraph) : FGraphDiffControl::GetGraphPath(OldGraph);
+				if (GraphPath.Equals(OtherGraphPath))
 				{
 					break;
 				}
