Fix a namespace-related editor crash when removing map/set elements out-of-order in a Blueprint variable's default value.

diff --git a/Engine/Source/Editor/Kismet/Private/BlueprintNamespaceUtilities.cpp b/Engine/Source/Editor/Kismet/Private/BlueprintNamespaceUtilities.cpp
--- a/Engine/Source/Editor/Kismet/Private/BlueprintNamespaceUtilities.cpp
+++ b/Engine/Source/Editor/Kismet/Private/BlueprintNamespaceUtilities.cpp
@@ -191,17 +191,17 @@ void FBlueprintNamespaceUtilities::GetPropertyValueNamespaces(const UStruct* InS
 		else if (const FSetProperty* SetProperty = CastField<FSetProperty>(InProperty))
 		{
 			FScriptSetHelper SetHelper(SetProperty, ValuePtr);
-			for (int32 ValueIdx = 0; ValueIdx < SetHelper.Num(); ++ValueIdx)
+			for (FScriptSetHelper::FIterator SetIt = SetHelper.CreateIterator(); SetIt; ++SetIt)
 			{
-				GetPropertyValueNamespaces(InStruct, SetProperty->ElementProp, SetHelper.GetElementPtr(ValueIdx), OutNamespaces);
+				GetPropertyValueNamespaces(InStruct, SetProperty->ElementProp, SetHelper.GetElementPtr(*SetIt), OutNamespaces);
 			}
 		}
 		else if (const FMapProperty* MapProperty = CastField<FMapProperty>(InProperty))
 		{
 			FScriptMapHelper MapHelper(MapProperty, ValuePtr);
-			for (int32 ValueIdx = 0; ValueIdx < MapHelper.Num(); ++ValueIdx)
+			for (FScriptMapHelper::FIterator MapIt = MapHelper.CreateIterator(); MapIt; ++MapIt)
 			{
-				const uint8* MapValuePtr = MapHelper.GetPairPtr(ValueIdx);
+				const uint8* MapValuePtr = MapHelper.GetPairPtr(*MapIt);
 				GetPropertyValueNamespaces(InStruct, MapProperty->KeyProp, MapValuePtr, OutNamespaces);
 				GetPropertyValueNamespaces(InStruct, MapProperty->ValueProp, MapValuePtr, OutNamespaces);
 			}
