Fix for crash in blueprint diffing tool when diffing shelved changes where the blueprint's new parent class doesn't exist locally

diff --git a/Engine/Source/Editor/Kismet/Private/SBlueprintDiff.cpp b/Engine/Source/Editor/Kismet/Private/SBlueprintDiff.cpp
--- a/Engine/Source/Editor/Kismet/Private/SBlueprintDiff.cpp
+++ b/Engine/Source/Editor/Kismet/Private/SBlueprintDiff.cpp
@@ -362,7 +362,7 @@ public:
 		{
 			if (OldBlueprint->ParentClass != NewBlueprint->ParentClass)
 			{
-				FText DiffText = FText::Format(LOCTEXT("ParentChanged", "Parent Class changed from {0} to {1}"), FText::FromString(GetNameSafe(OldBlueprint->ParentClass)), FText::FromString(NewBlueprint->ParentClass->GetName()));
+				FText DiffText = FText::Format(LOCTEXT("ParentChanged", "Parent Class changed from {0} to {1}"), FText::FromString(GetNameSafe(OldBlueprint->ParentClass)), FText::FromString(GetNameSafe(NewBlueprint->ParentClass)));
 
 				TSharedPtr<FBlueprintDifferenceTreeEntry> Entry = MakeShared<FBlueprintDifferenceTreeEntry>(
 					SelectionCallback,
