Add support for selecting/importing inclusive namespace paths.

Context:
- Given namespaces "A.B.C" and "A.B.C.D" assigned to different type assets, previously users would have to import each namespace explicitly to load/display those types in the Blueprint editor.
- Users can now choose to import the "inclusive" path "A.B" which implicitly includes types assigned to both "A.B.C" and "A.B.C.D" (along with any new namespaces that might later appear under "A.B").

Notes:
- Rather than a Java-style "wildcard" path (e.g. "import A.B.*"), this instead resembles the syntax for a C#-style "inclusive" path (e.g. "import A.B").
- Users can both import assets from and assign assets to an inclusive path using the same UI. Assigning to an inclusive path will also register it as an in-use namespace.
- Blueprints may import both the inclusive path (e.g. "A.B") and an actual assigned path under that subtree (e.g. "A.B.C"), but the latter is superseded by the inclusive path unless/until it is removed as an import.

Change summary:
- Refactored SBlueprintNamespaceEntry to use an STreeView. This allows inclusive namespace paths to be selected (e.g. for import).
- Exposed an optional tooltip attribute to SBlueprintNamespaceEntry for hierarchical namespace paths that are visible but not selectable.
- Modified SBlueprintNamespaceEntry::OnTextCommitted() to ensure that manually-entered names correct case to match registered paths.
- Modified SBlueprintNamespaceEntry::OnShowingSuggestions() to consider excluded paths when building the suggestion list.
- Avoid unnecessarily repopulating the registered path list when typing into SBlueprintNamespaceEntry's drop-down search filter box.
- Added FBlueprintNamespaceRegistry::IsInclusivePath().
- Fixed up existing UI areas that incorporate an SBlueprintNamespaceEntry widget to gather an exclusion list using the new slate event.
- Fixed up contextual import UI areas to additionally exclude shared global import lists (UE-149390).
- Fixed up FBlueprintNamespacePathTree::FNodeVisitorFunc to pass an immutable reference to each node.
- Fixed up FBlueprintNamespaceRegistry::FindAndRegisterAllNamespaces() to also consider loaded Blueprint macro library assets.
- Added FBlueprintNamespacePathTree::PathSeparator and replaced namespace-specific occurrences of TEXT(".") usage for consistency.

diff --git a/Engine/Source/Editor/Kismet/Private/BlueprintDetailsCustomization.cpp b/Engine/Source/Editor/Kismet/Private/BlueprintDetailsCustomization.cpp
--- a/Engine/Source/Editor/Kismet/Private/BlueprintDetailsCustomization.cpp
+++ b/Engine/Source/Editor/Kismet/Private/BlueprintDetailsCustomization.cpp
@@ -5644,7 +5644,8 @@ TSharedPtr<SWidget> FBlueprintImportsLayout::MakeAddItemWidget()
 	return SNew(SBlueprintNamespaceEntry)
 		.AllowTextEntry(false)
 		.OnNamespaceSelected(this, &FBlueprintImportsLayout::OnNamespaceSelected)
-		.OnFilterNamespaceList(this, &FBlueprintImportsLayout::OnFilterNamespaceList)
+		.OnGetNamespacesToExclude(this, &FBlueprintImportsLayout::OnGetNamespacesToExclude)
+		.ExcludedNamespaceTooltipText(LOCTEXT("CannotSelectNamespaceForImport", "This namespace is already imported."))
 		.ButtonContent()
 		[
 			SNew(STextBlock)
@@ -5716,14 +5717,14 @@ void FBlueprintImportsLayout::OnNamespaceSelected(const FString& InNamespace)
 	OnRefreshInDetailsView();
 }
 
-void FBlueprintImportsLayout::OnFilterNamespaceList(TArray<FString>& InOutNamespaceList)
+void FBlueprintImportsLayout::OnGetNamespacesToExclude(TSet<FString>& OutNamespacesToExclude) const
 {
 	const UBlueprint* Blueprint = GetBlueprintObjectChecked();
-	InOutNamespaceList.RemoveSwap(Blueprint->BlueprintNamespace);
-	for (const FString& ImportedNamespace : Blueprint->ImportedNamespaces)
-	{
-		InOutNamespaceList.RemoveSwap(ImportedNamespace);
-	}
+
+	FBlueprintNamespaceUtilities::GetSharedGlobalImports(OutNamespacesToExclude);
+	FBlueprintNamespaceUtilities::GetDefaultImportsForBlueprint(Blueprint, OutNamespacesToExclude);
+
+	OutNamespacesToExclude.Append(Blueprint->ImportedNamespaces);
 }
 
 
@@ -6043,7 +6044,7 @@ void FBlueprintGlobalOptionsDetails::HandleNamespaceValueChange(const FString& I
 
 	if (!InNewValue.IsEmpty() && !BlueprintNamespaceRegistry.IsRegisteredPath(InNewValue))
 	{
-		// Add the new namespace into the registry (it has not been seen yet).
+		// Add the new namespace into the registry (it has not been explicitly added yet).
 		BlueprintNamespaceRegistry.RegisterNamespace(InNewValue);
 	}
 	
@@ -6056,8 +6057,8 @@ void FBlueprintGlobalOptionsDetails::HandleNamespaceValueChange(const FString& I
 			bool bRefreshDetailsView = false;
 			if (Blueprint->ImportedNamespaces.Contains(InOldValue))
 			{
-				// Remove the import from the current editor context if it is no longer registered.
-				if (!BlueprintNamespaceRegistry.IsRegisteredPath(InOldValue))
+				// Remove the import from the current editor context if it is no longer inclusive of any path.
+				if (!BlueprintNamespaceRegistry.IsInclusivePath(InOldValue))
 				{
 					BlueprintEditor->RemoveNamespace(InOldValue);
 				}
diff --git a/Engine/Source/Editor/Kismet/Private/BlueprintDetailsCustomization.h b/Engine/Source/Editor/Kismet/Private/BlueprintDetailsCustomization.h
--- a/Engine/Source/Editor/Kismet/Private/BlueprintDetailsCustomization.h
+++ b/Engine/Source/Editor/Kismet/Private/BlueprintDetailsCustomization.h
@@ -719,7 +719,7 @@ protected:
 	/** END FBlueprintManagedListDetails interface */
 
 	void OnNamespaceSelected(const FString& InNamespace);
-	void OnFilterNamespaceList(TArray<FString>& InOutNamespaceList);
+	void OnGetNamespacesToExclude(TSet<FString>& OutNamespacesToExclude) const;
 
 private:
 	/** Whether we should show default (i.e. global) imports versus custom (i.e. local) imports */
diff --git a/Engine/Source/Editor/Kismet/Private/BlueprintManagedListDetails.cpp b/Engine/Source/Editor/Kismet/Private/BlueprintManagedListDetails.cpp
--- a/Engine/Source/Editor/Kismet/Private/BlueprintManagedListDetails.cpp
+++ b/Engine/Source/Editor/Kismet/Private/BlueprintManagedListDetails.cpp
@@ -12,6 +12,7 @@
 #include "Widgets/Text/STextBlock.h"
 #include "Subsystems/AssetEditorSubsystem.h"
 #include "Editor.h"
+#include "Algo/Sort.h"
 
 void FBlueprintManagedListDetails::GenerateHeaderRowContent(FDetailWidgetRow& HeaderRow)
 {
@@ -54,6 +55,11 @@ void FBlueprintManagedListDetails::GenerateChildContent(IDetailChildrenBuilder&
 	TArray<FManagedListItem> Items;
 	GetManagedListItems(Items);
 
+	Algo::Sort(Items, [](const FManagedListItem& ItemA, const FManagedListItem& ItemB)
+	{
+		return ItemA.ItemName.Compare(ItemB.ItemName) < 0;
+	});
+
 	if (Items.Num() > 0)
 	{
 		for (const FManagedListItem& Item : Items)
diff --git a/Engine/Source/Editor/Kismet/Private/BlueprintNamespacePathTree.h b/Engine/Source/Editor/Kismet/Private/BlueprintNamespacePathTree.h
--- a/Engine/Source/Editor/Kismet/Private/BlueprintNamespacePathTree.h
+++ b/Engine/Source/Editor/Kismet/Private/BlueprintNamespacePathTree.h
@@ -11,6 +11,9 @@
 struct FBlueprintNamespacePathTree
 {
 public:
+	/** Path separator; can be used for path parsing as well as path construction. */
+	inline static const TCHAR PathSeparator[] = TEXT(".");
+
 	/** Path tree node structure. */
 	struct FNode : TSharedFromThis<FNode>
 	{
@@ -60,7 +63,7 @@ public:
 		TSharedPtr<FNode> Node = GetRootNode();
 
 		TArray<FString> PathSegments;
-		InPath.ParseIntoArray(PathSegments, TEXT("."));
+		InPath.ParseIntoArray(PathSegments, PathSeparator);
 		for (const FString& PathSegment : PathSegments)
 		{
 			if (const TSharedPtr<FNode>* ChildNodePtr = Node->Children.Find(FName(*PathSegment)))
@@ -92,7 +95,7 @@ public:
 		TSharedRef<FNode> Node = GetRootNode();
 
 		TArray<FString> PathSegments;
-		InPath.ParseIntoArray(PathSegments, TEXT("."));
+		InPath.ParseIntoArray(PathSegments, PathSeparator);
 		for (const FString& PathSegment : PathSegments)
 		{
 			Node = Node->FindOrAddChild(FName(*PathSegment));
@@ -109,7 +112,7 @@ public:
 	void RemovePath(const FString& InPath)
 	{
 		TArray<FString> PathSegments;
-		InPath.ParseIntoArray(PathSegments, TEXT("."));
+		InPath.ParseIntoArray(PathSegments, PathSeparator);
 
 		using FPathTraceNode = TTuple<FName, TSharedPtr<FNode>>;
 		TArray<FPathTraceNode> PathTrace;
@@ -174,16 +177,16 @@ public:
 	/** Path node visitor function signature.
 	 * 
 	 * @param CurrentPath	Current path (represented as a stack of names).
-	 * @param Node			A reference to the node at the current visitor level.
+	 * @param Node			A read-only reference to the node at the current visitor level.
 	 */
-	typedef TFunctionRef<void(const TArray<FName>& /* CurrentPath */, TSharedRef<FBlueprintNamespacePathTree::FNode> /* Node */)> FNodeVisitorFunc;
+	typedef TFunctionRef<void(const TArray<FName>& /* CurrentPath */, TSharedRef<const FBlueprintNamespacePathTree::FNode> /* Node */)> FNodeVisitorFunc;
 
 	/**
 	 * A utility method that will recursively visit all added nodes.
 	 * 
 	 * @param VisitorFunc	A function that will be called for each visited node.
 	 */
-	void ForeachNode(FNodeVisitorFunc VisitorFunc)
+	void ForeachNode(FNodeVisitorFunc VisitorFunc) const
 	{
 		TArray<FName> CurrentPath;
 		RecursiveNodeVisitor(GetRootNode(), CurrentPath, VisitorFunc);
@@ -191,7 +194,7 @@ public:
 
 protected:
 	/** Helper method for recursively visiting all nodes. */
-	void RecursiveNodeVisitor(TSharedPtr<FBlueprintNamespacePathTree::FNode> Node, TArray<FName>& CurrentPath, FNodeVisitorFunc VisitorFunc)
+	void RecursiveNodeVisitor(TSharedPtr<FBlueprintNamespacePathTree::FNode> Node, TArray<FName>& CurrentPath, FNodeVisitorFunc VisitorFunc) const
 	{
 		for (auto ChildIt = Node->Children.CreateConstIterator(); ChildIt; ++ChildIt)
 		{
diff --git a/Engine/Source/Editor/Kismet/Private/BlueprintNamespaceRegistry.cpp b/Engine/Source/Editor/Kismet/Private/BlueprintNamespaceRegistry.cpp
--- a/Engine/Source/Editor/Kismet/Private/BlueprintNamespaceRegistry.cpp
+++ b/Engine/Source/Editor/Kismet/Private/BlueprintNamespaceRegistry.cpp
@@ -149,8 +149,16 @@ void FBlueprintNamespaceRegistry::OnReloadComplete(EReloadCompleteReason InReaso
 	Rebuild();
 }
 
+bool FBlueprintNamespaceRegistry::IsInclusivePath(const FString& InPath) const
+{
+	// A path is considered inclusive if it represents any valid subpath in the tree.
+	TSharedPtr<FBlueprintNamespacePathTree::FNode> Node = PathTree->FindPathNode(InPath);
+	return Node.IsValid();
+}
+
 bool FBlueprintNamespaceRegistry::IsRegisteredPath(const FString& InPath) const
 {
+	// A path is considered to be registered only if it was explicitly added to the tree.
 	TSharedPtr<FBlueprintNamespacePathTree::FNode> Node = PathTree->FindPathNode(InPath);
 	return Node.IsValid() && Node->bIsAddedPath;
 }
@@ -169,7 +177,7 @@ void FBlueprintNamespaceRegistry::GetNamesUnderPath(const FString& InPath, TArra
 
 void FBlueprintNamespaceRegistry::GetAllRegisteredPaths(TArray<FString>& OutPaths) const
 {
-	PathTree->ForeachNode([&OutPaths](const TArray<FName>& CurrentPath, TSharedRef<FBlueprintNamespacePathTree::FNode> Node)
+	PathTree->ForeachNode([&OutPaths](const TArray<FName>& CurrentPath, TSharedRef<const FBlueprintNamespacePathTree::FNode> Node)
 	{
 		if (Node->bIsAddedPath)
 		{
@@ -179,7 +187,7 @@ void FBlueprintNamespaceRegistry::GetAllRegisteredPaths(TArray<FString>& OutPath
 			{
 				if (PathBuilder.Len() > 0)
 				{
-					PathBuilder += TEXT(".");
+					PathBuilder += FBlueprintNamespacePathTree::PathSeparator;
 				}
 				PathBuilder += PathSegment.ToString();
 			}
@@ -198,18 +206,7 @@ void FBlueprintNamespaceRegistry::FindAndRegisterAllNamespaces()
 	{
 		if (const UObject* ExcludedObject = ExcludedObjectPath.ResolveObject())
 		{
-			if (const UBlueprint* ExcludedBlueprintObject = Cast<UBlueprint>(ExcludedObject))
-			{
-				// Redirect to the generated class when excluding a loaded Blueprint asset.
-				if (ExcludedBlueprintObject->GeneratedClass)
-				{
-					ExcludedObjects.Add(ExcludedBlueprintObject->GeneratedClass);
-				}
-			}
-			else
-			{
-				ExcludedObjects.Add(ExcludedObject);
-			}
+			ExcludedObjects.Add(ExcludedObject);
 		}
 	}
 
@@ -223,12 +220,25 @@ void FBlueprintNamespaceRegistry::FindAndRegisterAllNamespaces()
 			UEdGraphSchema_K2::IsAllowableBlueprintVariableType(ClassObject)
 			|| ClassObject->IsChildOf<UBlueprintFunctionLibrary>();
 
-		if (bShouldRegisterClassObject && !ExcludedObjects.Contains(ClassObject))
+		const UBlueprint* BlueprintObject = UBlueprint::GetBlueprintFromClass(ClassObject);
+
+		if (bShouldRegisterClassObject && !ExcludedObjects.Contains(ClassObject) && (!BlueprintObject || !ExcludedObjects.Contains(BlueprintObject)))
 		{
 			RegisterNamespace(ClassObject);
 		}
 	}
 
+	// Register loaded macro asset namespace identifiers.
+	for (TObjectIterator<UBlueprint> BlueprintIt; BlueprintIt; ++BlueprintIt)
+	{
+		const UBlueprint* BlueprintObject = *BlueprintIt;
+
+		if (BlueprintObject->BlueprintType == BPTYPE_MacroLibrary && !ExcludedObjects.Contains(BlueprintObject))
+		{
+			RegisterNamespace(BlueprintObject);
+		}
+	}
+
 	// Register loaded struct type namespace identifiers.
 	for (TObjectIterator<UScriptStruct> StructIt; StructIt; ++StructIt)
 	{
diff --git a/Engine/Source/Editor/Kismet/Private/BlueprintNamespaceUtilities.cpp b/Engine/Source/Editor/Kismet/Private/BlueprintNamespaceUtilities.cpp
--- a/Engine/Source/Editor/Kismet/Private/BlueprintNamespaceUtilities.cpp
+++ b/Engine/Source/Editor/Kismet/Private/BlueprintNamespaceUtilities.cpp
@@ -1,6 +1,7 @@
 // Copyright Epic Games, Inc. All Rights Reserved.
 
 #include "BlueprintNamespaceUtilities.h"
+#include "BlueprintNamespacePathTree.h"
 #include "Engine/Blueprint.h"
 #include "AssetRegistry/AssetRegistryModule.h"
 #include "AssetRegistry/AssetData.h"
@@ -23,8 +24,8 @@ namespace UE::Editor::Kismet::Private
 void FBlueprintNamespaceUtilities::ConvertPackagePathToNamespacePath(const FString& InPackagePath, FString& OutNamespacePath)
 {
 	OutNamespacePath = InPackagePath;
-	OutNamespacePath.ReplaceCharInline(TEXT('/'), TEXT('.'));
-	if (OutNamespacePath.StartsWith(TEXT(".")))
+	OutNamespacePath.ReplaceCharInline(TEXT('/'), FBlueprintNamespacePathTree::PathSeparator[0]);
+	if (OutNamespacePath.StartsWith(FBlueprintNamespacePathTree::PathSeparator))
 	{
 		OutNamespacePath.RemoveAt(0);
 	}
diff --git a/Engine/Source/Editor/Kismet/Private/SBlueprintActionMenu.cpp b/Engine/Source/Editor/Kismet/Private/SBlueprintActionMenu.cpp
--- a/Engine/Source/Editor/Kismet/Private/SBlueprintActionMenu.cpp
+++ b/Engine/Source/Editor/Kismet/Private/SBlueprintActionMenu.cpp
@@ -274,8 +274,9 @@ void SBlueprintActionMenu::Construct( const FArguments& InArgs, TSharedPtr<FBlue
 			[
 				SNew(SBlueprintNamespaceEntry)
 					.AllowTextEntry(false)
-					.OnFilterNamespaceList(this, &SBlueprintActionMenu::OnFilterImportNamespaceList)
 					.OnNamespaceSelected(this, &SBlueprintActionMenu::OnNamespaceSelectedForImport)
+					.OnGetNamespacesToExclude(this, &SBlueprintActionMenu::OnGetNamespacesToExcludeFromImportMenu)
+					.ExcludedNamespaceTooltipText(LOCTEXT("CannotSelectNamespaceForImport", "This namespace has already been imported by this Blueprint."))
 			]
 		];
 	}
@@ -635,19 +636,17 @@ void SBlueprintActionMenu::TryInsertPromoteToVariable(FBlueprintActionContext co
 	}
 }
 
-void SBlueprintActionMenu::OnFilterImportNamespaceList(TArray<FString>& InOutNamespaceList)
+void SBlueprintActionMenu::OnGetNamespacesToExcludeFromImportMenu(TSet<FString>& OutNamespacesToExclude)
 {
 	FBlueprintActionContext MenuContext;
 	ConstructActionContext(MenuContext);
 
+	FBlueprintNamespaceUtilities::GetSharedGlobalImports(OutNamespacesToExclude);
+
 	for (const UBlueprint* Blueprint : MenuContext.Blueprints)
 	{
-		InOutNamespaceList.RemoveSwap(Blueprint->BlueprintNamespace);
-
-		for (const FString& ImportedNamespace : Blueprint->ImportedNamespaces)
-		{
-			InOutNamespaceList.RemoveSwap(ImportedNamespace);
-		}
+		FBlueprintNamespaceUtilities::GetDefaultImportsForBlueprint(Blueprint, OutNamespacesToExclude);
+		OutNamespacesToExclude.Append(Blueprint->ImportedNamespaces);
 	}
 }
 
diff --git a/Engine/Source/Editor/Kismet/Private/SBlueprintActionMenu.h b/Engine/Source/Editor/Kismet/Private/SBlueprintActionMenu.h
--- a/Engine/Source/Editor/Kismet/Private/SBlueprintActionMenu.h
+++ b/Engine/Source/Editor/Kismet/Private/SBlueprintActionMenu.h
@@ -72,12 +72,12 @@ protected:
 	/** Function to try to insert a promote to variable entry if it is possible to do so. */
 	void TryInsertPromoteToVariable(FBlueprintActionContext const& Context, FGraphActionListBuilderBase& OutAllActions);
 
-	/** Callback to filter the list of namespaces available for import. */
-	void OnFilterImportNamespaceList(TArray<FString>& InOutNamespaceList);
-
 	/** Invoked when an available namespace is selected for import. */
 	void OnNamespaceSelectedForImport(const FString& InNamespace);
 
+	/** Gathers the list of namespaces to exclude from the import menu. */
+	void OnGetNamespacesToExcludeFromImportMenu(TSet<FString>& OutNamespacesToExclude);
+
 private:
 	UEdGraph* GraphObj;
 	TArray<UEdGraphPin*> DraggedFromPins;
diff --git a/Engine/Source/Editor/Kismet/Private/SBlueprintNamespaceEntry.cpp b/Engine/Source/Editor/Kismet/Private/SBlueprintNamespaceEntry.cpp
--- a/Engine/Source/Editor/Kismet/Private/SBlueprintNamespaceEntry.cpp
+++ b/Engine/Source/Editor/Kismet/Private/SBlueprintNamespaceEntry.cpp
@@ -2,6 +2,7 @@
 
 #include "SBlueprintNamespaceEntry.h"
 #include "BlueprintNamespaceRegistry.h"
+#include "BlueprintNamespacePathTree.h"
 #include "Framework/Application/SlateApplication.h"
 #include "Misc/TextFilterExpressionEvaluator.h"
 #include "Widgets/Input/SComboButton.h"
@@ -18,7 +19,8 @@ void SBlueprintNamespaceEntry::Construct(const FArguments& InArgs)
 {
 	CurrentNamespace = InArgs._CurrentNamespace;
 	OnNamespaceSelected = InArgs._OnNamespaceSelected;
-	OnFilterNamespaceList = InArgs._OnFilterNamespaceList;
+	OnGetNamespacesToExclude = InArgs._OnGetNamespacesToExclude;
+	ExcludedNamespaceTooltipText = InArgs._ExcludedNamespaceTooltipText;
 
 	ChildSlot
 	[
@@ -40,7 +42,7 @@ void SBlueprintNamespaceEntry::Construct(const FArguments& InArgs)
 		[
 			SAssignNew(ComboButton, SComboButton)
 			.CollapseMenuOnParentFocus(true)
-			.OnGetMenuContent(this, &SBlueprintNamespaceEntry::OnGetNamespaceListMenuContent)
+			.OnGetMenuContent(this, &SBlueprintNamespaceEntry::OnGetNamespaceTreeMenuContent)
 			.ButtonContent()
 			[
 				InArgs._ButtonContent.Widget
@@ -94,6 +96,35 @@ void SBlueprintNamespaceEntry::OnTextChanged(const FText& InText)
 
 void SBlueprintNamespaceEntry::OnTextCommitted(const FText& NewText, ETextCommit::Type InTextCommit)
 {
+	// Ensure that case correctly matches up with any registered path name(s).
+	FString CaseCorrectedPath;
+	TArray<FString> PathSegments;
+	CurrentNamespace.ParseIntoArray(PathSegments, FBlueprintNamespacePathTree::PathSeparator);
+	for (FString& PathSegment : PathSegments)
+	{
+		TArray<FName> ActualNames;
+		FBlueprintNamespaceRegistry::Get().GetNamesUnderPath(CaseCorrectedPath, ActualNames);
+		for (const FName& ActualName : ActualNames)
+		{
+			FString ActualNameAsString = ActualName.ToString();
+			if (PathSegment.Equals(ActualNameAsString, ESearchCase::IgnoreCase))
+			{
+				PathSegment = MoveTemp(ActualNameAsString);
+				break;
+			}
+		}
+
+		if (!CaseCorrectedPath.IsEmpty())
+		{
+			CaseCorrectedPath += FBlueprintNamespacePathTree::PathSeparator;
+		}
+
+		CaseCorrectedPath += PathSegment;
+	}
+
+	// Update the current namespace string.
+	CurrentNamespace = MoveTemp(CaseCorrectedPath);
+
 	// Not using the current textbox value here because it might be invalid, and we want to revert to the last-known valid namespace string on commit.
 	SelectNamespace(CurrentNamespace);
 }
@@ -103,7 +134,7 @@ void SBlueprintNamespaceEntry::OnShowingSuggestions(const FString& InputText, TA
 	int32 PathEnd;
 	FString CurrentPath;
 	FString CurrentName;
-	if (InputText.FindLastChar(TEXT('.'), PathEnd))
+	if (InputText.FindLastChar(FBlueprintNamespacePathTree::PathSeparator[0], PathEnd))
 	{
 		CurrentPath = InputText.LeftChop(InputText.Len() - PathEnd);
 		CurrentName = InputText.RightChop(PathEnd + 1);
@@ -120,6 +151,10 @@ void SBlueprintNamespaceEntry::OnShowingSuggestions(const FString& InputText, TA
 	// Sort the list alphabetically.
 	Algo::Sort(SuggestedNames, FNameLexicalLess());
 
+	// Allow the owner to exclude one or more paths.
+	TSet<FString> ExcludedPaths;
+	OnGetNamespacesToExclude.ExecuteIfBound(ExcludedPaths);
+
 	// Build the suggestion set based on the set of matching names we found above.
 	TStringBuilder<128> PathBuilder;
 	for (FName SuggestedName : SuggestedNames)
@@ -130,43 +165,64 @@ void SBlueprintNamespaceEntry::OnShowingSuggestions(const FString& InputText, TA
 			if (CurrentPath.Len() > 0)
 			{
 				PathBuilder += CurrentPath;
-				PathBuilder += TEXT(".");
+				PathBuilder += FBlueprintNamespacePathTree::PathSeparator;
 			}
 
 			PathBuilder += SuggestedNameAsString;
 
 			FString SuggestedNamespace = PathBuilder.ToString();
-			OutSuggestions.Add(MoveTemp(SuggestedNamespace));
+			if (!ExcludedPaths.Contains(SuggestedNamespace))
+			{
+				OutSuggestions.Add(MoveTemp(SuggestedNamespace));
+			}
 
 			PathBuilder.Reset();
 		}
 	}
 }
 
-TSharedRef<SWidget> SBlueprintNamespaceEntry::OnGetNamespaceListMenuContent()
+TSharedRef<SWidget> SBlueprintNamespaceEntry::OnGetNamespaceTreeMenuContent()
 {
-	// Find and filter all registered paths.
-	PopulateNamespaceList();
+	// Gather the full set of registered namespace paths.
+	AllRegisteredPaths.Reset();
+	FBlueprintNamespaceRegistry::Get().GetAllRegisteredPaths(AllRegisteredPaths);
+
+	// Allow external owners to filter the list.
+	ExcludedTreeViewPaths.Reset();
+	OnGetNamespacesToExclude.ExecuteIfBound(ExcludedTreeViewPaths);
+	if (!ExcludedTreeViewPaths.IsEmpty())
+	{
+		AllRegisteredPaths.RemoveAllSwap([&ExcludedTreeViewPaths = this->ExcludedTreeViewPaths](const FString& Path)
+		{
+			return ExcludedTreeViewPaths.Contains(Path);
+		});
+	}
+
+	// Sort the list alphabetically.
+	AllRegisteredPaths.Sort();
+
+	// Build the namespace item tree from the filtered list of registered paths.
+	PopulateNamespaceTree();
 
-	// Construct the list view widget that we'll use for the menu content.
-	SAssignNew(ListView, SListView<TSharedPtr<FString>>)
+	// Construct the tree view widget that we'll use for the menu content.
+	SAssignNew(TreeView, STreeView<TSharedPtr<FPathTreeNodeItem>>)
 		.SelectionMode(ESelectionMode::Single)
-		.ListItemsSource(&ListItems)
-		.OnGenerateRow(this, &SBlueprintNamespaceEntry::OnGenerateRowForNamespaceList)
-		.OnSelectionChanged(this, &SBlueprintNamespaceEntry::OnNamespaceListSelectionChanged);
+		.TreeItemsSource(&RootNodes)
+		.OnGenerateRow(this, &SBlueprintNamespaceEntry::OnGenerateRowForNamespaceTreeItem)
+		.OnGetChildren(this, &SBlueprintNamespaceEntry::OnGetChildrenForNamespaceTreeItem)
+		.OnSelectionChanged(this, &SBlueprintNamespaceEntry::OnNamespaceTreeSelectionChanged)
+		.OnIsSelectableOrNavigable(this, &SBlueprintNamespaceEntry::OnIsNamespaceTreeItemSelectable);
+
+	// All tree view items are always expanded by default.
+	ExpandAllTreeViewItems();
 
 	// If we are allowing manual entry and the current namespace is non-empty, look for a matching item in the set.
 	if (!CurrentNamespace.IsEmpty() && TextBox.IsValid() && TextBox->GetVisibility().IsVisible())
 	{
-		const TSharedPtr<FString>* CurrentItemPtr = ListItems.FindByPredicate([&CurrentNamespace = this->CurrentNamespace](const TSharedPtr<FString>& Item)
-		{
-			return Item.IsValid() && *Item == CurrentNamespace;
-		});
-
 		// If we found a match, make it the initial selection.
-		if (CurrentItemPtr)
+		if (const TSharedPtr<FPathTreeNodeItem>* CurrentItemPtr = FindTreeViewNode(CurrentNamespace))
 		{
-			ListView->SetSelection(*CurrentItemPtr);
+			TreeView->SetSelection(*CurrentItemPtr);
 		}
 	}
 
@@ -181,7 +237,7 @@ TSharedRef<SWidget> SBlueprintNamespaceEntry::OnGetNamespaceListMenuContent()
 				.AutoHeight()
 				[
 					SAssignNew(SearchBox, SSearchBox)
-					.OnTextChanged(this, &SBlueprintNamespaceEntry::OnNamespaceListFilterTextChanged)
+					.OnTextChanged(this, &SBlueprintNamespaceEntry::OnNamespaceTreeFilterTextChanged)
 				]
 				+SVerticalBox::Slot()
 				.AutoHeight()
@@ -191,48 +247,69 @@ TSharedRef<SWidget> SBlueprintNamespaceEntry::OnGetNamespaceListMenuContent()
 				+SVerticalBox::Slot()
 				.FillHeight(1.0f)
 				[
-					ListView.ToSharedRef()
+					TreeView.ToSharedRef()
 				]
 			]
 		];
 }
 
-TSharedRef<ITableRow> SBlueprintNamespaceEntry::OnGenerateRowForNamespaceList(TSharedPtr<FString> Item, const TSharedRef<STableViewBase>& OwnerTable)
+TSharedRef<ITableRow> SBlueprintNamespaceEntry::OnGenerateRowForNamespaceTreeItem(TSharedPtr<FPathTreeNodeItem> Item, const TSharedRef<STableViewBase>& OwnerTable)
 {
 	check(Item.IsValid());
 
-	// Check for an empty list and add a single (disabled) item if found.
+	// Check for an empty tree and add a single (disabled) item if found.
 	bool bIsEnabled = true;
-	FText ItemText = FText::FromString(*Item);
-	if (ItemText.IsEmpty() && ListItems.Num() == 1)
+	FText ItemText;
+	if (ItemText.IsEmpty() && RootNodes.Num() == 1)
 	{
 		bIsEnabled = false;
 		ItemText = LOCTEXT("BlueprintNamespaceList_NoItems", "No Matching Items");
 	}
+	else
+	{
+		ItemText = FText::FromString(*Item->NodePath);
+	}
+
+	FText ToolTipText = FText::GetEmpty();
+	if (!Item->bIsSelectable && ExcludedNamespaceTooltipText.IsSet())
+	{
+		ToolTipText = ExcludedNamespaceTooltipText.Get();
+	}
 
 	// Construct a new row widget, highlighting any text that matches the search filter.
-	return SNew(STableRow<TSharedPtr<FName>>, OwnerTable)
-		.IsEnabled(bIsEnabled)
-		[
-			SNew(STextBlock)
-			.Text(ItemText)
-			.HighlightText(bIsEnabled && SearchBox.IsValid() ? SearchBox->GetText() : FText::GetEmpty())
-		];
+	return SNew(STableRow<TSharedPtr<FPathTreeNodeItem>>, OwnerTable)
+	.IsEnabled(bIsEnabled)
+	.ShowSelection(Item->bIsSelectable)
+	[
+		SNew(STextBlock)
+		.Text(ItemText)
+		.ToolTipText(ToolTipText)
+		.IsEnabled(Item->bIsSelectable)
+		.HighlightText(bIsEnabled && SearchBox.IsValid() ? SearchBox->GetText() : FText::GetEmpty())
+	];
 }
 
-void SBlueprintNamespaceEntry::OnNamespaceListFilterTextChanged(const FText& InText)
+void SBlueprintNamespaceEntry::OnGetChildrenForNamespaceTreeItem(TSharedPtr<FPathTreeNodeItem> Item, TArray<TSharedPtr<FPathTreeNodeItem>>& OutChildren)
+{
+	check(Item.IsValid());
+	OutChildren.Append(Item->ChildNodes);
+}
+
+void SBlueprintNamespaceEntry::OnNamespaceTreeFilterTextChanged(const FText& InText)
 {
 	// Gather/filter all registered paths.
-	PopulateNamespaceList();
+	PopulateNamespaceTree();
 
-	// Refresh the namespace item list view.
-	if (ListView.IsValid())
+	// Refresh the namespace item tree view.
+	if (TreeView.IsValid())
 	{
-		ListView->RequestListRefresh();
+		ExpandAllTreeViewItems();
+
+		TreeView->RequestTreeRefresh();
 	}
 }
 
-void SBlueprintNamespaceEntry::OnNamespaceListSelectionChanged(TSharedPtr<FString> Item, ESelectInfo::Type SelectInfo)
+void SBlueprintNamespaceEntry::OnNamespaceTreeSelectionChanged(TSharedPtr<FPathTreeNodeItem> Item, ESelectInfo::Type SelectInfo)
 {
 	// These actions should not trigger a selection.
 	if (SelectInfo == ESelectInfo::OnNavigation || SelectInfo == ESelectInfo::Direct)
@@ -242,7 +319,13 @@ void SBlueprintNamespaceEntry::OnNamespaceListSelectionChanged(TSharedPtr<FStrin
 
 	if (Item.IsValid())
 	{
-		SelectNamespace(*Item);
+		// Disallow selection of inclusive nodes that are filtered out.
+		if (!Item->bIsSelectable)
+		{
+			return;
+		}
+
+		SelectNamespace(*Item->NodePath);
 	}
 
 	// Clear the search filter text.
@@ -265,44 +348,74 @@ void SBlueprintNamespaceEntry::OnNamespaceListSelectionChanged(TSharedPtr<FStrin
 	}
 }
 
+bool SBlueprintNamespaceEntry::OnIsNamespaceTreeItemSelectable(TSharedPtr<FPathTreeNodeItem> Item) const
+{
+	return Item.IsValid() && Item->bIsSelectable;
+}
+
 FText SBlueprintNamespaceEntry::GetCurrentNamespaceText() const
 {
 	return FText::FromString(CurrentNamespace);
 }
 
-void SBlueprintNamespaceEntry::PopulateNamespaceList()
+void SBlueprintNamespaceEntry::PopulateNamespaceTree()
 {
 	// Clear the current list.
-	ListItems.Empty();
-
-	// Gather the full set of registered namespace paths.
-	TArray<FString> AllPaths;
-	FBlueprintNamespaceRegistry::Get().GetAllRegisteredPaths(AllPaths);
-
-	// Invoke the delegate to allow owners to filter the list as needed.
-	OnFilterNamespaceList.ExecuteIfBound(AllPaths);
-
-	// Sort the list alphabetically.
-	AllPaths.Sort();
+	RootNodes.Reset();
 
 	// Set up an expression evaluator to further trim the list according to the search filter.
 	FTextFilterExpressionEvaluator SearchFilter(ETextFilterExpressionEvaluatorMode::BasicString);
 	SearchFilter.SetFilterText(SearchBox.IsValid() ? SearchBox->GetText() : FText::GetEmpty());
 
-	// Build the source item list for the list view widget.
-	for (const FString& Path : AllPaths)
+	// Build the source for the tree view widget.
+	for (const FString& Path : AllRegisteredPaths)
 	{
 		// Only include items that match the current search filter text.
 		if (SearchFilter.TestTextFilter(FBasicStringFilterExpressionContext(Path)))
 		{
-			ListItems.Add(MakeShared<FString>(Path));
+			FString CurrentNodePath;
+			TArray<TSharedPtr<FPathTreeNodeItem>>* NodeList = &RootNodes;
+
+			TArray<FString> PathSegments;
+			Path.ParseIntoArray(PathSegments, FBlueprintNamespacePathTree::PathSeparator);
+			for (const FString& PathSegment : PathSegments)
+			{
+				CurrentNodePath += PathSegment;
+
+				TSharedPtr<FPathTreeNodeItem>* NodePtr = NodeList->FindByPredicate([&CurrentNodePath](const TSharedPtr<FPathTreeNodeItem>& Value)
+				{
+					return CurrentNodePath.Equals(Value->NodePath, ESearchCase::IgnoreCase);
+				});
+
+				// Add a new node to the current list if an existing one was not found.
+				if (!NodePtr)
+				{
+					NodePtr = &NodeList->Add_GetRef(MakeShared<FPathTreeNodeItem>());
+
+					// Note: We're intentionally using the full path name here rather than just the segment name.
+					(*NodePtr)->NodePath = CurrentNodePath;
+
+					// Treat inclusive nodes that are filtered out by the owner as non-selectable.
+					// 
+					// Example: Consider that "MyProject.MyNamespace" is a registered namespace. When building the tree view for the
+					// drop-down menu, we will first include "MyProject" as the subtree root, and add "MyProject.MyNamespace" as a
+					// child node. If the owner excludes "MyProject" as a namespace (e.g. because it's already imported as an inclusive
+					// namespace), we could also exclude it from the drop-down, but this would mask the hierarchical relationship. So
+					// rather than exclude it from the tree, we make it a non-selectable node.
+					(*NodePtr)->bIsSelectable = !ExcludedTreeViewPaths.Contains(CurrentNodePath);
+				}
+
+				NodeList = &(*NodePtr)->ChildNodes;
+
+				CurrentNodePath += FBlueprintNamespacePathTree::PathSeparator;
+			}
 		}
 	}
 
 	// If no items were added, we signal this by adding a single blank entry.
-	if (ListItems.Num() == 0)
+	if (RootNodes.Num() == 0)
 	{
-		ListItems.Add(MakeShared<FString>());
+		RootNodes.Add(MakeShared<FPathTreeNodeItem>());
 	}
 }
 
@@ -318,4 +431,77 @@ void SBlueprintNamespaceEntry::SelectNamespace(const FString& InNamespace)
 	OnNamespaceSelected.ExecuteIfBound(InNamespace);
 }
 
+void SBlueprintNamespaceEntry::ExpandAllTreeViewItems(const TArray<TSharedPtr<FPathTreeNodeItem>>* NodeListPtr)
+{
+	if (!TreeView.IsValid())
+	{
+		return;
+	}
+
+	if (NodeListPtr == nullptr)
+	{
+		NodeListPtr = &RootNodes;
+	}
+
+	for (const TSharedPtr<FPathTreeNodeItem>& NodeItem : *NodeListPtr)
+	{
+		check(NodeItem.IsValid());
+		TreeView->SetItemExpansion(NodeItem, true);
+
+		ExpandAllTreeViewItems(&NodeItem->ChildNodes);
+	}
+}
+
+const TSharedPtr<SBlueprintNamespaceEntry::FPathTreeNodeItem>* SBlueprintNamespaceEntry::FindTreeViewNode(const FString& NodePath, const TArray<TSharedPtr<FPathTreeNodeItem>>* NodeListPtr) const
+{
+	const TSharedPtr<SBlueprintNamespaceEntry::FPathTreeNodeItem>* Result = nullptr;
+
+	if (NodeListPtr == nullptr)
+	{
+		NodeListPtr = &RootNodes;
+	}
+
+	for (const TSharedPtr<FPathTreeNodeItem>& NodeItem : *NodeListPtr)
+	{
+		check(NodeItem.IsValid());
+		if (NodeItem->NodePath.Equals(NodePath, ESearchCase::IgnoreCase))
+		{
+			Result = &NodeItem;
+		}
+		else
+		{
+			Result = FindTreeViewNode(NodePath, &NodeItem->ChildNodes);
+		}
+
+		if (Result)
+		{
+			break;
+		}
+	}
+
+	return Result;
+}
+
+PRAGMA_DISABLE_DEPRECATION_WARNINGS
+void SBlueprintNamespaceEntry::HandleLegacyOnFilterNamespaceList(TSet<FString>& OutNamespacesToExclude, FOnFilterNamespaceList LegacyDelegate)
+PRAGMA_ENABLE_DEPRECATION_WARNINGS
+{
+	if (LegacyDelegate.IsBound())
+	{
+		TArray<FString> SourceList;
+		FBlueprintNamespaceRegistry::Get().GetAllRegisteredPaths(SourceList);
+
+		TArray<FString> FilteredList = SourceList;
+		LegacyDelegate.Execute(FilteredList);
+
+		for (const FString& SourcePath : SourceList)
+		{
+			if (!FilteredList.Contains(SourcePath))
+			{
+				OutNamespacesToExclude.Add(SourcePath);
+			}
+		}
+	}
+}
+
 #undef LOCTEXT_NAMESPACE
\ No newline at end of file
diff --git a/Engine/Source/Editor/Kismet/Public/BlueprintNamespaceRegistry.h b/Engine/Source/Editor/Kismet/Public/BlueprintNamespaceRegistry.h
--- a/Engine/Source/Editor/Kismet/Public/BlueprintNamespaceRegistry.h
+++ b/Engine/Source/Editor/Kismet/Public/BlueprintNamespaceRegistry.h
@@ -35,6 +35,16 @@ public:
 	 * @return TRUE if the given path identifier is currently registered.
 	 */
 	bool IsRegisteredPath(const FString& InPath) const;
+	
+	/**
+	 * @return TRUE if the given path identifier is inclusive of any registered paths.
+	 * 
+	 * Example: If "MyProject.MyNamespace" is a registered path, then both "MyProject" and "MyProject.MyNamespace" are inclusive paths.
+	 * 
+	 * Also note if a registered path is removed, inclusive paths may still be valid. For instance, if both "MyProject.MyNamespace" and
+	 * "MyProject.MyNamespace_2" are registered paths, and "MyProject.MyNamespace_2" is removed, "MyProject" is still an inclusive path.
+	 */
+	bool IsInclusivePath(const FString& InPath) const;
 
 	/**
 	 * @param InPath	Path identifier string (e.g. "X.Y" or "X.Y.").
diff --git a/Engine/Source/Editor/Kismet/Public/SBlueprintNamespaceEntry.h b/Engine/Source/Editor/Kismet/Public/SBlueprintNamespaceEntry.h
--- a/Engine/Source/Editor/Kismet/Public/SBlueprintNamespaceEntry.h
+++ b/Engine/Source/Editor/Kismet/Public/SBlueprintNamespaceEntry.h
@@ -4,7 +4,7 @@
 
 #include "CoreMinimal.h"
 #include "Widgets/SCompoundWidget.h"
-#include "Widgets/Views/SListView.h"
+#include "Widgets/Views/STreeView.h"
 #include "SlateFwd.h"
 
 class SComboButton;
@@ -15,6 +15,9 @@ class KISMET_API SBlueprintNamespaceEntry : public SCompoundWidget
 {
 public:
 	DECLARE_DELEGATE_OneParam(FOnNamespaceSelected, const FString&);
+	DECLARE_DELEGATE_OneParam(FOnGetNamespacesToExclude, TSet<FString>&);
+
+	UE_DEPRECATED(5.1, "Use FOnGetNamespacesToExclude instead.")
 	DECLARE_DELEGATE_OneParam(FOnFilterNamespaceList, TArray<FString>&);
 
 	SLATE_BEGIN_ARGS(SBlueprintNamespaceEntry)
@@ -36,8 +39,20 @@ public:
 		/** Called when a valid namespace is either entered or selected. */
 		SLATE_EVENT(FOnNamespaceSelected, OnNamespaceSelected)
 
-		/** Called to allow external code to filter out the namespace list. */
-		SLATE_EVENT(FOnFilterNamespaceList, OnFilterNamespaceList)
+		/** Called to allow external code to exclude one or more namespaces from the list. */
+		SLATE_EVENT(FOnGetNamespacesToExclude, OnGetNamespacesToExclude)
+
+		/** Tooltip used for excluded namespaces that are visible in the selection drop-down. */
+		SLATE_ATTRIBUTE(FText, ExcludedNamespaceTooltipText)
+
+		/** [DEPRECATED] Called to allow external code to filter out the namespace list. */
+PRAGMA_DISABLE_DEPRECATION_WARNINGS
+		FOnGetNamespacesToExclude ConvertOnFilterNamespaceListFn(const FOnFilterNamespaceList& LegacyDelegate)
+		{
+			return FOnGetNamespacesToExclude::CreateStatic(SBlueprintNamespaceEntry::HandleLegacyOnFilterNamespaceList, LegacyDelegate);
+		}
+		SLATE_EVENT_DEPRECATED(5.1, "Use OnGetNamespacesToExclude instead.", FOnFilterNamespaceList, OnFilterNamespaceList, OnGetNamespacesToExclude, ConvertOnFilterNamespaceListFn)
+PRAGMA_ENABLE_DEPRECATION_WARNINGS
 
 	SLATE_END_ARGS()
 
@@ -51,29 +66,46 @@ public:
 	void SetCurrentNamespace(const FString& InNamespace);
 
 protected:
+	struct FPathTreeNodeItem
+	{
+		FString NodePath;
+		TArray<TSharedPtr<FPathTreeNodeItem>> ChildNodes;
+		bool bIsSelectable = true;
+	};
+	
 	void OnTextChanged(const FText& InText);
 	void OnTextCommitted(const FText& NewText, ETextCommit::Type InTextCommit);
 	void OnShowingSuggestions(const FString& InputText, TArray<FString>& OutSuggestions);
-	TSharedRef<SWidget> OnGetNamespaceListMenuContent();
-	TSharedRef<ITableRow> OnGenerateRowForNamespaceList(TSharedPtr<FString> Item, const TSharedRef<STableViewBase>& OwnerTable);
-	void OnNamespaceListFilterTextChanged(const FText& InText);
-	void OnNamespaceListSelectionChanged(TSharedPtr<FString> Item, ESelectInfo::Type SelectInfo);
+	TSharedRef<SWidget> OnGetNamespaceTreeMenuContent();
+	TSharedRef<ITableRow> OnGenerateRowForNamespaceTreeItem(TSharedPtr<FPathTreeNodeItem> Item, const TSharedRef<STableViewBase>& OwnerTable);
+	void OnGetChildrenForNamespaceTreeItem(TSharedPtr<FPathTreeNodeItem> Item, TArray<TSharedPtr<FPathTreeNodeItem>>& OutChildren);
+	void OnNamespaceTreeFilterTextChanged(const FText& InText);
+	void OnNamespaceTreeSelectionChanged(TSharedPtr<FPathTreeNodeItem> Item, ESelectInfo::Type SelectInfo);
+	bool OnIsNamespaceTreeItemSelectable(TSharedPtr<FPathTreeNodeItem> Item) const;
 	FText GetCurrentNamespaceText() const;
 
-	void PopulateNamespaceList();
+	void PopulateNamespaceTree();
 	void SelectNamespace(const FString& InNamespace);
+	void ExpandAllTreeViewItems(const TArray<TSharedPtr<FPathTreeNodeItem>>* NodeListPtr = nullptr);
+	const TSharedPtr<FPathTreeNodeItem>* FindTreeViewNode(const FString& NodePath, const TArray<TSharedPtr<FPathTreeNodeItem>>* NodeListPtr = nullptr) const;
+PRAGMA_DISABLE_DEPRECATION_WARNINGS
+	static void HandleLegacyOnFilterNamespaceList(TSet<FString>& OutNamespacesToExclude, FOnFilterNamespaceList LegacyDelegate);
+PRAGMA_ENABLE_DEPRECATION_WARNINGS
 
 private:
 	FString CurrentNamespace;
-	TArray<TSharedPtr<FString>> ListItems;
+	TArray<FString> AllRegisteredPaths;
+	TSet<FString> ExcludedTreeViewPaths;
+	TArray<TSharedPtr<FPathTreeNodeItem>> RootNodes;
 
 	TSharedPtr<SComboButton> ComboButton;
 	TSharedPtr<SSuggestionTextBox> TextBox;
 	TSharedPtr<SSearchBox> SearchBox;
-	TSharedPtr<SListView<TSharedPtr<FString>>> ListView;
+	TSharedPtr<STreeView<TSharedPtr<FPathTreeNodeItem>>> TreeView;
 
 	FOnNamespaceSelected OnNamespaceSelected;
-	FOnFilterNamespaceList OnFilterNamespaceList;
+	FOnGetNamespacesToExclude OnGetNamespacesToExclude;
+	TAttribute<FText> ExcludedNamespaceTooltipText;
 
 	static float NamespaceListBorderPadding;
 	static float NamespaceListMinDesiredWidth;
diff --git a/Engine/Source/Editor/UnrealEd/Private/Settings/BlueprintEditorProjectSettingsCustomization.cpp b/Engine/Source/Editor/UnrealEd/Private/Settings/BlueprintEditorProjectSettingsCustomization.cpp
--- a/Engine/Source/Editor/UnrealEd/Private/Settings/BlueprintEditorProjectSettingsCustomization.cpp
+++ b/Engine/Source/Editor/UnrealEd/Private/Settings/BlueprintEditorProjectSettingsCustomization.cpp
@@ -42,7 +42,8 @@ namespace UE::Editor::BlueprintEditorProjectSettingsCustomization::Private
 			return SNew(SBlueprintNamespaceEntry)
 				.AllowTextEntry(false)
 				.OnNamespaceSelected(this, &FBlueprintGlobalProjectImportsLayout::OnNamespaceSelected)
-				.OnFilterNamespaceList(this, &FBlueprintGlobalProjectImportsLayout::OnFilterNamespaceList)
+				.OnGetNamespacesToExclude(this, &FBlueprintGlobalProjectImportsLayout::OnGetNamespacesToExclude)
+				.ExcludedNamespaceTooltipText(LOCTEXT("CannotSelectNamespace", "This namespace is already included in the list."))
 				.ButtonContent()
 				[
 					SNew(STextBlock)
@@ -95,12 +96,10 @@ namespace UE::Editor::BlueprintEditorProjectSettingsCustomization::Private
 			RegenerateChildContent();
 		}
 
-		void OnFilterNamespaceList(TArray<FString>& InOutNamespaceList)
+		void OnGetNamespacesToExclude(TSet<FString>& OutNamespacesToExclude) const
 		{
-			for (const FString& GlobalNamespace : GetDefault<UBlueprintEditorProjectSettings>()->NamespacesToAlwaysInclude)
-			{
-				InOutNamespaceList.RemoveSwap(GlobalNamespace);
-			}
+			// Exclude namespaces that are already listed.
+			OutNamespacesToExclude.Append(GetDefault<UBlueprintEditorProjectSettings>()->NamespacesToAlwaysInclude);
 		}
 
 	private:
diff --git a/Engine/Source/Editor/UnrealEd/Private/Settings/BlueprintEditorSettingsCustomization.cpp b/Engine/Source/Editor/UnrealEd/Private/Settings/BlueprintEditorSettingsCustomization.cpp
--- a/Engine/Source/Editor/UnrealEd/Private/Settings/BlueprintEditorSettingsCustomization.cpp
+++ b/Engine/Source/Editor/UnrealEd/Private/Settings/BlueprintEditorSettingsCustomization.cpp
@@ -41,7 +41,8 @@ namespace UE::Editor::BlueprintEditorSettingsCustomization::Private
 			return SNew(SBlueprintNamespaceEntry)
 				.AllowTextEntry(false)
 				.OnNamespaceSelected(this, &FBlueprintGlobalEditorImportsLayout::OnNamespaceSelected)
-				.OnFilterNamespaceList(this, &FBlueprintGlobalEditorImportsLayout::OnFilterNamespaceList)
+				.OnGetNamespacesToExclude(this, &FBlueprintGlobalEditorImportsLayout::OnGetNamespacesToExclude)
+				.ExcludedNamespaceTooltipText(LOCTEXT("CannotSelectNamespace", "This namespace is already included in the list."))
 				.ButtonContent()
 				[
 					SNew(STextBlock)
@@ -94,12 +95,10 @@ namespace UE::Editor::BlueprintEditorSettingsCustomization::Private
 			RegenerateChildContent();
 		}
 
-		void OnFilterNamespaceList(TArray<FString>& InOutNamespaceList)
+		void OnGetNamespacesToExclude(TSet<FString>& OutNamespacesToExclude) const
 		{
-			for (const FString& GlobalNamespace : GetDefault<UBlueprintEditorSettings>()->NamespacesToAlwaysInclude)
-			{
-				InOutNamespaceList.RemoveSwap(GlobalNamespace);
-			}
+			// Exclude namespaces that are already listed.
+			OutNamespacesToExclude.Append(GetDefault<UBlueprintEditorSettings>()->NamespacesToAlwaysInclude);
 		}
 
 	private:
