FindInBlueprintManager: Fix missing AssetManagerModule variable from previous change.

diff --git a/Engine/Source/Editor/Kismet/Private/FindInBlueprintManager.cpp b/Engine/Source/Editor/Kismet/Private/FindInBlueprintManager.cpp
--- a/Engine/Source/Editor/Kismet/Private/FindInBlueprintManager.cpp
+++ b/Engine/Source/Editor/Kismet/Private/FindInBlueprintManager.cpp
@@ -1987,10 +1987,14 @@ FFindInBlueprintSearchManager::~FFindInBlueprintSearchManager()
 {
 	if (AssetRegistryModule)
 	{
-		AssetRegistryModule->Get().OnAssetAdded().RemoveAll(this);
-		AssetRegistryModule->Get().OnAssetRemoved().RemoveAll(this);
-		AssetRegistryModule->Get().OnAssetRenamed().RemoveAll(this);
-		AssetRegistryModule->Get().OnFilesLoaded().RemoveAll(this);
+		IAssetRegistry* AssetRegistry = AssetRegistryModule->TryGet();
+		if (AssetRegistry)
+		{
+			AssetRegistry->OnAssetAdded().RemoveAll(this);
+			AssetRegistry->OnAssetRemoved().RemoveAll(this);
+			AssetRegistry->OnAssetRenamed().RemoveAll(this);
+			AssetRegistry->OnFilesLoaded().RemoveAll(this);
+		}
 	}
 	FKismetEditorUtilities::OnBlueprintUnloaded.RemoveAll(this);
 	FCoreUObjectDelegates::GetPreGarbageCollectDelegate().RemoveAll(this);
@@ -2028,7 +2032,8 @@ void FFindInBlueprintSearchManager::Initialize()
 	// Must ensure we do not attempt to load the AssetRegistry Module while saving a package, however, if it is loaded already we can safely obtain it
 	if (!GIsSavingPackage || (GIsSavingPackage && FModuleManager::Get().IsModuleLoaded(TEXT("AssetRegistry"))))
 	{
-		IAssetRegistry* AssetRegistry = FModuleManager::LoadModuleChecked<FAssetRegistryModule>(TEXT("AssetRegistry")).TryGet();
+		AssetRegistryModule = &FModuleManager::LoadModuleChecked<FAssetRegistryModule>(TEXT("AssetRegistry"));
+		IAssetRegistry* AssetRegistry = AssetRegistryModule->TryGet();
 		if (AssetRegistry)
 		{
 			AssetRegistry->OnAssetAdded().AddRaw(this, &FFindInBlueprintSearchManager::OnAssetAdded);
