Filtering SMyBlueprint tabs in anim BPs

Various categories of actions are hidden/prohibited:
- Macros (creation/viewing)
- Event dispatchers/delegates
- Event graphs (as of Lyra we favor the thread safe update function)

diff --git a/Engine/Source/Editor/AnimationBlueprintEditor/Private/AnimationBlueprintEditor.cpp b/Engine/Source/Editor/AnimationBlueprintEditor/Private/AnimationBlueprintEditor.cpp
--- a/Engine/Source/Editor/AnimationBlueprintEditor/Private/AnimationBlueprintEditor.cpp
+++ b/Engine/Source/Editor/AnimationBlueprintEditor/Private/AnimationBlueprintEditor.cpp
@@ -89,6 +89,7 @@
 #include "AnimGraphNode_BlendSpaceGraph.h"
 #include "ScopedTransaction.h"
 #include "Preferences/PersonaOptions.h"
+#include "Settings/AnimBlueprintSettings.h"
 
 #define LOCTEXT_NAMESPACE "AnimationBlueprintEditor"
 
@@ -1706,6 +1707,58 @@ void FAnimationBlueprintEditor::CreateEditorModeManager()
 	EditorModeManager = MakeShareable(FModuleManager::LoadModuleChecked<FPersonaModule>("Persona").CreatePersonaEditorModeManager());
 }
 
+bool FAnimationBlueprintEditor::IsSectionVisible(NodeSectionID::Type InSectionID) const
+{
+	const UAnimBlueprintSettings* AnimBlueprintSettings = GetDefault<UAnimBlueprintSettings>();
+
+	switch (InSectionID)
+	{
+	case NodeSectionID::GRAPH:
+		return AnimBlueprintSettings->bAllowEventGraphs;
+	case NodeSectionID::ANIMGRAPH:
+	case NodeSectionID::ANIMLAYER:
+	case NodeSectionID::FUNCTION:
+	case NodeSectionID::FUNCTION_OVERRIDABLE:
+	case NodeSectionID::INTERFACE:
+		return true;
+	case NodeSectionID::MACRO:
+		return AnimBlueprintSettings->bAllowMacros;
+	case NodeSectionID::VARIABLE:
+		return true;
+	case NodeSectionID::COMPONENT:
+		return false;
+	case NodeSectionID::DELEGATE:
+		return AnimBlueprintSettings->bAllowDelegates;
+	case NodeSectionID::USER_ENUM:
+	case NodeSectionID::LOCAL_VARIABLE:
+	case NodeSectionID::USER_STRUCT:
+	case NodeSectionID::USER_SORTED:
+		return true;
+	default:
+		break;
+	}
+
+	return true;
+}
+
+bool FAnimationBlueprintEditor::AreEventGraphsAllowed() const
+{
+	const UAnimBlueprintSettings* AnimBlueprintSettings = GetDefault<UAnimBlueprintSettings>();
+	return AnimBlueprintSettings->bAllowEventGraphs;
+}
+
+bool FAnimationBlueprintEditor::AreMacrosAllowed() const
+{
+	const UAnimBlueprintSettings* AnimBlueprintSettings = GetDefault<UAnimBlueprintSettings>();
+	return AnimBlueprintSettings->bAllowMacros;
+}
+
+bool FAnimationBlueprintEditor::AreDelegatesAllowed() const
+{
+	const UAnimBlueprintSettings* AnimBlueprintSettings = GetDefault<UAnimBlueprintSettings>();
+	return AnimBlueprintSettings->bAllowDelegates;
+}
+
 void FAnimationBlueprintEditor::OnCreateComment()
 {
 	TSharedPtr<SGraphEditor> GraphEditor = FocusedGraphEdPtr.Pin();
diff --git a/Engine/Source/Editor/AnimationBlueprintEditor/Private/AnimationBlueprintEditor.h b/Engine/Source/Editor/AnimationBlueprintEditor/Private/AnimationBlueprintEditor.h
--- a/Engine/Source/Editor/AnimationBlueprintEditor/Private/AnimationBlueprintEditor.h
+++ b/Engine/Source/Editor/AnimationBlueprintEditor/Private/AnimationBlueprintEditor.h
@@ -196,6 +196,10 @@ protected:
 	virtual FText GetGraphDecorationString(UEdGraph* InGraph) const override;
 	virtual void OnBlueprintChangedImpl(UBlueprint* InBlueprint, bool bIsJustBeingCompiled = false) override;
 	virtual void CreateEditorModeManager() override;
+	virtual bool IsSectionVisible(NodeSectionID::Type InSectionID) const override;
+	virtual bool AreEventGraphsAllowed() const override;
+	virtual bool AreMacrosAllowed() const override;
+	virtual bool AreDelegatesAllowed() const override;
 	virtual void OnCreateComment() override;
 	//~ End FBlueprintEditor Interface
 
diff --git a/Engine/Source/Editor/AnimationBlueprintEditor/Private/AnimationBlueprintEditorModule.cpp b/Engine/Source/Editor/AnimationBlueprintEditor/Private/AnimationBlueprintEditorModule.cpp
--- a/Engine/Source/Editor/AnimationBlueprintEditor/Private/AnimationBlueprintEditorModule.cpp
+++ b/Engine/Source/Editor/AnimationBlueprintEditor/Private/AnimationBlueprintEditorModule.cpp
@@ -70,7 +70,7 @@ TSharedRef<IAnimationBlueprintEditor> FAnimationBlueprintEditorModule::CreateAni
 
 void FAnimationBlueprintEditorModule::OnNewBlueprintCreated(UBlueprint* InBlueprint)
 {
-	if (ensure(InBlueprint->UbergraphPages.Num() > 0))
+	if (InBlueprint->UbergraphPages.Num() > 0)
 	{
 		UEdGraph* EventGraph = InBlueprint->UbergraphPages[0];
 
diff --git a/Engine/Source/Editor/Kismet/Private/BlueprintEditor.cpp b/Engine/Source/Editor/Kismet/Private/BlueprintEditor.cpp
--- a/Engine/Source/Editor/Kismet/Private/BlueprintEditor.cpp
+++ b/Engine/Source/Editor/Kismet/Private/BlueprintEditor.cpp
@@ -10422,6 +10422,21 @@ void FBlueprintEditor::OnBlueprintEditorPreferencesChanged(UObject*, struct FPro
 	ModifyDuringPIEStatus = ESafeToModifyDuringPIEStatus::Unknown;
 }
 
+bool FBlueprintEditor::AreEventGraphsAllowed() const
+{
+	return true;
+}
+
+bool FBlueprintEditor::AreMacrosAllowed() const
+{
+	return true;
+}
+
+bool FBlueprintEditor::AreDelegatesAllowed() const
+{
+	return true;
+}
+
 /////////////////////////////////////////////////////
 /////////////////////////////////////////////////////
 /////////////////////////////////////////////////////
diff --git a/Engine/Source/Editor/Kismet/Private/SMyBlueprint.cpp b/Engine/Source/Editor/Kismet/Private/SMyBlueprint.cpp
--- a/Engine/Source/Editor/Kismet/Private/SMyBlueprint.cpp
+++ b/Engine/Source/Editor/Kismet/Private/SMyBlueprint.cpp
@@ -1341,6 +1341,8 @@ void SMyBlueprint::CollectAllActions(FGraphActionListBuilderBase& OutAllActions)
 	UBlueprint* BlueprintObj = GetBlueprintObj();
 	check(BlueprintObj);
 
+	TSharedPtr<FBlueprintEditor> BlueprintEditor = BlueprintEditorPtr.Pin();
+
 	EFieldIteratorFlags::SuperClassFlags FieldIteratorSuperFlag = EFieldIteratorFlags::IncludeSuper;
 	if ( ShowUserVarsOnly() )
 	{
@@ -1411,7 +1413,7 @@ void SMyBlueprint::CollectAllActions(FGraphActionListBuilderBase& OutAllActions)
 			NewVarAction->SetVariableInfo(PropertyName, BlueprintObj->SkeletonGeneratedClass, CastField<FBoolProperty>(TestProperty) != nullptr);
 			SortList.AddAction( UserCategoryName, NewVarAction );
 		}
-		else if (bShouldShowAsDelegate)
+		else if (bShouldShowAsDelegate && BlueprintEditor.IsValid() && BlueprintEditor->AreDelegatesAllowed())
 		{
 			TSharedPtr<FEdGraphSchemaAction_K2Delegate> NewDelegateAction;
 			// Delegate is visible in MyBlueprint when not-native or its category name is not empty.
@@ -1471,30 +1473,33 @@ void SMyBlueprint::CollectAllActions(FGraphActionListBuilderBase& OutAllActions)
 		ImplementedFunctionCache.Add(Graph->GetFName());
 	}
 
-	// Grab macros implemented by the blueprint
-	for (int32 i = 0; i < BlueprintObj->MacroGraphs.Num(); i++)
+	if(BlueprintEditor.IsValid() && BlueprintEditor->AreMacrosAllowed())
 	{
-		UEdGraph* Graph = BlueprintObj->MacroGraphs[i];
-		check(Graph);
+		// Grab macros implemented by the blueprint
+		for (int32 i = 0; i < BlueprintObj->MacroGraphs.Num(); i++)
+		{
+			UEdGraph* Graph = BlueprintObj->MacroGraphs[i];
+			check(Graph);
 		
-		const FName MacroName = Graph->GetFName();
+			const FName MacroName = Graph->GetFName();
 
-		FGraphDisplayInfo DisplayInfo;
-		Graph->GetSchema()->GetGraphDisplayInformation(*Graph, DisplayInfo);
+			FGraphDisplayInfo DisplayInfo;
+			Graph->GetSchema()->GetGraphDisplayInformation(*Graph, DisplayInfo);
 
-		FText MacroCategory = GetGraphCategory(Graph);
+			FText MacroCategory = GetGraphCategory(Graph);
 
-		TSharedPtr<FEdGraphSchemaAction_K2Graph> NewMacroAction = MakeShareable(new FEdGraphSchemaAction_K2Graph(EEdGraphSchemaAction_K2Graph::Macro, MacroCategory, DisplayInfo.PlainName, DisplayInfo.Tooltip, 1, NodeSectionID::MACRO));
-		NewMacroAction->FuncName = MacroName;
-		NewMacroAction->EdGraph = Graph;
+			TSharedPtr<FEdGraphSchemaAction_K2Graph> NewMacroAction = MakeShareable(new FEdGraphSchemaAction_K2Graph(EEdGraphSchemaAction_K2Graph::Macro, MacroCategory, DisplayInfo.PlainName, DisplayInfo.Tooltip, 1, NodeSectionID::MACRO));
+			NewMacroAction->FuncName = MacroName;
+			NewMacroAction->EdGraph = Graph;
 
-		const FString UserCategoryName = FEditorCategoryUtils::GetCategoryDisplayString(MacroCategory.ToString());
-		SortList.AddAction(UserCategoryName, NewMacroAction);
+			const FString UserCategoryName = FEditorCategoryUtils::GetCategoryDisplayString(MacroCategory.ToString());
+			SortList.AddAction(UserCategoryName, NewMacroAction);
 
-		GetChildGraphs(Graph, NewMacroAction->GetSectionID(), SortList, MacroCategory);
-		GetChildEvents(Graph, NewMacroAction->GetSectionID(), SortList, MacroCategory);
+			GetChildGraphs(Graph, NewMacroAction->GetSectionID(), SortList, MacroCategory);
+			GetChildEvents(Graph, NewMacroAction->GetSectionID(), SortList, MacroCategory);
 
-		ImplementedFunctionCache.Add(MacroName);
+			ImplementedFunctionCache.Add(MacroName);
+		}
 	}
 
 	OverridableFunctionActions.Reset();
@@ -1650,22 +1655,26 @@ void SMyBlueprint::CollectAllActions(FGraphActionListBuilderBase& OutAllActions)
 		}
 	}
 
-	// Grab ubergraph pages
-	for (int32 i = 0; i < BlueprintObj->UbergraphPages.Num(); i++)
+	
+	if(BlueprintEditor.IsValid() && BlueprintEditor->AreEventGraphsAllowed())
 	{
-		UEdGraph* Graph = BlueprintObj->UbergraphPages[i];
-		check(Graph);
+		// Grab ubergraph pages
+		for (int32 i = 0; i < BlueprintObj->UbergraphPages.Num(); i++)
+		{
+			UEdGraph* Graph = BlueprintObj->UbergraphPages[i];
+			check(Graph);
 		
-		FGraphDisplayInfo DisplayInfo;
-		Graph->GetSchema()->GetGraphDisplayInformation(*Graph, DisplayInfo);
+			FGraphDisplayInfo DisplayInfo;
+			Graph->GetSchema()->GetGraphDisplayInformation(*Graph, DisplayInfo);
 
-		TSharedPtr<FEdGraphSchemaAction_K2Graph> NeUbergraphAction = MakeShareable(new FEdGraphSchemaAction_K2Graph(EEdGraphSchemaAction_K2Graph::Graph, FText::GetEmpty(), DisplayInfo.PlainName, DisplayInfo.Tooltip, 2, NodeSectionID::GRAPH));
-		NeUbergraphAction->FuncName = Graph->GetFName();
-		NeUbergraphAction->EdGraph = Graph;
-		OutAllActions.AddAction(NeUbergraphAction);
+			TSharedPtr<FEdGraphSchemaAction_K2Graph> NeUbergraphAction = MakeShareable(new FEdGraphSchemaAction_K2Graph(EEdGraphSchemaAction_K2Graph::Graph, FText::GetEmpty(), DisplayInfo.PlainName, DisplayInfo.Tooltip, 2, NodeSectionID::GRAPH));
+			NeUbergraphAction->FuncName = Graph->GetFName();
+			NeUbergraphAction->EdGraph = Graph;
+			OutAllActions.AddAction(NeUbergraphAction);
 
-		GetChildGraphs(Graph, NeUbergraphAction->GetSectionID(), SortList);
-		GetChildEvents(Graph, NeUbergraphAction->GetSectionID(), SortList);
+			GetChildGraphs(Graph, NeUbergraphAction->GetSectionID(), SortList);
+			GetChildEvents(Graph, NeUbergraphAction->GetSectionID(), SortList);
+		}
 	}
 
 	// Grab intermediate pages
diff --git a/Engine/Source/Editor/Kismet/Public/BlueprintEditor.h b/Engine/Source/Editor/Kismet/Public/BlueprintEditor.h
--- a/Engine/Source/Editor/Kismet/Public/BlueprintEditor.h
+++ b/Engine/Source/Editor/Kismet/Public/BlueprintEditor.h
@@ -1316,6 +1316,15 @@ public://@TODO
 	/** Update all nodes' unrelated states when the graph has changed */
 	void UpdateNodesUnrelatedStatesAfterGraphChange();
 
+	/** Whether event graphs are allowed to be displayed/created in a blueprint */
+	virtual bool AreEventGraphsAllowed() const;
+
+	/** Whether macros are allowed  to be displayed/created in a blueprint */
+	virtual bool AreMacrosAllowed() const;
+
+	/** Whether delegates are allowed  to be displayed/created in a blueprint */
+	virtual bool AreDelegatesAllowed() const;
+
 protected:
 
 	/** Should intermediate build products be saved when recompiling? */
diff --git a/Engine/Source/Editor/UnrealEd/Classes/Settings/AnimBlueprintSettings.h b/Engine/Source/Editor/UnrealEd/Classes/Settings/AnimBlueprintSettings.h
new file mode 100644
--- /dev/null
+++ b/Engine/Source/Editor/UnrealEd/Classes/Settings/AnimBlueprintSettings.h
@@ -0,0 +1,29 @@
+// Copyright Epic Games, Inc. All Rights Reserved.
+
+#pragma once
+
+
+#include "CoreMinimal.h"
+#include "AnimBlueprintSettings.generated.h"
+
+/**
+ * Implements Editor settings for animation blueprints
+ */
+UCLASS(config = EditorPerProjectUserSettings)
+class UNREALED_API UAnimBlueprintSettings : public UObject
+{
+	GENERATED_BODY()
+
+public:
+	/** Whether to allow event graphs to be created/displayed in animation blueprints */
+	UPROPERTY()
+	bool bAllowEventGraphs = true;
+
+	/** Whether to allow macros to be created/displayed in animation blueprints */
+	UPROPERTY()
+	bool bAllowMacros = true;
+
+	/** Whether to allow delegates to be created/displayed in animation blueprints */
+	UPROPERTY()
+	bool bAllowDelegates = true;
+};
\ No newline at end of file
diff --git a/Engine/Source/Editor/UnrealEd/Private/Kismet2/Kismet2.cpp b/Engine/Source/Editor/UnrealEd/Private/Kismet2/Kismet2.cpp
--- a/Engine/Source/Editor/UnrealEd/Private/Kismet2/Kismet2.cpp
+++ b/Engine/Source/Editor/UnrealEd/Private/Kismet2/Kismet2.cpp
@@ -458,7 +458,7 @@ UBlueprint* FKismetEditorUtilities::CreateBlueprint(UClass* ParentClass, UObject
 	}
 
 	// Create default event graph(s)
-	if (FBlueprintEditorUtils::DoesSupportEventGraphs(NewBP))
+	if (FBlueprintEditorUtils::DoesSupportEventGraphs(NewBP) && NewBP->SupportsEventGraphs())
 	{
 		check(NewBP->UbergraphPages.Num() == 0);
 		CreateDefaultEventGraphs(NewBP);
diff --git a/Engine/Source/Runtime/Engine/Classes/Animation/AnimBlueprint.h b/Engine/Source/Runtime/Engine/Classes/Animation/AnimBlueprint.h
--- a/Engine/Source/Runtime/Engine/Classes/Animation/AnimBlueprint.h
+++ b/Engine/Source/Runtime/Engine/Classes/Animation/AnimBlueprint.h
@@ -176,7 +176,11 @@ class ENGINE_API UAnimBlueprint : public UBlueprint, public IInterface_PreviewMe
 	virtual void PostLoad() override;
 	virtual bool FindDiffs(const UBlueprint* OtherBlueprint, FDiffResults& Results) const override;
 	virtual void SetObjectBeingDebugged(UObject* NewObject) override;
-	
+	virtual bool SupportsAnimLayers() const override;
+	virtual bool SupportsEventGraphs() const override;
+	virtual bool SupportsDelegates() const override;
+	virtual bool SupportsMacros() const override;
+
 protected:
 	// Broadcast when an override is changed, allowing derived blueprints to be updated
 	FOnOverrideChangedMulticaster OnOverrideChanged;
diff --git a/Engine/Source/Runtime/Engine/Classes/Engine/Blueprint.h b/Engine/Source/Runtime/Engine/Classes/Engine/Blueprint.h
--- a/Engine/Source/Runtime/Engine/Classes/Engine/Blueprint.h
+++ b/Engine/Source/Runtime/Engine/Classes/Engine/Blueprint.h
@@ -1230,7 +1230,7 @@ public:
 	/**
 	 * Returns true if this blueprint supports animation layers
 	 */
-	virtual bool SupportsAnimLayers() const { return true; }
+	virtual bool SupportsAnimLayers() const { return false; }
 
 	/**
 	 * Copies a given graph into a text buffer. Returns true if successful.
diff --git a/Engine/Source/Runtime/Engine/Private/AnimBlueprint.cpp b/Engine/Source/Runtime/Engine/Private/AnimBlueprint.cpp
--- a/Engine/Source/Runtime/Engine/Private/AnimBlueprint.cpp
+++ b/Engine/Source/Runtime/Engine/Private/AnimBlueprint.cpp
@@ -4,7 +4,7 @@
 #include "UObject/FrameworkObjectVersion.h"
 #include "Animation/AnimBlueprintGeneratedClass.h"
 #if WITH_EDITOR
-#include "Settings/EditorExperimentalSettings.h"
+#include "Settings/AnimBlueprintSettings.h"
 #include "Modules/ModuleManager.h"
 #endif
 #if WITH_EDITORONLY_DATA
@@ -217,6 +217,26 @@ void UAnimBlueprint::SetObjectBeingDebugged(UObject* NewObject)
 	Super::SetObjectBeingDebugged(NewObject);
 }
 
+bool UAnimBlueprint::SupportsEventGraphs() const
+{
+	return GetDefault<UAnimBlueprintSettings>()->bAllowEventGraphs;
+}
+
+bool UAnimBlueprint::SupportsAnimLayers() const
+{
+	return true;
+}
+
+bool UAnimBlueprint::SupportsDelegates() const
+{
+	return GetDefault<UAnimBlueprintSettings>()->bAllowDelegates;
+}
+
+bool UAnimBlueprint::SupportsMacros() const
+{
+	return GetDefault<UAnimBlueprintSettings>()->bAllowMacros;
+}
+
 #endif
 
 USkeletalMesh* UAnimBlueprint::GetPreviewMesh(bool bFindIfNotSet/*=false*/)
