Export FImaginaryFiBData::DumpParsedObject to use from another project

#rnx

diff --git a/Engine/Source/Editor/Kismet/Public/ImaginaryBlueprintData.h b/Engine/Source/Editor/Kismet/Public/ImaginaryBlueprintData.h
--- a/Engine/Source/Editor/Kismet/Public/ImaginaryBlueprintData.h
+++ b/Engine/Source/Editor/Kismet/Public/ImaginaryBlueprintData.h
@@ -176,7 +176,7 @@ public:
 	}
 
 	/** Dumps the parsed object (including all children) to the given archive */
-	void DumpParsedObject(FArchive& Ar, int32 InTreeLevel = 0) const;
+	KISMET_API void DumpParsedObject(FArchive& Ar, int32 InTreeLevel = 0) const;
 
 	/** Builds a SearchTree ready to be displayed in the Find-in-Blueprints window */
 	static FSearchResult CreateSearchTree(FSearchResult InParentSearchResult, FImaginaryFiBDataWeakPtr InCurrentPointer, TArray< const FImaginaryFiBData* >& InValidSearchResults, TMultiMap< const FImaginaryFiBData*, FComponentUniqueDisplay >& InMatchingSearchComponents);
