[Bug Fix] "Short type name provided for TryFindType" warning occurs when rearranging a Child Actor component in the Components Hierarchy

diff --git a/Engine/Source/Editor/Kismet/Private/BlueprintDetailsCustomization.cpp b/Engine/Source/Editor/Kismet/Private/BlueprintDetailsCustomization.cpp
--- a/Engine/Source/Editor/Kismet/Private/BlueprintDetailsCustomization.cpp
+++ b/Engine/Source/Editor/Kismet/Private/BlueprintDetailsCustomization.cpp
@@ -6990,11 +6990,9 @@ void FChildActorComponentDetails::CustomizeDetails(IDetailLayoutBuilder& DetailB
 				static FText RestrictReason = LOCTEXT("NoSelfChildActors", "Cannot append a child-actor of this blueprint type (could cause infinite recursion).");
 				TSharedPtr<FPropertyRestriction> ClassRestriction = MakeShareable(new FPropertyRestriction(RestrictReason));
 
-				ClassRestriction->AddDisabledValue(Blueprint->GetName());
 				ClassRestriction->AddDisabledValue(Blueprint->GetPathName());
 				if (Blueprint->GeneratedClass)
 				{
-					ClassRestriction->AddDisabledValue(Blueprint->GeneratedClass->GetName());
 					ClassRestriction->AddDisabledValue(Blueprint->GeneratedClass->GetPathName());
 				}
 
