Fixing blueprint diff tool potentially calling IsChildOf on nullptr for newly added classes.
Fixing that blueprint diff tool wouldn't have components panel for newly added blueprints.

diff --git a/Engine/Source/Editor/Kismet/Private/SBlueprintDiff.cpp b/Engine/Source/Editor/Kismet/Private/SBlueprintDiff.cpp
--- a/Engine/Source/Editor/Kismet/Private/SBlueprintDiff.cpp
+++ b/Engine/Source/Editor/Kismet/Private/SBlueprintDiff.cpp
@@ -1572,8 +1572,11 @@ void SBlueprintDiff::GenerateDifferencesList()
 	}
 
 	bool bHasComponents = false;
-	UClass* BlueprintClass = PanelOld.Blueprint->GeneratedClass;
-	if (BlueprintClass->IsChildOf<AActor>())
+	UClass* BlueprintClassOld = PanelOld.Blueprint->GeneratedClass;
+	UClass* BlueprintClassNew = PanelNew.Blueprint->GeneratedClass;
+	const bool bIsOldClassActor = BlueprintClassOld && BlueprintClassOld->IsChildOf<AActor>();
+	const bool bIsNewClassActor = BlueprintClassNew && BlueprintClassNew->IsChildOf<AActor>();
+	if (bIsOldClassActor || bIsNewClassActor)
 	{
 		bHasComponents = true;
 	}
