Fix for a regression that can trigger an ensure when editing inherited variable default values in the Blueprint editor.

Change summary:
- The struct type check that triggered the ensure() is no longer required. This was left over from some validation code I had started with as part of testing Blueprint namespace auto-import in response to value changes.

diff --git a/Engine/Source/Editor/Kismet/Private/BlueprintDetailsCustomization.cpp b/Engine/Source/Editor/Kismet/Private/BlueprintDetailsCustomization.cpp
--- a/Engine/Source/Editor/Kismet/Private/BlueprintDetailsCustomization.cpp
+++ b/Engine/Source/Editor/Kismet/Private/BlueprintDetailsCustomization.cpp
@@ -2744,7 +2744,7 @@ void FBlueprintVarActionDetails::OnFinishedChangingVariable(const FPropertyChang
 		return;
 	}
 
-	ImportNamespacesForPropertyValue(InModifiedObjectInstance->GetClass(), InPropertyChangedEvent.MemberProperty, InModifiedObjectInstance);
+	ImportNamespacesForPropertyValue(InPropertyChangedEvent.MemberProperty, InModifiedObjectInstance);
 }
 
 void FBlueprintVarActionDetails::OnFinishedChangingLocalVariable(const FPropertyChangedEvent& InPropertyChangedEvent, TSharedPtr<FStructOnScope> InStructData, TWeakObjectPtr<UK2Node_EditablePinBase> InEntryNode)
@@ -2791,11 +2791,11 @@ void FBlueprintVarActionDetails::OnFinishedChangingLocalVariable(const FProperty
 			}
 		}
 
-		ImportNamespacesForPropertyValue(InStructData->GetStruct(), DirectProperty, InStructData->GetStructMemory());
+		ImportNamespacesForPropertyValue(DirectProperty, InStructData->GetStructMemory());
 	}
 }
 
-void FBlueprintVarActionDetails::ImportNamespacesForPropertyValue(const UStruct* InStruct, const FProperty* InProperty, const void* InContainer)
+void FBlueprintVarActionDetails::ImportNamespacesForPropertyValue(const FProperty* InProperty, const void* InContainer)
 {
 	// Auto-import any namespace(s) associated with the property's value into the current editor context.
 	TSharedPtr<SMyBlueprint> MyBlueprintPtr = MyBlueprint.Pin();
@@ -2805,7 +2805,7 @@ void FBlueprintVarActionDetails::ImportNamespacesForPropertyValue(const UStruct*
 		if (BlueprintEditor.IsValid())
 		{
 			FBlueprintEditor::FImportNamespaceExParameters Params;
-			FBlueprintNamespaceUtilities::GetPropertyValueNamespaces(InStruct, InProperty, InContainer, Params.NamespacesToImport);
+			FBlueprintNamespaceUtilities::GetPropertyValueNamespaces(InProperty, InContainer, Params.NamespacesToImport);
 			BlueprintEditor->ImportNamespaceEx(Params);
 		}
 	}
diff --git a/Engine/Source/Editor/Kismet/Private/BlueprintDetailsCustomization.h b/Engine/Source/Editor/Kismet/Private/BlueprintDetailsCustomization.h
--- a/Engine/Source/Editor/Kismet/Private/BlueprintDetailsCustomization.h
+++ b/Engine/Source/Editor/Kismet/Private/BlueprintDetailsCustomization.h
@@ -164,11 +164,10 @@ private:
 	/**
 	 * Auto-import any namespaces associated with a variable's value into the current editor context
 	 *
-	 * @param InStruct					A reference to the container's type
-	 * @param InProperty				A reference to the variable property; owner is expected to match the container's type
+	 * @param InProperty				A reference to the variable property
 	 * @param InContainer				A pointer to the data container (e.g. struct or object) where the property's value is stored
 	 */
-	void ImportNamespacesForPropertyValue(const UStruct* InStruct, const FProperty* InProperty, const void* InContainer);
+	void ImportNamespacesForPropertyValue(const FProperty* InProperty, const void* InContainer);
 
 	/** Callback to decide if the category drop down menu should be enabled */
 	bool GetVariableCategoryChangeEnabled() const;
diff --git a/Engine/Source/Editor/Kismet/Private/BlueprintNamespaceUtilities.cpp b/Engine/Source/Editor/Kismet/Private/BlueprintNamespaceUtilities.cpp
--- a/Engine/Source/Editor/Kismet/Private/BlueprintNamespaceUtilities.cpp
+++ b/Engine/Source/Editor/Kismet/Private/BlueprintNamespaceUtilities.cpp
@@ -152,20 +152,9 @@ FString FBlueprintNamespaceUtilities::GetObjectNamespace(const FSoftObjectPath&
 	return GetAssetNamespace(AssetData);
 }
 
-void FBlueprintNamespaceUtilities::GetPropertyValueNamespaces(const UStruct* InStruct, const FProperty* InProperty, const void* InContainer, TSet<FString>& OutNamespaces)
+void FBlueprintNamespaceUtilities::GetPropertyValueNamespaces(const FProperty* InProperty, const void* InContainer, TSet<FString>& OutNamespaces)
 {
-	if (!InStruct || !InProperty || !InContainer)
-	{
-		return;
-	}
-
-	const UStruct* PropertyOwner = InProperty->GetOwnerStruct();
-	if (!PropertyOwner)
-	{
-		return;
-	}
-
-	if (!ensureMsgf(InStruct == PropertyOwner, TEXT("Property %s is a member of struct %s which does not match the given struct %s"), *InProperty->GetName(), *PropertyOwner->GetName(), *InStruct->GetName()))
+	if (!InProperty || !InContainer)
 	{
 		return;
 	}
@@ -178,7 +167,7 @@ void FBlueprintNamespaceUtilities::GetPropertyValueNamespaces(const UStruct* InS
 		{
 			for (TFieldIterator<FProperty> It(StructProperty->Struct); It; ++It)
 			{
-				GetPropertyValueNamespaces(StructProperty->Struct, *It, ValuePtr, OutNamespaces);
+				GetPropertyValueNamespaces(*It, ValuePtr, OutNamespaces);
 			}
 		}
 		else if (const FArrayProperty* ArrayProperty = CastField<FArrayProperty>(InProperty))
@@ -186,7 +175,7 @@ void FBlueprintNamespaceUtilities::GetPropertyValueNamespaces(const UStruct* InS
 			FScriptArrayHelper ArrayHelper(ArrayProperty, ValuePtr);
 			for (int32 ValueIdx = 0; ValueIdx < ArrayHelper.Num(); ++ValueIdx)
 			{
-				GetPropertyValueNamespaces(InStruct, ArrayProperty->Inner, ArrayHelper.GetRawPtr(ValueIdx), OutNamespaces);
+				GetPropertyValueNamespaces(ArrayProperty->Inner, ArrayHelper.GetRawPtr(ValueIdx), OutNamespaces);
 			}
 		}
 		else if (const FSetProperty* SetProperty = CastField<FSetProperty>(InProperty))
@@ -194,7 +183,7 @@ void FBlueprintNamespaceUtilities::GetPropertyValueNamespaces(const UStruct* InS
 			FScriptSetHelper SetHelper(SetProperty, ValuePtr);
 			for (FScriptSetHelper::FIterator SetIt = SetHelper.CreateIterator(); SetIt; ++SetIt)
 			{
-				GetPropertyValueNamespaces(InStruct, SetProperty->ElementProp, SetHelper.GetElementPtr(*SetIt), OutNamespaces);
+				GetPropertyValueNamespaces(SetProperty->ElementProp, SetHelper.GetElementPtr(*SetIt), OutNamespaces);
 			}
 		}
 		else if (const FMapProperty* MapProperty = CastField<FMapProperty>(InProperty))
@@ -203,8 +192,8 @@ void FBlueprintNamespaceUtilities::GetPropertyValueNamespaces(const UStruct* InS
 			for (FScriptMapHelper::FIterator MapIt = MapHelper.CreateIterator(); MapIt; ++MapIt)
 			{
 				const uint8* MapValuePtr = MapHelper.GetPairPtr(*MapIt);
-				GetPropertyValueNamespaces(InStruct, MapProperty->KeyProp, MapValuePtr, OutNamespaces);
-				GetPropertyValueNamespaces(InStruct, MapProperty->ValueProp, MapValuePtr, OutNamespaces);
+				GetPropertyValueNamespaces(MapProperty->KeyProp, MapValuePtr, OutNamespaces);
+				GetPropertyValueNamespaces(MapProperty->ValueProp, MapValuePtr, OutNamespaces);
 			}
 		}
 		else if (const FSoftObjectProperty* SoftObjectProperty = CastField<FSoftObjectProperty>(InProperty))
diff --git a/Engine/Source/Editor/Kismet/Private/SKismetInspector.cpp b/Engine/Source/Editor/Kismet/Private/SKismetInspector.cpp
--- a/Engine/Source/Editor/Kismet/Private/SKismetInspector.cpp
+++ b/Engine/Source/Editor/Kismet/Private/SKismetInspector.cpp
@@ -1037,12 +1037,12 @@ void SKismetInspector::ImportNamespacesForPropertyValue(const FProperty* InPrope
 				// For example, a details customization may select an unrelated object
 				// and then customize each row with an external object reference. In
 				// those cases, the customization would need to handle this explicitly.
-				if (SelectedType != InProperty->GetOwnerStruct())
+				if (!InProperty->IsIn(SelectedType))
 				{
 					continue;
 				}
 
-				FBlueprintNamespaceUtilities::GetPropertyValueNamespaces(SelectedType, InProperty, SelectedObject, AssociatedNamespaces);
+				FBlueprintNamespaceUtilities::GetPropertyValueNamespaces(InProperty, SelectedObject, AssociatedNamespaces);
 			}
 		}
 
diff --git a/Engine/Source/Editor/Kismet/Public/BlueprintNamespaceUtilities.h b/Engine/Source/Editor/Kismet/Public/BlueprintNamespaceUtilities.h
--- a/Engine/Source/Editor/Kismet/Public/BlueprintNamespaceUtilities.h
+++ b/Engine/Source/Editor/Kismet/Public/BlueprintNamespaceUtilities.h
@@ -48,12 +48,17 @@ public:
 	/**
 	 * Analyzes a property value to determine explicitly-assigned namespace identifiers from any object references, or otherwise returns the default namespace for each occurrence (default).
 	 * 
-	 * @param InStruct		A reference to the container's struct/object type.
 	 * @param InProperty	The property for which we will analyze the value.
 	 * @param InContainer	The source address of the struct/object containing the property's value.
 	 * @param OutNamespaces	Zero or more unique namespace identifier(s) referenced by the property value. An entry with an empty string equates to the default global namespace.
 	 */
-	static void GetPropertyValueNamespaces(const UStruct* InStruct, const FProperty* InProperty, const void* InContainer, TSet<FString>& OutNamespaces);
+	static void GetPropertyValueNamespaces(const FProperty* InProperty, const void* InContainer, TSet<FString>& OutNamespaces);
+
+	UE_DEPRECATED(5.1, "Please use the updated version that removes the InStruct parameter (no longer needed).")
+	static void GetPropertyValueNamespaces(const UStruct* InStruct, const FProperty* InProperty, const void* InContainer, TSet<FString>& OutNamespaces)
+	{
+		GetPropertyValueNamespaces(InProperty, InContainer, OutNamespaces);
+	}
 
 	/**
 	 * Gathers the set of global namespaces that are implicitly imported by all Blueprint assets.
