Don't attempt to reinstance artifacts from a previous reinstancing pass on a Blueprint asset package reload.

diff --git a/Engine/Source/Editor/Kismet/Private/BlueprintCompilationManager.cpp b/Engine/Source/Editor/Kismet/Private/BlueprintCompilationManager.cpp
--- a/Engine/Source/Editor/Kismet/Private/BlueprintCompilationManager.cpp
+++ b/Engine/Source/Editor/Kismet/Private/BlueprintCompilationManager.cpp
@@ -1884,6 +1884,12 @@ void FBlueprintCompilationManagerImpl::ReparentHierarchies(const TMap<UClass*, U
 					continue;
 				}
 
+				if (DerivedClass->HasAnyClassFlags(CLASS_NewerVersionExists))
+				{
+					// Don't reinstance artifacts from classes that have been previously been reinstanced/trashed but aren't cleaned up yet.
+					continue;
+				}
+
 				if (OldToNewClasses.Find(DerivedClass) == nullptr)
 				{
 					Classes.Add(DerivedClass);
