Fix for an editor crash that would occur after replacing one or more references to a conflicting variable name after reparenting a Blueprint class.

Change summary:
- Cache the SelectedProperty scope prior to reference replacement. This avoids attempting to dereference an invalid ptr to access it on completion (the source of the crash).
- Reset the target property whenever we clear the source property on the reference replacement widget for consistency. This fully resets the widget for a new operation.
- Fixed 'bShowReplacementsWhenFinished' logic to no longer attempt to request a local search when 'bFindWithinBlueprint' is toggled on by the user. We cannot perform a local search in a global FiB "nomad" tab, which is currently what's being used to display the results.
- Added an ensureMsgf() to SFindInBlueprints::MakeSearchQuery() in order to guard against future attempts to dereference an invalid Blueprint editor context (this will now result in performing a global search instead as a fallback).
- Modified FBlueprintCompilationManagerImpl::FastGenerateSkeletonClass() to run a variable name validation pass prior to skeleton class compilation if we see that the parent class has changed. This avoids an ICE that previously was being raised by an MBASM call on the BP change event that occurs prior to the post-reparent full compile.

diff --git a/Engine/Source/Editor/Kismet/Private/BlueprintCompilationManager.cpp b/Engine/Source/Editor/Kismet/Private/BlueprintCompilationManager.cpp
--- a/Engine/Source/Editor/Kismet/Private/BlueprintCompilationManager.cpp
+++ b/Engine/Source/Editor/Kismet/Private/BlueprintCompilationManager.cpp
@@ -2522,6 +2522,13 @@ UClass* FBlueprintCompilationManagerImpl::FastGenerateSkeletonClass(UBlueprint*
 	else
 	{
 		Ret = CastChecked<UBlueprintGeneratedClass>(*(BP->SkeletonGeneratedClass));
+
+		// If we're changing the parent class, first validate variable names against the inherited set to avoid collisions when we create properties.
+		if (Ret->GetSuperClass() != ParentClass)
+		{
+			CompilerContext.ValidateVariableNames();
+		}
+
 		CompilerContext.CleanAndSanitizeClass(Ret, Ret->ClassDefaultObject);
 	}
 	
diff --git a/Engine/Source/Editor/Kismet/Private/FindInBlueprints.cpp b/Engine/Source/Editor/Kismet/Private/FindInBlueprints.cpp
--- a/Engine/Source/Editor/Kismet/Private/FindInBlueprints.cpp
+++ b/Engine/Source/Editor/Kismet/Private/FindInBlueprints.cpp
@@ -1050,7 +1050,8 @@ void SFindInBlueprints::MakeSearchQuery(FString InSearchString, bool bInIsFindWi
 		TreeView->RequestTreeRefresh();
 		HighlightText = FText::FromString( InSearchString );
 
-		if (bInIsFindWithinBlueprint)
+		if (bInIsFindWithinBlueprint
+			&& ensureMsgf(BlueprintEditorPtr.IsValid(), TEXT("A local search was requested, but this widget does not support it.")))
 		{
 			const double StartTime = FPlatformTime::Seconds();
 
diff --git a/Engine/Source/Editor/Kismet/Private/SReplaceNodeReferences.cpp b/Engine/Source/Editor/Kismet/Private/SReplaceNodeReferences.cpp
--- a/Engine/Source/Editor/Kismet/Private/SReplaceNodeReferences.cpp
+++ b/Engine/Source/Editor/Kismet/Private/SReplaceNodeReferences.cpp
@@ -405,6 +405,7 @@ void SReplaceNodeReferences::SetSourceVariable(FProperty* InProperty)
 	else
 	{
 		SourceProperty = nullptr;
+		SelectedTargetReferenceItem.Reset();
 	}
 
 	// Reset the FindInBlueprints results
@@ -693,9 +694,10 @@ void SReplaceNodeReferences::FindAllReplacementsComplete(TArray<FImaginaryFiBDat
 	if (SourceProperty != nullptr && SelectedTargetReferenceItem.IsValid())
 	{
 		FMemberReference SourceVariableReference;
-		SourceVariableReference.SetFromField<FProperty>(SourceProperty, SourceProperty->GetOwnerClass());
+		UClass* SourcePropertyScope = SourceProperty->GetOwnerClass();
+		SourceVariableReference.SetFromField<FProperty>(SourceProperty, SourcePropertyScope);
 		FMemberReference TargetVariableReference;
-		if (SelectedTargetReferenceItem->GetMemberReference(TargetVariableReference) && SourceVariableReference.ResolveMember<FProperty>(SourceProperty->GetOwnerClass()))
+		if (SelectedTargetReferenceItem->GetMemberReference(TargetVariableReference) && SourceVariableReference.ResolveMember<FProperty>(SourcePropertyScope))
 		{
 			TSharedPtr<FBlueprintEditor> PinnedEditor = BlueprintEditor.Pin();
 			if (PinnedEditor.IsValid())
@@ -706,17 +708,22 @@ void SReplaceNodeReferences::FindAllReplacementsComplete(TArray<FImaginaryFiBDat
 					const FScopedTransaction Transaction(GetTransactionTitle(TargetVariableReference));
 					Blueprint->Modify();
 
+					// Note: SourceProperty will be reset after this step! This occurs via a Refresh() that's triggered by an OnBlueprintChanged event.
 					FReplaceNodeReferencesHelper::ReplaceReferences(SourceVariableReference, TargetVariableReference, Blueprint, InRawDataList);
 
 					if (bShowReplacementsWhenFinished)
 					{
+						// @todo - Possibly move this into a local SFindInBlueprints context that's separate from the replacement context.
+						// That way we could limit the results to just the local Blueprint if 'bFindWithinBlueprint' is toggled on. For now
+						// we're just utilizing one of the "floating" global FiB nomad tabs to display the results, which will not have an
+						// associated Blueprint editor context, so we cannot perform a local Blueprint search there.
 						TSharedPtr<SFindInBlueprints> GlobalResults = FFindInBlueprintSearchManager::Get().GetGlobalFindResults();
 						if (GlobalResults)
 						{
 							FStreamSearchOptions SearchOptions;
 							SearchOptions.ImaginaryDataFilter = ESearchQueryFilter::NodesFilter;
 							SearchOptions.MinimiumVersionRequirement = EFiBVersion::FIB_VER_VARIABLE_REFERENCE;
-							GlobalResults->MakeSearchQuery(TargetVariableReference.GetReferenceSearchString(SourceProperty->GetOwnerClass()), bFindWithinBlueprint, SearchOptions);
+							GlobalResults->MakeSearchQuery(TargetVariableReference.GetReferenceSearchString(SourcePropertyScope), /*bInIsFindWithinBlueprint =*/ false, SearchOptions);
 						}
 					}
 				}
