Fixed a crash from FBlueprintVarActionDetails::OnFinishedChangingVariable

This assumed that every property update would be for the BP, but when you have a detail customization that adds an external struct (which several things do) this may not be the case.

This change updates FBlueprintVarActionDetails::OnFinishedChangingVariable to respect the top-level objects passed via the FPropertyChangedEvent, as this will be empty when editing an external struct.

Typically something editing an external struct will trigger its own notification when re-applying that data back to the object, so FBlueprintVarActionDetails::OnFinishedChangingVariable should be able to update via that notification instead.

#rnx

diff --git a/Engine/Source/Editor/Kismet/Private/BlueprintDetailsCustomization.cpp b/Engine/Source/Editor/Kismet/Private/BlueprintDetailsCustomization.cpp
--- a/Engine/Source/Editor/Kismet/Private/BlueprintDetailsCustomization.cpp
+++ b/Engine/Source/Editor/Kismet/Private/BlueprintDetailsCustomization.cpp
@@ -998,7 +998,7 @@ void FBlueprintVarActionDetails::CustomizeDetails( IDetailLayoutBuilder& DetailL
 
 					if (DetailsView.IsValid())
 					{
-						DetailsView->OnFinishedChangingProperties().AddSP(this, &FBlueprintVarActionDetails::OnFinishedChangingVariable, TargetBlueprintDefaultObject);
+						DetailsView->OnFinishedChangingProperties().AddSP(this, &FBlueprintVarActionDetails::OnFinishedChangingVariable);
 					}
 				}
 			}
@@ -2737,14 +2737,14 @@ EVisibility FBlueprintVarActionDetails::IsTooltipEditVisible() const
 	return EVisibility::Collapsed;
 }
 
-void FBlueprintVarActionDetails::OnFinishedChangingVariable(const FPropertyChangedEvent& InPropertyChangedEvent, UObject* InModifiedObjectInstance)
+void FBlueprintVarActionDetails::OnFinishedChangingVariable(const FPropertyChangedEvent& InPropertyChangedEvent)
 {
-	if (!InModifiedObjectInstance)
+	if (InPropertyChangedEvent.GetNumObjectsBeingEdited() == 0)
 	{
 		return;
 	}
 
-	ImportNamespacesForPropertyValue(InPropertyChangedEvent.MemberProperty, InModifiedObjectInstance);
+	ImportNamespacesForPropertyValue(InPropertyChangedEvent.MemberProperty, InPropertyChangedEvent.GetObjectBeingEdited(0));
 }
 
 void FBlueprintVarActionDetails::OnFinishedChangingLocalVariable(const FPropertyChangedEvent& InPropertyChangedEvent, TSharedPtr<FStructOnScope> InStructData, TWeakObjectPtr<UK2Node_EditablePinBase> InEntryNode)
diff --git a/Engine/Source/Editor/Kismet/Private/BlueprintDetailsCustomization.h b/Engine/Source/Editor/Kismet/Private/BlueprintDetailsCustomization.h
--- a/Engine/Source/Editor/Kismet/Private/BlueprintDetailsCustomization.h
+++ b/Engine/Source/Editor/Kismet/Private/BlueprintDetailsCustomization.h
@@ -148,9 +148,8 @@ private:
 	 * Callback when changing a variable property
 	 *
 	 * @param InPropertyChangedEvent	Information on the property changed
-	 * @param InModifiedObjectInstance	The object instance for which the value was changed
 	 */
-	void OnFinishedChangingVariable(const FPropertyChangedEvent& InPropertyChangedEvent, UObject* InModifiedObjectInstance);
+	void OnFinishedChangingVariable(const FPropertyChangedEvent& InPropertyChangedEvent);
 
 	/**
 	 * Callback when changing a local variable property
