[Bug Fix] Crash when opening a collapsed node after undoing the creation of a second execution pin

diff --git a/Engine/Source/Editor/Kismet/Private/BlueprintDetailsCustomization.cpp b/Engine/Source/Editor/Kismet/Private/BlueprintDetailsCustomization.cpp
--- a/Engine/Source/Editor/Kismet/Private/BlueprintDetailsCustomization.cpp
+++ b/Engine/Source/Editor/Kismet/Private/BlueprintDetailsCustomization.cpp
@@ -3047,7 +3047,7 @@ void FBlueprintGraphArgumentGroupLayout::GenerateChildContent( IDetailChildrenBu
 				}
 				else
 				{
-					ArgumentName = bIsInputNode ? FString::Printf(TEXT("%s.InputArgument%i"), i) : FString::Printf(TEXT("OutputArgument%i"), i);
+					ArgumentName = bIsInputNode ? FString::Printf(TEXT("InputArgument%i"), i) : FString::Printf(TEXT("OutputArgument%i"), i);
 				}
 
 				TSharedRef<class FBlueprintGraphArgumentLayout> BlueprintArgumentLayout = MakeShareable(new FBlueprintGraphArgumentLayout(
