Removing uses of master/slave found in Kismet with base folder Editor

diff --git a/Engine/Source/Editor/Kismet/Private/SBlueprintDiff.cpp b/Engine/Source/Editor/Kismet/Private/SBlueprintDiff.cpp
--- a/Engine/Source/Editor/Kismet/Private/SBlueprintDiff.cpp
+++ b/Engine/Source/Editor/Kismet/Private/SBlueprintDiff.cpp
@@ -958,7 +958,7 @@ void SBlueprintDiff::Construct( const FArguments& InArgs)
 		, TAttribute<FSlateIcon>(this, &SBlueprintDiff::GetSplitViewModeImage)
 	);
 
-	DifferencesTreeView = DiffTreeView::CreateTreeView(&MasterDifferencesList);
+	DifferencesTreeView = DiffTreeView::CreateTreeView(&PrimaryDifferencesList);
 
 	GenerateDifferencesList();
 
@@ -1193,12 +1193,12 @@ TSharedPtr<SWindow> SBlueprintDiff::CreateDiffWindow(FText WindowTitle, UBluepri
 
 void SBlueprintDiff::NextDiff()
 {
-	DiffTreeView::HighlightNextDifference(DifferencesTreeView.ToSharedRef(), RealDifferences, MasterDifferencesList);
+	DiffTreeView::HighlightNextDifference(DifferencesTreeView.ToSharedRef(), RealDifferences, PrimaryDifferencesList);
 }
 
 void SBlueprintDiff::PrevDiff()
 {
-	DiffTreeView::HighlightPrevDifference(DifferencesTreeView.ToSharedRef(), RealDifferences, MasterDifferencesList);
+	DiffTreeView::HighlightPrevDifference(DifferencesTreeView.ToSharedRef(), RealDifferences, PrimaryDifferencesList);
 }
 
 bool SBlueprintDiff::HasNextDiff() const
@@ -1514,7 +1514,7 @@ void SBlueprintDiff::HandleGraphChanged( const FString& GraphPath )
 
 void SBlueprintDiff::GenerateDifferencesList()
 {
-	MasterDifferencesList.Empty();
+	PrimaryDifferencesList.Empty();
 	RealDifferences.Empty();
 	Graphs.Empty();
 	ModePanels.Empty();
@@ -1599,7 +1599,7 @@ void SBlueprintDiff::GenerateDifferencesList()
 
 	for (const TSharedPtr<FGraphToDiff>& Graph : Graphs)
 	{
-		Graph->GenerateTreeEntries(MasterDifferencesList, RealDifferences);
+		Graph->GenerateTreeEntries(PrimaryDifferencesList, RealDifferences);
 	}
 
 	DifferencesTreeView->RebuildList();
@@ -1608,7 +1608,7 @@ void SBlueprintDiff::GenerateDifferencesList()
 SBlueprintDiff::FDiffControl SBlueprintDiff::GenerateBlueprintTypePanel()
 {
 	TSharedPtr<FBlueprintTypeDiffControl> NewDiffControl = MakeShared<FBlueprintTypeDiffControl>(PanelOld.Blueprint, PanelNew.Blueprint, FOnDiffEntryFocused::CreateRaw(this, &SBlueprintDiff::SetCurrentMode, BlueprintTypeMode));
-	NewDiffControl->GenerateTreeEntries(MasterDifferencesList, RealDifferences);
+	NewDiffControl->GenerateTreeEntries(PrimaryDifferencesList, RealDifferences);
 
 	SBlueprintDiff::FDiffControl Ret;
 	//Splitter for left and right blueprint. Current convention is for the local (probably newer?) blueprint to be on the right:
@@ -1640,7 +1640,7 @@ SBlueprintDiff::FDiffControl SBlueprintDiff::GenerateBlueprintTypePanel()
 SBlueprintDiff::FDiffControl SBlueprintDiff::GenerateMyBlueprintPanel()
 {
 	TSharedPtr<FMyBlueprintDiffControl> NewDiffControl = MakeShared<FMyBlueprintDiffControl>(PanelOld.Blueprint, PanelNew.Blueprint, FOnDiffEntryFocused::CreateRaw(this, &SBlueprintDiff::SetCurrentMode, MyBlueprintMode));
-	NewDiffControl->GenerateTreeEntries(MasterDifferencesList, RealDifferences);
+	NewDiffControl->GenerateTreeEntries(PrimaryDifferencesList, RealDifferences);
 
 	SBlueprintDiff::FDiffControl Ret;
 
@@ -1746,7 +1746,7 @@ SBlueprintDiff::FDiffControl SBlueprintDiff::GenerateDefaultsPanel()
 	const UObject* B = DiffUtils::GetCDO(PanelNew.Blueprint);
 
 	TSharedPtr<FCDODiffControl> NewDiffControl = MakeShared<FCDODiffControl>(A, B, FOnDiffEntryFocused::CreateRaw(this, &SBlueprintDiff::SetCurrentMode, DefaultsMode));
-	NewDiffControl->GenerateTreeEntries(MasterDifferencesList, RealDifferences);
+	NewDiffControl->GenerateTreeEntries(PrimaryDifferencesList, RealDifferences);
 
 	SBlueprintDiff::FDiffControl Ret;
 	Ret.DiffControl = NewDiffControl;
@@ -1769,7 +1769,7 @@ SBlueprintDiff::FDiffControl SBlueprintDiff::GenerateDefaultsPanel()
 SBlueprintDiff::FDiffControl SBlueprintDiff::GenerateClassSettingsPanel()
 {
 	TSharedPtr<FClassSettingsDiffControl> NewDiffControl = MakeShared<FClassSettingsDiffControl>(PanelOld.Blueprint, PanelNew.Blueprint, FOnDiffEntryFocused::CreateRaw(this, &SBlueprintDiff::SetCurrentMode, ClassSettingsMode));
-	NewDiffControl->GenerateTreeEntries(MasterDifferencesList, RealDifferences);
+	NewDiffControl->GenerateTreeEntries(PrimaryDifferencesList, RealDifferences);
 
 	SBlueprintDiff::FDiffControl Ret;
 	Ret.DiffControl = NewDiffControl;
@@ -1792,7 +1792,7 @@ SBlueprintDiff::FDiffControl SBlueprintDiff::GenerateClassSettingsPanel()
 SBlueprintDiff::FDiffControl SBlueprintDiff::GenerateComponentsPanel()
 {
 	TSharedPtr<FSCSDiffControl> NewDiffControl = MakeShared<FSCSDiffControl>(PanelOld.Blueprint, PanelNew.Blueprint, FOnDiffEntryFocused::CreateRaw(this, &SBlueprintDiff::SetCurrentMode, ComponentsMode));
-	NewDiffControl->GenerateTreeEntries(MasterDifferencesList, RealDifferences);
+	NewDiffControl->GenerateTreeEntries(PrimaryDifferencesList, RealDifferences);
 
 	SBlueprintDiff::FDiffControl Ret;
 	Ret.DiffControl = NewDiffControl;
diff --git a/Engine/Source/Editor/Kismet/Public/SBlueprintDiff.h b/Engine/Source/Editor/Kismet/Public/SBlueprintDiff.h
--- a/Engine/Source/Editor/Kismet/Public/SBlueprintDiff.h
+++ b/Engine/Source/Editor/Kismet/Public/SBlueprintDiff.h
@@ -248,7 +248,7 @@ protected:
 	TSharedPtr<FTabManager> TabManager;
 
 	/** Tree of differences collected across all panels: */
-	TArray< TSharedPtr<class FBlueprintDifferenceTreeEntry> > MasterDifferencesList;
+	TArray< TSharedPtr<class FBlueprintDifferenceTreeEntry> > PrimaryDifferencesList;
 
 	/** List of all differences, cached so that we can iterate only the differences and not labels, etc: */
 	TArray< TSharedPtr<class FBlueprintDifferenceTreeEntry> > RealDifferences;
