Used the code which builds the access specifier (previously located in the building on the function UI) to create a lambda expression to build it, which I then used in both the portion that builds the variable view (only if it's not a local variable), and again in the portion that builds the function view.

diff --git a/Engine/Source/Editor/Kismet/Private/SBlueprintPalette.cpp b/Engine/Source/Editor/Kismet/Private/SBlueprintPalette.cpp
--- a/Engine/Source/Editor/Kismet/Private/SBlueprintPalette.cpp
+++ b/Engine/Source/Editor/Kismet/Private/SBlueprintPalette.cpp
@@ -1274,9 +1274,37 @@ void SBlueprintPaletteItem::Construct(const FArguments& InArgs, FCreateWidgetFor
 	// Create the widget with an icon
 	TSharedRef<SHorizontalBox> ActionBox = SNew(SHorizontalBox)		
 		.AddMetaData<FTutorialMetaData>(TagMeta);
+
+
+	auto CreateAccessSpecifierLambda = [&ActionBox, &AccessSpecifierEnabled, &AccessModifierText, &ActionAccessSpecifier]() {
+
+		ActionBox.Get().AddSlot()
+			.MaxWidth(50.f)
+			.FillWidth(AccessSpecifierEnabled ? 0.4f : 0.0f)
+			.Padding(FMargin(/* horizontal */ AccessSpecifierEnabled ? 6.0f : 0.0f, /* vertical */ 0.0f))
+			.VAlign(VAlign_Center)
+			.HAlign(HAlign_Right)
+			[
+				SNew(STextBlock)
+				// Will only display text if we have a modifier level
+			.IsEnabled(AccessSpecifierEnabled)
+			.Text(AccessModifierText)
+			.ColorAndOpacity(FSlateColor::UseSubduedForeground())
+			// Bold if public
+			.TextStyle(FAppStyle::Get(), ActionAccessSpecifier == EAccessSpecifier::Public ? "BlueprintEditor.AccessModifier.Public" : "BlueprintEditor.AccessModifier.Default")
+			];
+	};
+
 	
 	if (GraphAction->IsA(FEdGraphSchemaAction_BlueprintVariableBase::StaticGetTypeId()))
 	{
+
+
+		if (ActionAccessSpecifier != EAccessSpecifier::None && GraphAction->GetTypeId() != FEdGraphSchemaAction_K2LocalVar::StaticGetTypeId())
+		{
+			CreateAccessSpecifierLambda();
+		}
+
 		ActionBox.Get().AddSlot()
 			.FillWidth(0.6f)
 			.VAlign(VAlign_Center)
@@ -1326,21 +1354,7 @@ void SBlueprintPaletteItem::Construct(const FArguments& InArgs, FCreateWidgetFor
 		// Only add an access specifier if we have one
 		if (ActionAccessSpecifier != EAccessSpecifier::None)
 		{
-			ActionBox.Get().AddSlot()
-				.MaxWidth(50.f)
-				.FillWidth(AccessSpecifierEnabled ? 0.4f : 0.0f)
-				.Padding(FMargin(/* horizontal */ AccessSpecifierEnabled ? 6.0f : 0.0f, /* vertical */ 0.0f))
-				.VAlign(VAlign_Center)
-				.HAlign(HAlign_Right)
-				[
-					SNew(STextBlock)
-					// Will only display text if we have a modifier level
-					.IsEnabled(AccessSpecifierEnabled)
-					.Text(AccessModifierText)
-					.ColorAndOpacity(FSlateColor::UseSubduedForeground())
-					// Bold if public
-					.TextStyle(FAppStyle::Get(), ActionAccessSpecifier == EAccessSpecifier::Public ? "BlueprintEditor.AccessModifier.Public" : "BlueprintEditor.AccessModifier.Default")
-				];
+			CreateAccessSpecifierLambda();
 		}
 
 		ActionBox.Get().AddSlot()
