Fix crash when duplicating a state machine directly in the MyBlueprint panel within an Animation BP, by adding node clone functionality to the BlueprintEditor

diff --git a/Engine/Source/Editor/Kismet/Private/BlueprintEditor.cpp b/Engine/Source/Editor/Kismet/Private/BlueprintEditor.cpp
--- a/Engine/Source/Editor/Kismet/Private/BlueprintEditor.cpp
+++ b/Engine/Source/Editor/Kismet/Private/BlueprintEditor.cpp
@@ -2246,6 +2246,15 @@ void FBlueprintEditor::RemoveNamespace(const FString& InNamespace)
 	}
 }
 
+void FBlueprintEditor::SelectAndDuplicateNode(UEdGraphNode* InNode)
+{
+	check(InNode != nullptr);
+
+	ClearSelectionStateFor(CurrentUISelection);
+	AddToSelection(InNode);
+	DuplicateNodes();
+}
+
 void FBlueprintEditor::RegisterTabSpawners(const TSharedRef<class FTabManager>& InTabManager)
 {
 	//@TODO: Can't we do this sooner?
diff --git a/Engine/Source/Editor/Kismet/Private/SMyBlueprint.cpp b/Engine/Source/Editor/Kismet/Private/SMyBlueprint.cpp
--- a/Engine/Source/Editor/Kismet/Private/SMyBlueprint.cpp
+++ b/Engine/Source/Editor/Kismet/Private/SMyBlueprint.cpp
@@ -38,6 +38,8 @@
 #include "BlueprintEditorCommands.h"
 #include "GraphEditorActions.h"
 #include "AnimationGraph.h"
+#include "AnimationStateMachineGraph.h"
+#include "AnimGraphNode_StateMachineBase.h"
 #include "SBlueprintEditorToolbar.h"
 #include "Kismet2/BlueprintEditorUtils.h"
 #include "ObjectEditorUtils.h"
@@ -3289,44 +3291,56 @@ void SMyBlueprint::OnDuplicateAction()
 
 	if (FEdGraphSchemaAction_K2Graph* GraphAction = SelectionAsGraph())
 	{
-		const FScopedTransaction Transaction( LOCTEXT( "DuplicateGraph", "Duplicate Graph" ) );
-		GetBlueprintObj()->Modify();
-
-		UEdGraph* DuplicatedGraph = GraphAction->EdGraph->GetSchema()->DuplicateGraph(GraphAction->EdGraph);
-		check(DuplicatedGraph);
-
-		DuplicatedGraph->Modify();
+		// Only StateMachine, function, anim graph and macro duplication is supported
+		EGraphType GraphType = GraphAction->EdGraph->GetSchema()->GetGraphType(GraphAction->EdGraph);
+		check(GraphType == GT_StateMachine || GraphType == GT_Function || GraphType == GT_Macro || GraphType == GT_Animation);
 
-		// Generate new Guids and component templates for all relevant nodes in the graph
-		// *NOTE* this cannot occur during PostDuplicate, node Guids and component templates need to remain static during duplication for Blueprint compilation
-		for (UEdGraphNode* EdGraphNode : DuplicatedGraph->Nodes)
+		if (GraphType == GT_StateMachine)
 		{
-			if (EdGraphNode)
+			// StateMachine is handled using the BlueprintEditor copy / paste functionality
+			if (const UAnimationStateMachineGraph* AnimationStateMachineGraph = Cast<UAnimationStateMachineGraph>(GraphAction->EdGraph))
 			{
-				EdGraphNode->CreateNewGuid();
+				BlueprintEditorPtr.Pin()->SelectAndDuplicateNode(AnimationStateMachineGraph->OwnerAnimGraphNode.Get());
+			}
+		}
+		else
+		{
+			const FScopedTransaction Transaction(LOCTEXT("DuplicateGraph", "Duplicate Graph"));
+			GetBlueprintObj()->Modify();
+
+			UEdGraph* DuplicatedGraph = GraphAction->EdGraph->GetSchema()->DuplicateGraph(GraphAction->EdGraph);
+			check(DuplicatedGraph);
 
-				if (UK2Node_AddComponent* AddComponentNode = Cast<UK2Node_AddComponent>(EdGraphNode))
+			DuplicatedGraph->Modify();
+
+			// Generate new Guids and component templates for all relevant nodes in the graph
+			// *NOTE* this cannot occur during PostDuplicate, node Guids and component templates need to remain static during duplication for Blueprint compilation
+			for (UEdGraphNode* EdGraphNode : DuplicatedGraph->Nodes)
+			{
+				if (EdGraphNode)
 				{
-					AddComponentNode->MakeNewComponentTemplate();
+					EdGraphNode->CreateNewGuid();
+
+					if (UK2Node_AddComponent* AddComponentNode = Cast<UK2Node_AddComponent>(EdGraphNode))
+					{
+						AddComponentNode->MakeNewComponentTemplate();
+					}
 				}
 			}
-		}
-		// Only function, anim graph and macro duplication is supported
-		EGraphType GraphType = DuplicatedGraph->GetSchema()->GetGraphType(GraphAction->EdGraph);
-		check(GraphType == GT_Function || GraphType == GT_Macro || GraphType == GT_Animation);
 
-		if (GraphType == GT_Function || GraphType == GT_Animation)
-		{
-			GetBlueprintObj()->FunctionGraphs.Add(DuplicatedGraph);
-		}
-		else if (GraphType == GT_Macro)
-		{
-			GetBlueprintObj()->MacroGraphs.Add(DuplicatedGraph);
-		}
-		FBlueprintEditorUtils::MarkBlueprintAsStructurallyModified(GetBlueprintObj());
+			if (GraphType == GT_Function || GraphType == GT_Animation)
+			{
+				GetBlueprintObj()->FunctionGraphs.Add(DuplicatedGraph);
+			}
+			else if (GraphType == GT_Macro)
+			{
+				GetBlueprintObj()->MacroGraphs.Add(DuplicatedGraph);
+			}
+			FBlueprintEditorUtils::MarkBlueprintAsStructurallyModified(GetBlueprintObj());
 
-		BlueprintEditorPtr.Pin()->OpenDocument(DuplicatedGraph, FDocumentTracker::ForceOpenNewDocument);
-		DuplicateActionName = DuplicatedGraph->GetFName();
+			BlueprintEditorPtr.Pin()->OpenDocument(DuplicatedGraph, FDocumentTracker::ForceOpenNewDocument);
+			DuplicateActionName = DuplicatedGraph->GetFName();
+		}
 	}
 	else if (FEdGraphSchemaAction_K2Var* VarAction = SelectionAsVar())
 	{
diff --git a/Engine/Source/Editor/Kismet/Public/BlueprintEditor.h b/Engine/Source/Editor/Kismet/Public/BlueprintEditor.h
--- a/Engine/Source/Editor/Kismet/Public/BlueprintEditor.h
+++ b/Engine/Source/Editor/Kismet/Public/BlueprintEditor.h
@@ -756,6 +756,9 @@ public:
 	/** Selects a local variable to load in the details panel. */
 	virtual bool SelectLocalVariable(const UEdGraph* Graph, const FName& VariableName) { return false; }
 
+	/** Duplicates an EdGraphNode using selection and copy / paste functionality */
+	void SelectAndDuplicateNode(UEdGraphNode* InNode);
+
 protected:
 	UE_DEPRECATED(4.26, "Please do any validation inside the UBlueprint class during compilation, extra errors during compiling only supplied by the designer can lead to design time only errors being reported and being missed during cooks/content validation.")
 	virtual void AppendExtraCompilerResults(TSharedPtr<class IMessageLogListing> ResultsListing) {}
