UE-161502 - CompileAllBlueprints crashing on blueprint function with a delegate input.

When we fixup delegate properties during compilation, we need to ensure that the new SignatureFunction is added to the parent UStruct's ScriptAndPropertyObjectReferences array. This array is what keeps the function alive during a GC run, and it needs to be in sync with what the underlying delegate property points to.

diff --git a/Engine/Source/Editor/Kismet/Private/BlueprintCompilationManager.cpp b/Engine/Source/Editor/Kismet/Private/BlueprintCompilationManager.cpp
--- a/Engine/Source/Editor/Kismet/Private/BlueprintCompilationManager.cpp
+++ b/Engine/Source/Editor/Kismet/Private/BlueprintCompilationManager.cpp
@@ -378,6 +378,22 @@ struct FSkeletonFixupData
 {
 	FSimpleMemberReference MemberReference;
 	FProperty* DelegateProperty;
+
+	template<typename T>
+	static void FixUpDelegateProperty(T* DelegateProperty, const FSimpleMemberReference& MemberReference, UClass* SkeletonGeneratedClass)
+	{
+		check(DelegateProperty);
+	
+		UFunction* OldFunction = DelegateProperty->SignatureFunction;
+		DelegateProperty->SignatureFunction = FMemberReference::ResolveSimpleMemberReference<UFunction>(MemberReference, SkeletonGeneratedClass);
+		UStruct* Owner = DelegateProperty->template GetOwnerChecked<UStruct>();
+
+		int32 Index = Owner->ScriptAndPropertyObjectReferences.Find(OldFunction);
+		if (Index != INDEX_NONE)
+		{
+			Owner->ScriptAndPropertyObjectReferences[Index] = DelegateProperty->SignatureFunction;
+		}
+	}
 };
 
 struct FCompilerData
@@ -1030,16 +1046,16 @@ void FBlueprintCompilationManagerImpl::FlushCompilationQueueImpl(bool bSuppressB
 				DECLARE_SCOPE_HIERARCHICAL_COUNTER(FixUpDelegateParameters)
 
 				UBlueprint* BP = CompilerData.BP;
-				TArray< FSkeletonFixupData >& ParamsToFix = CompilerData.SkeletonFixupData;
+				TArray<FSkeletonFixupData>& ParamsToFix = CompilerData.SkeletonFixupData;
 				for( const FSkeletonFixupData& FixupData : ParamsToFix )
 				{
 					if(FDelegateProperty* DelegateProperty = CastField<FDelegateProperty>(FixupData.DelegateProperty))
 					{
-						DelegateProperty->SignatureFunction = FMemberReference::ResolveSimpleMemberReference<UFunction>(FixupData.MemberReference, BP->SkeletonGeneratedClass);
+						FSkeletonFixupData::FixUpDelegateProperty(DelegateProperty, FixupData.MemberReference, BP->SkeletonGeneratedClass);
 					}
 					else if(FMulticastDelegateProperty* MCDelegateProperty = CastField<FMulticastDelegateProperty>(FixupData.DelegateProperty))
 					{
-						MCDelegateProperty->SignatureFunction = FMemberReference::ResolveSimpleMemberReference<UFunction>(FixupData.MemberReference, BP->SkeletonGeneratedClass);
+						FSkeletonFixupData::FixUpDelegateProperty(MCDelegateProperty, FixupData.MemberReference, BP->SkeletonGeneratedClass);
 					}
 				}
 			}
