Auto-import the default namespace set for Blueprint types instead of only its assigned namespace.

Change summary:
- Added FBlueprintNamespaceUtilities::GetDefaultImportsForObject().
- Deprecated FBlueprintNamespaceUtilities::GetDefaultImportsForBlueprint() in favor of the more generalized GetDefaultImportsForObject().
- Changed auto-import call sites to query FBlueprintNamespaceUtilities::GetDefaultImportsForObject() rather than using GetObjectNamespace().
- Removed IsNonImportedObject() as a condition for whether or not call sites should initiate auto-import; the logic already determines whether or not any imports will need to occur based on the input set.

diff --git a/Engine/Source/Editor/Kismet/Private/BlueprintDetailsCustomization.cpp b/Engine/Source/Editor/Kismet/Private/BlueprintDetailsCustomization.cpp
--- a/Engine/Source/Editor/Kismet/Private/BlueprintDetailsCustomization.cpp
+++ b/Engine/Source/Editor/Kismet/Private/BlueprintDetailsCustomization.cpp
@@ -1456,15 +1456,13 @@ void FBlueprintVarActionDetails::OnVarTypeChanged(const FEdGraphPinType& NewPinT
 
 				TSharedPtr<FBlueprintEditor> BlueprintEditor = MyBlueprint.Pin()->GetBlueprintEditor().Pin();
 					
-				// If the underlying type object's namespace is not imported, auto-import it now into the current editor context.
+				// Auto-import the underlying type object's default namespace set into the current editor context.
 				const UObject* PinSubCategoryObject = NewPinType.PinSubCategoryObject.Get();
-				if (PinSubCategoryObject && BlueprintEditor.IsValid() && BlueprintEditor->IsNonImportedObject(PinSubCategoryObject))
+				if (PinSubCategoryObject && BlueprintEditor.IsValid())
 				{
-					FString Namespace = FBlueprintNamespaceUtilities::GetObjectNamespace(PinSubCategoryObject);
-					if (!Namespace.IsEmpty())
-					{
-						BlueprintEditor->ImportNamespace(Namespace);
-					}
+					FBlueprintEditor::FImportNamespaceExParameters Params;
+					FBlueprintNamespaceUtilities::GetDefaultImportsForObject(PinSubCategoryObject, Params.NamespacesToImport);
+					BlueprintEditor->ImportNamespaceEx(Params);
 				}
 			}
 		}
@@ -3468,15 +3466,13 @@ void FBlueprintGraphArgumentLayout::PinInfoChanged(const FEdGraphPinType& PinTyp
 								BlueprintEditor = MyBPPinned->GetBlueprintEditor().Pin();
 							}
 
-							// If the underlying type object's namespace is not imported, auto-import it now into the current editor context.
+							// Auto-import the underlying type object's default namespace set into the current editor context.
 							const UObject* PinSubCategoryObject = PinType.PinSubCategoryObject.Get();
-							if (PinSubCategoryObject && BlueprintEditor.IsValid() && BlueprintEditor->IsNonImportedObject(PinSubCategoryObject))
+							if (PinSubCategoryObject && BlueprintEditor.IsValid())
 							{
-								FString Namespace = FBlueprintNamespaceUtilities::GetObjectNamespace(PinSubCategoryObject);
-								if (!Namespace.IsEmpty())
-								{
-									BlueprintEditor->ImportNamespace(Namespace);
-								}
+								FBlueprintEditor::FImportNamespaceExParameters Params;
+								FBlueprintNamespaceUtilities::GetDefaultImportsForObject(PinSubCategoryObject, Params.NamespacesToImport);
+								BlueprintEditor->ImportNamespaceEx(Params);
 							}
 						}
 						GraphActionDetailsPinned->OnParamsChanged(Node);
@@ -5673,7 +5669,7 @@ void FBlueprintImportsLayout::GetManagedListItems(TArray<FManagedListItem>& OutL
 
 	// Default imports (non-removable). These include anything from the shared global set, as well as any namespaces assigned to the Blueprint hierarchy.
 	FBlueprintNamespaceUtilities::GetSharedGlobalImports(NamespaceItems);
-	FBlueprintNamespaceUtilities::GetDefaultImportsForBlueprint(Blueprint, NamespaceItems);
+	FBlueprintNamespaceUtilities::GetDefaultImportsForObject(Blueprint, NamespaceItems);
 
 	if(!bShouldShowDefaultImports)
 	{
@@ -5722,7 +5718,7 @@ void FBlueprintImportsLayout::OnGetNamespacesToExclude(TSet<FString>& OutNamespa
 	const UBlueprint* Blueprint = GetBlueprintObjectChecked();
 
 	FBlueprintNamespaceUtilities::GetSharedGlobalImports(OutNamespacesToExclude);
-	FBlueprintNamespaceUtilities::GetDefaultImportsForBlueprint(Blueprint, OutNamespacesToExclude);
+	FBlueprintNamespaceUtilities::GetDefaultImportsForObject(Blueprint, OutNamespacesToExclude);
 
 	OutNamespacesToExclude.Append(Blueprint->ImportedNamespaces);
 }
diff --git a/Engine/Source/Editor/Kismet/Private/BlueprintEditor.cpp b/Engine/Source/Editor/Kismet/Private/BlueprintEditor.cpp
--- a/Engine/Source/Editor/Kismet/Private/BlueprintEditor.cpp
+++ b/Engine/Source/Editor/Kismet/Private/BlueprintEditor.cpp
@@ -1287,9 +1287,10 @@ void FBlueprintEditor::OnComponentAddedToBlueprint(const FSubobjectData& NewSubo
 	const UObject* NewSubobject = NewSubobjectData.GetObject();
 	check(NewSubobject);
 
-	// If necessary, auto-import the namespace associated with the new subobject's class.
-	const FString SubobjectTypeNamespace = FBlueprintNamespaceUtilities::GetObjectNamespace(NewSubobject->GetClass());
-	ImportNamespace(SubobjectTypeNamespace);
+	// If necessary, auto-import the default namespace set associated with the new subobject's class.
+	FImportNamespaceExParameters Params;
+	FBlueprintNamespaceUtilities::GetDefaultImportsForObject(NewSubobject, Params.NamespacesToImport);
+	ImportNamespaceEx(Params);
 }
 
 TSharedRef<SWidget> FBlueprintEditor::CreateGraphTitleBarWidget(TSharedRef<FTabInfo> InTabInfo, UEdGraph* InGraph)
@@ -3405,14 +3406,14 @@ void FBlueprintEditor::ReparentBlueprint_NewParentChosen(UClass* ChosenClass)
 
 			// Gather the set of default imports with the old parent class set.
 			TSet<FString> OldDefaultImports;
-			FBlueprintNamespaceUtilities::GetDefaultImportsForBlueprint(BlueprintObj, OldDefaultImports);
+			FBlueprintNamespaceUtilities::GetDefaultImportsForObject(BlueprintObj, OldDefaultImports);
 
 			UClass* OldParentClass = BlueprintObj->ParentClass;
 			BlueprintObj->ParentClass = ChosenClass;
 
 			// Gather the set of default imports with the new parent class set.
 			TSet<FString> NewDefaultImports;
-			FBlueprintNamespaceUtilities::GetDefaultImportsForBlueprint(BlueprintObj, NewDefaultImports);
+			FBlueprintNamespaceUtilities::GetDefaultImportsForObject(BlueprintObj, NewDefaultImports);
 
 			// Move namespace imports that no longer appear in the default set to the explicit set.
 			FImportNamespaceExParameters Params;
@@ -7402,11 +7403,7 @@ void FBlueprintEditor::PasteNodesHere(class UEdGraph* DestinationGraph, const FV
 			{
 				for (const UStruct* ExternalDependency : ExternalDependencies)
 				{
-					FString ObjectNamespace = FBlueprintNamespaceUtilities::GetObjectNamespace(ExternalDependency);
-					if (!ObjectNamespace.IsEmpty())
-					{
-						NamespacesToImport.Add(MoveTemp(ObjectNamespace));
-					}
+					FBlueprintNamespaceUtilities::GetDefaultImportsForObject(ExternalDependency, NamespacesToImport);
 				}
 			}
 
diff --git a/Engine/Source/Editor/Kismet/Private/BlueprintNamespaceHelper.cpp b/Engine/Source/Editor/Kismet/Private/BlueprintNamespaceHelper.cpp
--- a/Engine/Source/Editor/Kismet/Private/BlueprintNamespaceHelper.cpp
+++ b/Engine/Source/Editor/Kismet/Private/BlueprintNamespaceHelper.cpp
@@ -147,7 +147,7 @@ void FBlueprintNamespaceHelper::AddBlueprint(const UBlueprint* InBlueprint)
 
 	// Add the default import set for the given Blueprint.
 	TSet<FString> DefaultImports;
-	FBlueprintNamespaceUtilities::GetDefaultImportsForBlueprint(InBlueprint, DefaultImports);
+	FBlueprintNamespaceUtilities::GetDefaultImportsForObject(InBlueprint, DefaultImports);
 	AddNamespaces(DefaultImports);
 
 	// Additional namespaces that are explicitly imported by this Blueprint.
diff --git a/Engine/Source/Editor/Kismet/Private/BlueprintNamespaceUtilities.cpp b/Engine/Source/Editor/Kismet/Private/BlueprintNamespaceUtilities.cpp
--- a/Engine/Source/Editor/Kismet/Private/BlueprintNamespaceUtilities.cpp
+++ b/Engine/Source/Editor/Kismet/Private/BlueprintNamespaceUtilities.cpp
@@ -229,43 +229,54 @@ void FBlueprintNamespaceUtilities::GetSharedGlobalImports(TSet<FString>& OutName
 	OutNamespaces.Remove(FString());
 }
 
-void FBlueprintNamespaceUtilities::GetDefaultImportsForBlueprint(const UBlueprint* InBlueprint, TSet<FString>& OutNamespaces)
+void FBlueprintNamespaceUtilities::GetDefaultImportsForObject(const UObject* InObject, TSet<FString>& OutNamespaces)
 {
-	if (!InBlueprint)
+	if (!InObject)
 	{
 		return;
 	}
 
-	// Blueprint namespace (if set).
-	FString BlueprintNamespace = GetObjectNamespace(InBlueprint);
-	if (!BlueprintNamespace.IsEmpty())
+	// Get the object's associated namespace (if set).
+	FString ObjectNamespace = GetObjectNamespace(InObject);
+	if (!ObjectNamespace.IsEmpty())
 	{
-		OutNamespaces.Add(BlueprintNamespace);
+		OutNamespaces.Add(ObjectNamespace);
 	}
 
-	const bool bAddParentClassImportedNamespaces = GetDefault<UBlueprintEditorSettings>()->bInheritImportedNamespacesFromParentBP;
-
-	// Inherited namespaces from the parent class hierarchy (if set).
-	UClass* ParentClass = InBlueprint->ParentClass;
-	while (ParentClass)
+	if (const UBlueprint* Blueprint = Cast<UBlueprint>(InObject))
 	{
-		// Parent class namespace (if set).
-		FString ParentClassNamespace = GetObjectNamespace(ParentClass);
-		if (!ParentClassNamespace.IsEmpty())
-		{
-			OutNamespaces.Add(ParentClassNamespace);
-		}
+		// Append inherited namespaces from the parent class hierarchy.
+		GetDefaultImportsForObject(Blueprint->ParentClass, OutNamespaces);
+	}
+	else if (const UStruct* StructType = Cast<UStruct>(InObject))
+	{
+		const bool bAddParentClassImportedNamespaces = GetDefault<UBlueprintEditorSettings>()->bInheritImportedNamespacesFromParentBP;
 
-		// If enabled, also include namespaces that are explicitly imported by all ancestor BPs.
-		if (bAddParentClassImportedNamespaces)
+		UStruct* ParentStruct = StructType->GetSuperStruct();
+		while (ParentStruct)
 		{
-			if (const UBlueprint* ParentBP = UBlueprint::GetBlueprintFromClass(ParentClass))
+			// Parent type namespace (if set).
+			FString ParentTypeNamespace = GetObjectNamespace(ParentStruct);
+			if (!ParentTypeNamespace.IsEmpty())
 			{
-				OutNamespaces.Append(ParentBP->ImportedNamespaces);
+				OutNamespaces.Add(ParentTypeNamespace);
 			}
-		}
 
-		ParentClass = ParentClass->GetSuperClass();
+			// If enabled, also include namespaces that are explicitly imported by all ancestor BPs.
+			if (bAddParentClassImportedNamespaces)
+			{
+				if (const UClass* ParentClass = Cast<UClass>(ParentStruct))
+				{
+					const UBlueprint* ParentBP = UBlueprint::GetBlueprintFromClass(ParentClass);
+					if (ParentBP)
+					{
+						OutNamespaces.Append(ParentBP->ImportedNamespaces);
+					}
+				}
+			}
+
+			ParentStruct = ParentStruct->GetSuperStruct();
+		}
 	}
 }
 
diff --git a/Engine/Source/Editor/Kismet/Private/SBlueprintActionMenu.cpp b/Engine/Source/Editor/Kismet/Private/SBlueprintActionMenu.cpp
--- a/Engine/Source/Editor/Kismet/Private/SBlueprintActionMenu.cpp
+++ b/Engine/Source/Editor/Kismet/Private/SBlueprintActionMenu.cpp
@@ -592,11 +592,7 @@ void SBlueprintActionMenu::OnActionSelected( const TArray< TSharedPtr<FEdGraphSc
 						{
 							for (const UStruct* ExternalDependency : ExternalDependencies)
 							{
-								FString ObjectNamespace = FBlueprintNamespaceUtilities::GetObjectNamespace(ExternalDependency);
-								if (!ObjectNamespace.IsEmpty())
-								{
-									NamespacesToImport.Add(MoveTemp(ObjectNamespace));
-								}
+								FBlueprintNamespaceUtilities::GetDefaultImportsForObject(ExternalDependency, NamespacesToImport);
 							}
 						}
 
@@ -645,7 +641,7 @@ void SBlueprintActionMenu::OnGetNamespacesToExcludeFromImportMenu(TSet<FString>&
 
 	for (const UBlueprint* Blueprint : MenuContext.Blueprints)
 	{
-		FBlueprintNamespaceUtilities::GetDefaultImportsForBlueprint(Blueprint, OutNamespacesToExclude);
+		FBlueprintNamespaceUtilities::GetDefaultImportsForObject(Blueprint, OutNamespacesToExclude);
 		OutNamespacesToExclude.Append(Blueprint->ImportedNamespaces);
 	}
 }
diff --git a/Engine/Source/Editor/Kismet/Private/SBlueprintPalette.cpp b/Engine/Source/Editor/Kismet/Private/SBlueprintPalette.cpp
--- a/Engine/Source/Editor/Kismet/Private/SBlueprintPalette.cpp
+++ b/Engine/Source/Editor/Kismet/Private/SBlueprintPalette.cpp
@@ -844,15 +844,13 @@ private:
 				ActionPtr.Pin()->ChangeVariableType(InNewPinType);
 			}
 
-			// If the underlying type object's namespace is not imported, auto-import it now into the current editor context.
+			// Auto-import the underlying type object's default namespace set into the current editor context.
 			const UObject* PinSubCategoryObject = InNewPinType.PinSubCategoryObject.Get();
-			if (PinSubCategoryObject && BlueprintEditor.IsValid() && BlueprintEditor->IsNonImportedObject(PinSubCategoryObject))
+			if (PinSubCategoryObject && BlueprintEditor.IsValid())
 			{
-				FString Namespace = FBlueprintNamespaceUtilities::GetObjectNamespace(PinSubCategoryObject);
-				if (!Namespace.IsEmpty() && BlueprintEditor.IsValid())
-				{
-					BlueprintEditor->ImportNamespace(Namespace);
-				}
+				FBlueprintEditor::FImportNamespaceExParameters Params;
+				FBlueprintNamespaceUtilities::GetDefaultImportsForObject(PinSubCategoryObject, Params.NamespacesToImport);
+				BlueprintEditor->ImportNamespaceEx(Params);
 			}
 		}
 
diff --git a/Engine/Source/Editor/Kismet/Public/BlueprintNamespaceUtilities.h b/Engine/Source/Editor/Kismet/Public/BlueprintNamespaceUtilities.h
--- a/Engine/Source/Editor/Kismet/Public/BlueprintNamespaceUtilities.h
+++ b/Engine/Source/Editor/Kismet/Public/BlueprintNamespaceUtilities.h
@@ -73,7 +73,19 @@ public:
 	 * @param InBlueprint	A reference to a loaded Blueprint asset.
 	 * @param OutNamespaces	Zero or more unique namespace identifier(s) representing the Blueprint asset's default namespace set.
 	 */
-	static void GetDefaultImportsForBlueprint(const UBlueprint* InBlueprint, TSet<FString>& OutNamespaces);
+	UE_DEPRECATED(5.1, "Please use GetDefaultImportsForObject instead")
+	inline static void GetDefaultImportsForBlueprint(const UBlueprint* InBlueprint, TSet<FString>& OutNamespaces)
+	{
+		GetDefaultImportsForObject(InBlueprint, OutNamespaces);
+	}
+
+	/**
+	 * Gathers the set of default namespaces that are implicitly imported by a given object's type.
+	 *
+	 * @param InObject		A reference to a loaded object or type.
+	 * @param OutNamespaces	Zero or more unique namespace identifier(s) representing the type's default namespace set.
+	 */
+	static void GetDefaultImportsForObject(const UObject* InObject, TSet<FString>& OutNamespaces);
 
 	/**
 	 * Sets the default Blueprint namespace type that objects/assets should use when not explicitly assigned.
