CIS non-unity compile fix (202729).

diff --git a/Engine/Source/Editor/Kismet/Private/BlueprintNamespaceUtilities.cpp b/Engine/Source/Editor/Kismet/Private/BlueprintNamespaceUtilities.cpp
--- a/Engine/Source/Editor/Kismet/Private/BlueprintNamespaceUtilities.cpp
+++ b/Engine/Source/Editor/Kismet/Private/BlueprintNamespaceUtilities.cpp
@@ -229,6 +229,11 @@ void FBlueprintNamespaceUtilities::GetSharedGlobalImports(TSet<FString>& OutName
 	OutNamespaces.Remove(FString());
 }
 
+void FBlueprintNamespaceUtilities::GetDefaultImportsForBlueprint(const UBlueprint* InBlueprint, TSet<FString>& OutNamespaces)
+{
+	GetDefaultImportsForObject(InBlueprint, OutNamespaces);
+}
+
 void FBlueprintNamespaceUtilities::GetDefaultImportsForObject(const UObject* InObject, TSet<FString>& OutNamespaces)
 {
 	if (!InObject)
diff --git a/Engine/Source/Editor/Kismet/Public/BlueprintNamespaceUtilities.h b/Engine/Source/Editor/Kismet/Public/BlueprintNamespaceUtilities.h
--- a/Engine/Source/Editor/Kismet/Public/BlueprintNamespaceUtilities.h
+++ b/Engine/Source/Editor/Kismet/Public/BlueprintNamespaceUtilities.h
@@ -74,10 +74,7 @@ public:
 	 * @param OutNamespaces	Zero or more unique namespace identifier(s) representing the Blueprint asset's default namespace set.
 	 */
 	UE_DEPRECATED(5.1, "Please use GetDefaultImportsForObject instead")
-	inline static void GetDefaultImportsForBlueprint(const UBlueprint* InBlueprint, TSet<FString>& OutNamespaces)
-	{
-		GetDefaultImportsForObject(InBlueprint, OutNamespaces);
-	}
+	static void GetDefaultImportsForBlueprint(const UBlueprint* InBlueprint, TSet<FString>& OutNamespaces);
 
 	/**
 	 * Gathers the set of default namespaces that are implicitly imported by a given object's type.
