Variables marked EditInstanceOnly inside UStruct variables are no longer visible in the blueprint details view.

The code was only checking the top-level property in the FPropertyAndParent chain, rather than all of them.

diff --git a/Engine/Source/Editor/Kismet/Private/SKismetInspector.cpp b/Engine/Source/Editor/Kismet/Private/SKismetInspector.cpp
--- a/Engine/Source/Editor/Kismet/Private/SKismetInspector.cpp
+++ b/Engine/Source/Editor/Kismet/Private/SKismetInspector.cpp
@@ -789,42 +789,67 @@ bool SKismetInspector::IsPropertyVisible( const FPropertyAndParent& PropertyAndP
 {
 	const FProperty& Property = PropertyAndParent.Property;
 
-
 	// If we are in 'instance preview' - hide anything marked 'disabled edit on instance'
 	if ((ECheckBoxState::Checked == PublicViewState) && Property.HasAnyPropertyFlags(CPF_DisableEditOnInstance))
 	{
 		return false;
 	}
 
-	bool bEditOnTemplateDisabled = Property.HasAnyPropertyFlags(CPF_DisableEditOnTemplate);
-	if (bEditOnTemplateDisabled)
+	// Only hide EditInstanceOnly properties if we are editing a CDO/archetype
+	bool bIsEditingTemplate = true;
+	for (const TWeakObjectPtr<UObject>& SelectedObject : SelectedObjects)
 	{
-		// Only hide properties if we are editing a CDO/archetype
-		for (const TWeakObjectPtr<UObject>& SelectedObject : SelectedObjects)
+		UObject* Object = SelectedObject.Get();
+		if (!Object->IsTemplate())
 		{
-			UObject* Object = SelectedObject.Get();
-			if (!Object->IsTemplate())
-			{
-				bEditOnTemplateDisabled = false;
-				break;
-			}
+			bIsEditingTemplate = false;
+			break;
 		}
 	}
-	
-	if(const UClass* OwningClass = Property.GetOwner<UClass>())
+
+	if (bIsEditingTemplate)
 	{
-		const UBlueprint* BP = BlueprintEditorPtr.IsValid() ? BlueprintEditorPtr.Pin()->GetBlueprintObj() : nullptr;
-		const bool VariableAddedInCurentBlueprint = (OwningClass->ClassGeneratedBy == BP);
-		// If we did not add this var, hide it!
-		if(!VariableAddedInCurentBlueprint)
+		// check if the property (or any of its parent properties) was added by this blueprint
+		// this is necessary because of Instanced objects, which will have a different owning class yet are conceptually contained in this blueprint
+		bool bVariableAddedInCurrentBlueprint = false;
+		TSharedPtr<FBlueprintEditor> BlueprintEditor = BlueprintEditorPtr.Pin();
+		const UBlueprint* Blueprint = BlueprintEditor.IsValid() ? BlueprintEditor->GetBlueprintObj() : nullptr;
+
+		auto WasAddedInThisBlueprint = [Blueprint](const FProperty* Property)
 		{
-			if (bEditOnTemplateDisabled || Property.GetBoolMetaData(FBlueprintMetadata::MD_Private))
+			if (const UClass* OwningClass = Property->GetOwnerClass())
+			{
+				return OwningClass == Blueprint->ParentClass ||
+					OwningClass->ClassGeneratedBy == Blueprint;
+			}
+			return false;
+		};
+		
+		bVariableAddedInCurrentBlueprint |= WasAddedInThisBlueprint(&Property);
+
+		for (const FProperty* Parent : PropertyAndParent.ParentProperties)
+		{
+			bVariableAddedInCurrentBlueprint |= WasAddedInThisBlueprint(Parent);
+		}
+
+		// if this property wasn't added in this blueprint, we want to filter it out if it (or any of its parents) are marked EditInstanceOnly or private
+		if (!bVariableAddedInCurrentBlueprint)
+		{
+			if (Property.HasAnyPropertyFlags(CPF_DisableEditOnTemplate) || Property.GetBoolMetaData(FBlueprintMetadata::MD_Private))
 			{
 				return false;
 			}
+
+			for (const FProperty* Parent : PropertyAndParent.ParentProperties)
+			{
+				if (Property.HasAnyPropertyFlags(CPF_DisableEditOnTemplate) || Parent->GetBoolMetaData(FBlueprintMetadata::MD_Private))
+				{
+					return false;
+				}
+			}
 		}
 	}
-
+	
 	// figure out if this Blueprint variable is an Actor variable
 	const FArrayProperty* ArrayProperty = CastField<const FArrayProperty>(&Property);
 	const FSetProperty* SetProperty = CastField<const FSetProperty>(&Property);
@@ -834,7 +859,7 @@ bool SKismetInspector::IsPropertyVisible( const FPropertyAndParent& PropertyAndP
 	const FObjectPropertyBase* ObjectProperty = CastField<const FObjectPropertyBase>(TestProperty);
 	bool bIsActorProperty = (ObjectProperty != nullptr && ObjectProperty->PropertyClass->IsChildOf(AActor::StaticClass()));
 
-	if (bEditOnTemplateDisabled && bIsActorProperty)
+	if (bIsEditingTemplate && Property.HasAnyPropertyFlags(CPF_DisableEditOnTemplate) && bIsActorProperty)
 	{
 		// Actor variables can't have default values (because Blueprint templates are library elements that can 
 		// bridge multiple levels and different levels might not have the actor that the default is referencing).
@@ -854,7 +879,7 @@ bool SKismetInspector::IsPropertyVisible( const FPropertyAndParent& PropertyAndP
 		// If the current property is selected, it is visible.
 		return true;
 	}
-	else if ( PropertyAndParent.ParentProperties.Num() > 0 && SelectedObjectProperties.Num() > 0 )
+	else if (PropertyAndParent.ParentProperties.Num() > 0 && SelectedObjectProperties.Num() > 0)
 	{
 		if (IsAnyParentOrContainerSelected(PropertyAndParent))
 		{
@@ -875,7 +900,6 @@ bool SKismetInspector::IsPropertyVisible( const FPropertyAndParent& PropertyAndP
 		}
 	}
 
-
 	return SelectedObjectProperties.Num() == 0;
 }
 
