UE-160858 - Fix truncation warnings in all Kismet modules

Additionally, includes other misc fixes (eg: NULL->nullptr, east const->west const, various whitespace issues).

diff --git a/Engine/Source/Editor/Kismet/Private/BlueprintActionMenuItem.cpp b/Engine/Source/Editor/Kismet/Private/BlueprintActionMenuItem.cpp
--- a/Engine/Source/Editor/Kismet/Private/BlueprintActionMenuItem.cpp
+++ b/Engine/Source/Editor/Kismet/Private/BlueprintActionMenuItem.cpp
@@ -258,9 +258,9 @@ UEdGraphNode* FBlueprintActionMenuItem::PerformAction(UEdGraph* ParentGraph, UEd
 		{
 			UEdGraphNode* FromNode = FromPin->GetOwningNode();
 			check(FromNode != nullptr);
-			float const FromNodeX = FromNode->NodePosX;
+			const double FromNodeX = FromNode->NodePosX;
 
-			static const float MinNodeDistance = 60.f; // min distance between spawned nodes (to keep them from overlapping)
+			static const double MinNodeDistance = 60.0; // min distance between spawned nodes (to keep them from overlapping)
 			if (MinNodeDistance > FMath::Abs(FromNodeX - Location.X))
 			{
 				ModifiedLocation.X = FromNodeX - MinNodeDistance;
@@ -272,7 +272,7 @@ UEdGraphNode* FBlueprintActionMenuItem::PerformAction(UEdGraph* ParentGraph, UEd
 	}
 
 	TSet<const UEdGraphNode*> NodesToFocus;
-	int32 const PreSpawnNodeCount = ParentGraph->Nodes.Num();
+	const int32 PreSpawnNodeCount = ParentGraph->Nodes.Num();
 
 	UEdGraphNode* LastSpawnedNode = nullptr;
 	auto BoundObjIt = Bindings.CreateConstIterator();
@@ -287,7 +287,7 @@ UEdGraphNode* FBlueprintActionMenuItem::PerformAction(UEdGraph* ParentGraph, UEd
 			}
 		}
 
-		int32 const PreInvokeNodeCount = ParentGraph->Nodes.Num();
+		const int32 PreInvokeNodeCount = ParentGraph->Nodes.Num();
 
 		bool bNewNode = false;
 		LastSpawnedNode = InvokeAction(Action, ParentGraph, ModifiedLocation, BindingsSubset, /*out*/ bNewNode);
@@ -316,7 +316,7 @@ UEdGraphNode* FBlueprintActionMenuItem::PerformAction(UEdGraph* ParentGraph, UEd
 
 	if (bSelectNewNode)
 	{
-		int32 const PostSpawnCount = ParentGraph->Nodes.Num();
+		const int32 PostSpawnCount = ParentGraph->Nodes.Num();
 		for (int32 NodeIndex = PreSpawnNodeCount; NodeIndex < PostSpawnCount; ++NodeIndex)
 		{
 			NodesToFocus.Add(ParentGraph->Nodes[NodeIndex]);
diff --git a/Engine/Source/Editor/Kismet/Private/BlueprintDetailsCustomization.cpp b/Engine/Source/Editor/Kismet/Private/BlueprintDetailsCustomization.cpp
--- a/Engine/Source/Editor/Kismet/Private/BlueprintDetailsCustomization.cpp
+++ b/Engine/Source/Editor/Kismet/Private/BlueprintDetailsCustomization.cpp
@@ -159,7 +159,7 @@ void FBlueprintDetails::AddEventsCategory(IDetailLayoutBuilder& DetailBuilder, F
 						+ SHorizontalBox::Slot()
 						.HAlign(HAlign_Left)
 						.VAlign(VAlign_Center)
-						.Padding(0)
+						.Padding(0.0f)
 						[
 							SNew(SButton)
 							.ContentPadding(FMargin(3.0, 2.0))
@@ -578,7 +578,7 @@ void FBlueprintVarActionDetails::CustomizeDetails( IDetailLayoutBuilder& DetailL
 	.ValueContent()
 	[
 		SAssignNew(NewComboButton, SComboButton)
-		.ContentPadding(FMargin(0,0,5,0))
+		.ContentPadding(FMargin(0.0f, 0.0f, 5.0f, 0.0f))
 		.IsEnabled(this, &FBlueprintVarActionDetails::GetVariableCategoryChangeEnabled)
 		.ToolTip(CategoryTooltip)
 		.ButtonContent()
@@ -3148,7 +3148,7 @@ void FBlueprintGraphArgumentLayout::GenerateHeaderRowContent( FDetailWidgetRow&
 		SNew(SHorizontalBox)
 		+SHorizontalBox::Slot()
 		.VAlign(VAlign_Center)
-		.Padding(0)
+		.Padding(0.0f)
 		.FillWidth(1.0f)
 		[
 			SNew(SPinTypeSelector, FGetPinTypeTree::CreateUObject(K2Schema, &UEdGraphSchema_K2::GetVariableTypeTree))
@@ -3649,7 +3649,7 @@ void FBlueprintGraphActionDetails::CustomizeDetails( IDetailLayoutBuilder& Detai
 				.ValueContent()
 				[
 					SAssignNew(AccessSpecifierComboButton, SComboButton)
-					.ContentPadding(0)
+					.ContentPadding(0.0f)
 					.ButtonContent()
 					[
 						SNew(STextBlock)
@@ -3855,7 +3855,7 @@ void FBlueprintGraphActionDetails::CustomizeDetails( IDetailLayoutBuilder& Detai
 				+SVerticalBox::Slot()
 				[
 					SNew(SComboButton)
-						.ContentPadding(0)
+						.ContentPadding(0.0f)
 						.IsEnabled_Static(&LocalCustomEventUtils::IsNotCustomEventOverride, FunctionEntryNodePtr)
 						.ToolTipText_Static(&LocalCustomEventUtils::GetDisabledTooltip, FunctionEntryNodePtr)
 						.ButtonContent()
diff --git a/Engine/Source/Editor/Kismet/Private/BlueprintEditor.cpp b/Engine/Source/Editor/Kismet/Private/BlueprintEditor.cpp
--- a/Engine/Source/Editor/Kismet/Private/BlueprintEditor.cpp
+++ b/Engine/Source/Editor/Kismet/Private/BlueprintEditor.cpp
@@ -6175,8 +6175,8 @@ void FBlueprintEditor::MoveNodesToAveragePos(TSet<UEdGraphNode*>& AverageNodes,
 
 		for (UEdGraphNode* ExpandedNode : AverageNodes)
 		{
-			ExpandedNode->NodePosX = (ExpandedNode->NodePosX - AvgNodePosition.X) + SourcePos.X;
-			ExpandedNode->NodePosY = (ExpandedNode->NodePosY - AvgNodePosition.Y) + SourcePos.Y;
+			ExpandedNode->NodePosX = static_cast<int32>((ExpandedNode->NodePosX - AvgNodePosition.X) + SourcePos.X);
+			ExpandedNode->NodePosY = static_cast<int32>((ExpandedNode->NodePosY - AvgNodePosition.Y) + SourcePos.Y);
 
 			ExpandedNode->SnapToGrid(SNodePanel::GetSnapGridSize());
 
@@ -6340,8 +6340,8 @@ void FBlueprintEditor::ConvertFunctionToEvent(UK2Node_FunctionEntry* SelectedCal
 			check(NewNode != nullptr);
 			NewNode->CreateNewGuid();
 
-			NewNode->NodePosX = SpawnPos.X;
-			NewNode->NodePosY = SpawnPos.Y;
+			NewNode->NodePosX = static_cast<int32>(SpawnPos.X);
+			NewNode->NodePosY = static_cast<int32>(SpawnPos.Y);
 
 			NewNode->SetFlags(RF_Transactional);
 			NewNode->AllocateDefaultPins();
@@ -7231,7 +7231,7 @@ private:
 		return bResult;
 	}
 
-	void InitializeNewNode(UK2Node* NewNode, UK2Node* OldNode, float NodePosX = 0.0f, float NodePosY = 0.0f)
+	void InitializeNewNode(UK2Node* NewNode, UK2Node* OldNode, int32 NodePosX = 0, int32 NodePosY = 0)
 	{	
 		NewNode->NodePosX = OldNode ? OldNode->NodePosX : NodePosX;
 		NewNode->NodePosY = OldNode ? OldNode->NodePosY : NodePosY;
@@ -7264,7 +7264,7 @@ private:
 			check(NewTarget);
 			NewTarget->SetFromProperty(Property, true, Property->GetOwnerClass());
 			AddedTargets.Add(NewTarget);
-			const float AutoNodeOffsetX = 160.0f;
+			const int32 AutoNodeOffsetX = 160;
 			InitializeNewNode(NewTarget, OldTarget, OldCall->NodePosX - AutoNodeOffsetX, OldCall->NodePosY);
 		}
 
@@ -7389,8 +7389,8 @@ void FBlueprintEditor::PasteNodesHere(class UEdGraph* DestinationGraph, const FV
 			UEdGraphNode* Node = *It;
 			FocusedGraphEd->SetNodeSelection(Node, true);
 
-			Node->NodePosX = (Node->NodePosX - AvgNodePosition.X) + GraphLocation.X;
-			Node->NodePosY = (Node->NodePosY - AvgNodePosition.Y) + GraphLocation.Y;
+			Node->NodePosX = static_cast<int32>((Node->NodePosX - AvgNodePosition.X) + GraphLocation.X);
+			Node->NodePosY = static_cast<int32>((Node->NodePosY - AvgNodePosition.Y) + GraphLocation.Y);
 
 			Node->SnapToGrid(SNodePanel::GetSnapGridSize());
 
@@ -8362,12 +8362,12 @@ void FBlueprintEditor::CollapseNodesIntoGraph(UEdGraphNode* InGatewayNode, UK2No
 	const UEdGraphSchema_K2* K2Schema = GetDefault<UEdGraphSchema_K2>();
 
 	// Keep track of the statistics of the node positions so the new nodes can be located reasonably well
-	float SumNodeX = 0.0f;
-	float SumNodeY = 0.0f;
-	float MinNodeX = 1e9f;
-	float MinNodeY = 1e9f;
-	float MaxNodeX = -1e9f;
-	float MaxNodeY = -1e9f;
+	int32 SumNodeX = 0;
+	int32 SumNodeY = 0;
+	int32 MinNodeX = std::numeric_limits<int32>::max();
+	int32 MinNodeY = std::numeric_limits<int32>::max();
+	int32 MaxNodeX = std::numeric_limits<int32>::min();
+	int32 MaxNodeY = std::numeric_limits<int32>::min();
 
 	UEdGraphNode* InterfaceTemplateNode = nullptr;
 
@@ -8376,7 +8376,7 @@ void FBlueprintEditor::CollapseNodesIntoGraph(UEdGraphNode* InGatewayNode, UK2No
 	bool bDiscardReturnNode = true;
 
 	// For collapsing to functions can use a single event as a template for the function. This event MUST be deleted at the end, and the pins pre-generated. 
-	if(InGatewayNode->GetClass() == UK2Node_CallFunction::StaticClass())
+	if (InGatewayNode->GetClass() == UK2Node_CallFunction::StaticClass())
 	{
 		for (UEdGraphNode* Node : InCollapsableNodes)
 		{
@@ -8411,10 +8411,10 @@ void FBlueprintEditor::CollapseNodesIntoGraph(UEdGraphNode* InGatewayNode, UK2No
 		// Update stats
 		SumNodeX += Node->NodePosX;
 		SumNodeY += Node->NodePosY;
-		MinNodeX = FMath::Min<float>(MinNodeX, Node->NodePosX);
-		MinNodeY = FMath::Min<float>(MinNodeY, Node->NodePosY);
-		MaxNodeX = FMath::Max<float>(MaxNodeX, Node->NodePosX);
-		MaxNodeY = FMath::Max<float>(MaxNodeY, Node->NodePosY);
+		MinNodeX = FMath::Min<int32>(MinNodeX, Node->NodePosX);
+		MinNodeY = FMath::Min<int32>(MinNodeY, Node->NodePosY);
+		MaxNodeX = FMath::Max<int32>(MaxNodeX, Node->NodePosX);
+		MaxNodeY = FMath::Max<int32>(MaxNodeY, Node->NodePosY);
 
 		// Move the node over
 		InSourceGraph->Nodes.Remove(Node);
@@ -8422,7 +8422,7 @@ void FBlueprintEditor::CollapseNodesIntoGraph(UEdGraphNode* InGatewayNode, UK2No
 		Node->Rename(/*NewName=*/ nullptr, /*NewOuter=*/ InDestinationGraph);
 
 		// Move the sub-graph to the new graph
-		if(UK2Node_Composite* Composite = Cast<UK2Node_Composite>(Node))
+		if (UK2Node_Composite* Composite = Cast<UK2Node_Composite>(Node))
 		{
 			InSourceGraph->SubGraphs.Remove(Composite->BoundGraph);
 			InDestinationGraph->SubGraphs.Add(Composite->BoundGraph);
@@ -8477,7 +8477,7 @@ void FBlueprintEditor::CollapseNodesIntoGraph(UEdGraphNode* InGatewayNode, UK2No
 				UEdGraphPin* RemotePortPin = nullptr;
 
 				// Function graphs have a single exec path through them, so only one exec pin for input and another for output. In this fashion, they must not be handled by name.
-				if(InGatewayNode->GetClass() == UK2Node_CallFunction::StaticClass() && LocalPin->PinType.PinCategory == UEdGraphSchema_K2::PC_Exec)
+				if (InGatewayNode->GetClass() == UK2Node_CallFunction::StaticClass() && LocalPin->PinType.PinCategory == UEdGraphSchema_K2::PC_Exec)
 				{
 					LocalPortPin = LocalPort->Pins[0];
 					RemotePortPin = K2Schema->FindExecutionPin(*InGatewayNode, (LocalPortPin->Direction == EGPD_Input)? EGPD_Output : EGPD_Input);
@@ -8485,7 +8485,7 @@ void FBlueprintEditor::CollapseNodesIntoGraph(UEdGraphNode* InGatewayNode, UK2No
 				else
 				{
 					// If there is a custom event being used as a template, we must check to see if any connected pins have already been built
-					if(InterfaceTemplateNode && LocalPin->Direction == EGPD_Input)
+					if (InterfaceTemplateNode && LocalPin->Direction == EGPD_Input)
 					{
 						// Find the pin on the entry node, we will use that pin's name to find the pin on the remote port
 						UEdGraphPin* EntryNodePin = InEntryNode->FindPin(LocalPin->LinkedTo[0]->PinName);
@@ -8497,11 +8497,11 @@ void FBlueprintEditor::CollapseNodesIntoGraph(UEdGraphNode* InGatewayNode, UK2No
 						}
 					}
 
-					if(LocalPin->LinkedTo[0]->GetOwningNode() != InEntryNode)
+					if (LocalPin->LinkedTo[0]->GetOwningNode() != InEntryNode)
 					{
 						const FName UniquePortName = InGatewayNode->CreateUniquePinName(LocalPin->PinName);
 
-						if(!RemotePortPin && !LocalPortPin)
+						if (!RemotePortPin && !LocalPortPin)
 						{
 							if (LocalPin->Direction == EGPD_Output)
 							{
@@ -8514,7 +8514,7 @@ void FBlueprintEditor::CollapseNodesIntoGraph(UEdGraphNode* InGatewayNode, UK2No
 								PinType.bIsWeakPointer = false;
 							}
 							RemotePortPin = InGatewayNode->CreatePin(LocalPin->Direction, PinType, UniquePortName);
-							LocalPortPin = LocalPort->CreateUserDefinedPin(UniquePortName, PinType, (LocalPin->Direction == EGPD_Input)? EGPD_Output : EGPD_Input);
+							LocalPortPin = LocalPort->CreateUserDefinedPin(UniquePortName, PinType, (LocalPin->Direction == EGPD_Input) ? EGPD_Output : EGPD_Input);
 						}
 					}
 				}
@@ -8581,9 +8581,9 @@ void FBlueprintEditor::CollapseNodesIntoGraph(UEdGraphNode* InGatewayNode, UK2No
 	const int32 NumNodes = InCollapsableNodes.Num();
 
 	// Remove the template node if one was used for generating the function
-	if(InterfaceTemplateNode)
+	if (InterfaceTemplateNode)
 	{
-		if(NumNodes == 0)
+		if (NumNodes == 0)
 		{
 			SumNodeX = InterfaceTemplateNode->NodePosX;
 			SumNodeY = InterfaceTemplateNode->NodePosY;
@@ -8593,18 +8593,18 @@ void FBlueprintEditor::CollapseNodesIntoGraph(UEdGraphNode* InGatewayNode, UK2No
 	}
 
 	// Using the result pin, we will ensure that there is a path through the function by checking if it is connected. If it is not, link it to the entry node.
-	if(UEdGraphPin* ResultExecFunc = K2Schema->FindExecutionPin(*InResultNode, EGPD_Input))
+	if (UEdGraphPin* ResultExecFunc = K2Schema->FindExecutionPin(*InResultNode, EGPD_Input))
 	{
-		if(ResultExecFunc->LinkedTo.Num() == 0)
+		if (ResultExecFunc->LinkedTo.Num() == 0)
 		{
 			K2Schema->FindExecutionPin(*InEntryNode, EGPD_Output)->MakeLinkTo(K2Schema->FindExecutionPin(*InResultNode, EGPD_Input));
 		}
 	}
 
-	const float CenterX = NumNodes == 0? SumNodeX : SumNodeX / NumNodes;
-	const float CenterY = NumNodes == 0? SumNodeY : SumNodeY / NumNodes;
-	const float MinusOffsetX = 160.0f; //@TODO: Random magic numbers
-	const float PlusOffsetX = 300.0f;
+	const int32 CenterX = (NumNodes == 0) ? SumNodeX : SumNodeX / NumNodes;
+	const int32 CenterY = (NumNodes == 0) ? SumNodeY : SumNodeY / NumNodes;
+	const int32 MinusOffsetX = 160; //@TODO: Random magic numbers
+	const int32 PlusOffsetX = 300;
 
 	// Put the gateway node at the center of the empty space in the old graph
 	InGatewayNode->NodePosX = CenterX;
@@ -8613,7 +8613,7 @@ void FBlueprintEditor::CollapseNodesIntoGraph(UEdGraphNode* InGatewayNode, UK2No
 
 	// Put the entry and exit nodes on either side of the nodes in the new graph
 	//@TODO: Should we recenter the whole ensemble?
-	if(NumNodes != 0)
+	if (NumNodes != 0)
 	{
 		InEntryNode->NodePosX = MinNodeX - MinusOffsetX;
 		InEntryNode->NodePosY = CenterY;
diff --git a/Engine/Source/Editor/Kismet/Private/BlueprintWarningsConfigurationPanel.cpp b/Engine/Source/Editor/Kismet/Private/BlueprintWarningsConfigurationPanel.cpp
--- a/Engine/Source/Editor/Kismet/Private/BlueprintWarningsConfigurationPanel.cpp
+++ b/Engine/Source/Editor/Kismet/Private/BlueprintWarningsConfigurationPanel.cpp
@@ -190,7 +190,7 @@ void SBlueprintWarningsConfigurationPanel::Construct(const FArguments& InArgs)
 			.Padding(0.0f, 16.0f, 0.0f, 0.0f)
 		[
 			SAssignNew(Label, SBorder)
-				.Padding(3)
+				.Padding(3.0f)
 				.BorderImage(FAppStyle::GetBrush("DetailsView.CategoryTop"))
 				.BorderBackgroundColor(FLinearColor(.6, .6, .6, 1.0f))
 			[
@@ -269,25 +269,27 @@ void SBlueprintWarningsConfigurationPanel::UpdateSelectedWarningBehaviors(EBluep
 		{
 			// update config data:
 			FName CurrentIdentifer = Declaration->WarningIdentifier;
-			FBlueprintWarningSettings* WarningEntry = RuntimeSettings->WarningSettings.FindByPredicate(
+			int WarningEntryIndex = RuntimeSettings->WarningSettings.IndexOfByPredicate(
 				[CurrentIdentifer](const FBlueprintWarningSettings& Entry) -> bool
 				{
 					return Entry.WarningIdentifier == CurrentIdentifer;
 				}
 			);
+
 			switch(NewBehavior)
 			{
 			case EBlueprintWarningBehavior::Warn:
-				if( WarningEntry )
+				if (WarningEntryIndex >= 0)
 				{
-					RuntimeSettings->WarningSettings.RemoveAtSwap(WarningEntry - &(RuntimeSettings->WarningSettings[0]));
+					RuntimeSettings->WarningSettings.RemoveAtSwap(WarningEntryIndex);
 				}
 				break;
 			case EBlueprintWarningBehavior::Error:
 			case EBlueprintWarningBehavior::Suppress:
-				if( WarningEntry )
+				if (WarningEntryIndex >= 0)
 				{
-					WarningEntry->WarningBehavior = NewBehavior;
+					FBlueprintWarningSettings& WarningEntry = RuntimeSettings->WarningSettings[WarningEntryIndex];
+					WarningEntry.WarningBehavior = NewBehavior;
 				}
 				else
 				{
@@ -295,9 +297,7 @@ void SBlueprintWarningsConfigurationPanel::UpdateSelectedWarningBehaviors(EBluep
 					NewEntry.WarningIdentifier = CurrentIdentifer;
 					NewEntry.WarningDescription = Declaration->WarningDescription;
 					NewEntry.WarningBehavior = NewBehavior;
-					RuntimeSettings->WarningSettings.Add(
-						NewEntry
-					);
+					RuntimeSettings->WarningSettings.Add(NewEntry);
 				}
 				break;
 			}
diff --git a/Engine/Source/Editor/Kismet/Private/CallStackViewer.cpp b/Engine/Source/Editor/Kismet/Private/CallStackViewer.cpp
--- a/Engine/Source/Editor/Kismet/Private/CallStackViewer.cpp
+++ b/Engine/Source/Editor/Kismet/Private/CallStackViewer.cpp
@@ -347,7 +347,7 @@ void SCallStackViewer::Construct(const FArguments& InArgs, TArray<TSharedRef<FCa
 	ChildSlot
 	[
 		SNew(SBorder)
-		.Padding(4)
+		.Padding(4.0f)
 		.BorderImage( FAppStyle::GetBrush("ToolPanel.GroupBorder") )
 		[
 			SNew(SOverlay)
@@ -538,12 +538,13 @@ void CallStackViewer::UpdateDisplayedCallstack(TArrayView<const FFrame* const> S
 				}
 			}
 
+			int32 ScriptOffset = UE_PTRDIFF_TO_INT32(StackNode->Code - StackNode->Node->Script.GetData() - 1);
 			Private_CallStackSource.Add(
 				MakeShared<FCallStackRow>(
 					StackNode->Object,
 					StackNode->Node->GetOuter()->GetFName(),
 					StackNode->Node->GetFName(),
-					StackNode->Code - StackNode->Node->Script.GetData() - 1,
+					ScriptOffset,
 					ECallstackLanguages::Blueprints,
 					ScopeDisplayName,
 					FunctionDisplayName
diff --git a/Engine/Source/Editor/Kismet/Private/FindInBlueprints.cpp b/Engine/Source/Editor/Kismet/Private/FindInBlueprints.cpp
--- a/Engine/Source/Editor/Kismet/Private/FindInBlueprints.cpp
+++ b/Engine/Source/Editor/Kismet/Private/FindInBlueprints.cpp
@@ -739,8 +739,8 @@ void SFindInBlueprints::ConditionallyAddCacheBar()
 				{
 					TSharedRef<SWidget> DisplayWidget = 
 						SNew(SBox)
-						.MaxDesiredHeight(512)
-						.MaxDesiredWidth(512)
+						.MaxDesiredHeight(512.0f)
+						.MaxDesiredWidth(512.0f)
 						.Content()
 						[
 							SNew(SBorder)
diff --git a/Engine/Source/Editor/Kismet/Private/FormatTextDetails.cpp b/Engine/Source/Editor/Kismet/Private/FormatTextDetails.cpp
--- a/Engine/Source/Editor/Kismet/Private/FormatTextDetails.cpp
+++ b/Engine/Source/Editor/Kismet/Private/FormatTextDetails.cpp
@@ -137,7 +137,7 @@ void FFormatTextArgumentLayout::GenerateHeaderRowContent( FDetailWidgetRow& Node
 			.Padding(2, 0)
 			[
 				SNew(SButton)
-				.ContentPadding(0)
+				.ContentPadding(0.0f)
 				.OnClicked(this, &FFormatTextArgumentLayout::OnMoveArgumentUp)
 				.IsEnabled( bIsMoveUpEnabled )
 				[
@@ -150,7 +150,7 @@ void FFormatTextArgumentLayout::GenerateHeaderRowContent( FDetailWidgetRow& Node
 			.Padding(2, 0)
 			[
 				SNew(SButton)
-				.ContentPadding(0)
+				.ContentPadding(0.0f)
 				.OnClicked(this, &FFormatTextArgumentLayout::OnMoveArgumentDown)
 				.IsEnabled( bIsMoveDownEnabled )
 				[
diff --git a/Engine/Source/Editor/Kismet/Private/SBlueprintActionMenu.cpp b/Engine/Source/Editor/Kismet/Private/SBlueprintActionMenu.cpp
--- a/Engine/Source/Editor/Kismet/Private/SBlueprintActionMenu.cpp
+++ b/Engine/Source/Editor/Kismet/Private/SBlueprintActionMenu.cpp
@@ -301,12 +301,12 @@ void SBlueprintActionMenu::Construct( const FArguments& InArgs, TSharedPtr<FBlue
 	// Build the widget layout
 	SBorder::Construct( SBorder::FArguments()
 		.BorderImage( FAppStyle::GetBrush("Menu.Background") )
-		.Padding(5)
+		.Padding(5.0f)
 		[
 			// Achieving fixed width by nesting items within a fixed width box.
 			SNew(SBox)
-			.WidthOverride(400)
-			.HeightOverride(400)
+			.WidthOverride(400.0f)
+			.HeightOverride(400.0f)
 			[
 				SNew(SVerticalBox)
 
@@ -321,7 +321,7 @@ void SBlueprintActionMenu::Construct( const FArguments& InArgs, TSharedPtr<FBlue
 					+SHorizontalBox::Slot()
 					.AutoWidth()
 					.VAlign(VAlign_Center)
-					.Padding(0, 0, (ContextIcon != NULL) ? 5 : 0, 0)
+					.Padding(0.0f, 0.0f, (ContextIcon != nullptr) ? 5.0f : 0.0f, 0.0f)
 					[
 						SNew(SImage)
 						.ColorAndOpacity(TypeColor)
diff --git a/Engine/Source/Editor/Kismet/Private/SBlueprintContextTargetMenu.cpp b/Engine/Source/Editor/Kismet/Private/SBlueprintContextTargetMenu.cpp
--- a/Engine/Source/Editor/Kismet/Private/SBlueprintContextTargetMenu.cpp
+++ b/Engine/Source/Editor/Kismet/Private/SBlueprintContextTargetMenu.cpp
@@ -505,7 +505,7 @@ bool FContextMenuTargetProfile::LoadProfile()
 /**  */
 namespace BlueprintContextTargetMenuImpl
 {
-	static FText GetContextTargetDisplayName(UEnum* Enum, int32 EnumIndex)
+	static FText GetContextTargetDisplayName(const UEnum* Enum, int32 EnumIndex)
 	{
 		if (Enum != nullptr)
 		{
@@ -522,10 +522,10 @@ void SBlueprintContextTargetMenu::Construct(const FArguments& InArgs, const FBlu
 	OnTargetMaskChanged = InArgs._OnTargetMaskChanged;
 
 	FSlateFontInfo HeaderFontStyle = FAppStyle::GetFontStyle("BlueprintEditor.ActionMenu.ContextDescriptionFont");
-	HeaderFontStyle.Size -= 2.f;
-	FText const HeaderText = ContextMenuTargetProfileImpl::GetProfileDescription(MenuContext);
+	HeaderFontStyle.Size -= 2;
+	const FText HeaderText = ContextMenuTargetProfileImpl::GetProfileDescription(MenuContext);
 
-	FText const MenuToolTip = LOCTEXT("MenuToolTip", "Select whose functions/variables you want to see.\nNOTE: Unchecking everything is akin to 'SHOW EVERYTHING' (you're choosing to have NO target context and to not limit the scope)");
+	const FText MenuToolTip = LOCTEXT("MenuToolTip", "Select whose functions/variables you want to see.\nNOTE: Unchecking everything is akin to 'SHOW EVERYTHING' (you're choosing to have NO target context and to not limit the scope)");
 
 	TSharedRef<SWidget> CustomContentWidget = InArgs._CustomTargetContent.Widget;
 	if (CustomContentWidget != SNullWidget::NullWidget)
@@ -555,7 +555,7 @@ void SBlueprintContextTargetMenu::Construct(const FArguments& InArgs, const FBlu
 		.ToolTipText(MenuToolTip)
 		[
 			SNew(SBox)
-				.MinDesiredWidth(200)
+				.MinDesiredWidth(200.0f)
 				.ToolTipText(MenuToolTip)
 				.Padding(FMargin(0.f, 0.f, 0.f, 18.f))
 			[
@@ -602,8 +602,8 @@ void SBlueprintContextTargetMenu::Construct(const FArguments& InArgs, const FBlu
 		];
 	}
 
-	UEnum* const TargetEnum = StaticEnum<EContextTargetFlags::Type>();
-	uint32 const GreatestFlag = (EContextTargetFlags::ContextTargetFlagsEnd & ~1);
+	const UEnum* TargetEnum = StaticEnum<EContextTargetFlags::Type>();
+	const uint32 GreatestFlag = (EContextTargetFlags::ContextTargetFlagsEnd & ~1);
 
 	int32  ColIndex = 0;
 	for (int32 BitMaskOffset = 0; (1 << BitMaskOffset) <= GreatestFlag; ++BitMaskOffset)
@@ -621,7 +621,7 @@ void SBlueprintContextTargetMenu::Construct(const FArguments& InArgs, const FBlu
 			continue;
 		}
 		
-		FText const MenuName = BlueprintContextTargetMenuImpl::GetContextTargetDisplayName(TargetEnum, BitMaskOffset);
+		const FText MenuName = BlueprintContextTargetMenuImpl::GetContextTargetDisplayName(TargetEnum, BitMaskOffset);
 		const FContextMenuTargetProfile& ProfileRef = TargetProfile;
 
 		Columns[ColIndex]->AddSlot()
diff --git a/Engine/Source/Editor/Kismet/Private/SBlueprintLibraryPalette.cpp b/Engine/Source/Editor/Kismet/Private/SBlueprintLibraryPalette.cpp
--- a/Engine/Source/Editor/Kismet/Private/SBlueprintLibraryPalette.cpp
+++ b/Engine/Source/Editor/Kismet/Private/SBlueprintLibraryPalette.cpp
@@ -380,7 +380,7 @@ TSharedRef<SWidget> SBlueprintLibraryPalette::ConstructClassFilterDropdownConten
 		[
 			// achieving fixed width by nesting items within a fixed width box.
 			SNew(SBox)
-				.WidthOverride(350)
+				.WidthOverride(350.0f)
 			[
 				SNew(SVerticalBox)
 				
diff --git a/Engine/Source/Editor/Kismet/Private/SCSEditorViewportClient.cpp b/Engine/Source/Editor/Kismet/Private/SCSEditorViewportClient.cpp
--- a/Engine/Source/Editor/Kismet/Private/SCSEditorViewportClient.cpp
+++ b/Engine/Source/Editor/Kismet/Private/SCSEditorViewportClient.cpp
@@ -264,8 +264,8 @@ void FSCSEditorViewportClient::DrawCanvas( FViewport& InViewport, FSceneView& Vi
 
 		TGuardValue<bool> AutoRestore(GAllowActorScriptExecutionInEditor, true);
 
-		const int32 HalfX = 0.5f * Viewport->GetSizeXY().X;
-		const int32 HalfY = 0.5f * Viewport->GetSizeXY().Y;
+		const int32 HalfX = Viewport->GetSizeXY().X / 2;
+		const int32 HalfY = Viewport->GetSizeXY().Y / 2;
 
 		TArray<FSubobjectEditorTreeNodePtrType> SelectedNodes = BlueprintEditorPtr.Pin()->GetSelectedSubobjectEditorTreeNodes();
 
@@ -279,8 +279,8 @@ void FSCSEditorViewportClient::DrawCanvas( FViewport& InViewport, FSceneView& Vi
 				const FPlane Proj = View.Project(WidgetLocation);
 				if(Proj.W > 0.0f)
 				{
-					const int32 XPos = HalfX + (HalfX * Proj.X);
-					const int32 YPos = HalfY + (HalfY * (Proj.Y * -1));
+					const int32 XPos = static_cast<int32>(HalfX + (HalfX * Proj.X));
+					const int32 YPos = static_cast<int32>(HalfY + (HalfY * Proj.Y * -1.0));
 					DrawAngles(&Canvas, XPos, YPos, GetCurrentWidgetAxis(), GetWidgetMode(), GetWidgetCoordSystem().Rotator(), WidgetLocation);
 				}
 			}
@@ -781,15 +781,15 @@ void FSCSEditorViewportClient::ResetCamera()
 	}
 
 	// Clamp zoom to the actor's bounding sphere radius
-	float OrbitZoom = ThumbnailInfo->OrbitZoom;
-	if (PreviewActorBounds.SphereRadius + OrbitZoom < 0)
+	double OrbitZoom = ThumbnailInfo->OrbitZoom;
+	if (PreviewActorBounds.SphereRadius + OrbitZoom < 0.0)
 	{
 		OrbitZoom = -PreviewActorBounds.SphereRadius;
 	}
 
 	ToggleOrbitCamera(true);
 	{
-		float TargetDistance = PreviewActorBounds.SphereRadius;
+		double TargetDistance = PreviewActorBounds.SphereRadius;
 		if(TargetDistance <= 0.0f)
 		{
 			TargetDistance = AutoViewportOrbitCameraTranslate;
diff --git a/Engine/Source/Editor/Kismet/Private/SFilterableObjectList.cpp b/Engine/Source/Editor/Kismet/Private/SFilterableObjectList.cpp
--- a/Engine/Source/Editor/Kismet/Private/SFilterableObjectList.cpp
+++ b/Engine/Source/Editor/Kismet/Private/SFilterableObjectList.cpp
@@ -126,7 +126,7 @@ void SFilterableObjectList::InternalConstruct()
 		// The (possibly filtered) list of items
 		+SVerticalBox::Slot()
 			.FillHeight(1)
-			.Padding(2)
+			.Padding(2.0f)
 		[
 			SNew(SBorder)
 			[
diff --git a/Engine/Source/Editor/Kismet/Private/SGraphTitleBarAddNewBookmark.cpp b/Engine/Source/Editor/Kismet/Private/SGraphTitleBarAddNewBookmark.cpp
--- a/Engine/Source/Editor/Kismet/Private/SGraphTitleBarAddNewBookmark.cpp
+++ b/Engine/Source/Editor/Kismet/Private/SGraphTitleBarAddNewBookmark.cpp
@@ -30,7 +30,7 @@ void SGraphTitleBarAddNewBookmark::Construct(const FArguments& InArgs)
 	.MenuContent()
 	[
 		SNew(SBox)
-		.MinDesiredWidth(300)
+		.MinDesiredWidth(300.0f)
 		[
 			SNew(SVerticalBox)
 			+ SVerticalBox::Slot()
diff --git a/Engine/Source/Editor/Kismet/Private/SKismetInspector.cpp b/Engine/Source/Editor/Kismet/Private/SKismetInspector.cpp
--- a/Engine/Source/Editor/Kismet/Private/SKismetInspector.cpp
+++ b/Engine/Source/Editor/Kismet/Private/SKismetInspector.cpp
@@ -80,14 +80,14 @@ public:
 				.AutoWidth()
 				.HAlign(HAlign_Center)
 				.VAlign(VAlign_Center)
-				.Padding(2)
+				.Padding(2.0f)
 				[
 					SNew(SImage)
 					.Image(FAppStyle::Get().GetBrush("Icons.Warning"))
 				]
 				+ SHorizontalBox::Slot()
 					.VAlign(VAlign_Center)
-					.Padding(2)
+					.Padding(2.0f)
 					[
 						SNew(SRichTextBlock)
 						.DecoratorStyleSet(&FAppStyle::Get())
@@ -423,7 +423,7 @@ void SKismetInspector::Construct(const FArguments& InArgs)
 		.FillHeight(1.0f)
 		[
 			SAssignNew( ContextualEditingBorderWidget, SBorder )
-			.Padding(0)
+			.Padding(0.0f)
 			.BorderImage( FAppStyle::GetBrush("NoBorder") )
 		]
 	];
diff --git a/Engine/Source/Editor/Kismet/Private/SMyBlueprint.cpp b/Engine/Source/Editor/Kismet/Private/SMyBlueprint.cpp
--- a/Engine/Source/Editor/Kismet/Private/SMyBlueprint.cpp
+++ b/Engine/Source/Editor/Kismet/Private/SMyBlueprint.cpp
@@ -528,7 +528,7 @@ void SMyBlueprint::Construct(const FArguments& InArgs, TWeakPtr<FBlueprintEditor
 					.Padding(2, 0, 0, 0)
 					[
 						SNew(SComboButton)
-						.ContentPadding(0)
+						.ContentPadding(0.0f)
 						.ComboButtonStyle(&FAppStyle::Get().GetWidgetStyle<FComboButtonStyle>("SimpleComboButton"))
 						.HasDownArrow(false)
 						.AddMetaData<FTagMetaData>(FTagMetaData(TEXT("ViewOptions")))
@@ -794,7 +794,7 @@ TSharedRef<SWidget> SMyBlueprint::OnGetSectionWidget(TSharedRef<SWidget> RowWidg
 					.Visibility(this, &SMyBlueprint::OnGetSectionTextVisibility, WeakRowWidget, InSectionID)
 					.ForegroundColor(FAppStyle::GetSlateColor("DefaultForeground"))
 					.OnGetMenuContent(this, &SMyBlueprint::OnGetFunctionListMenu)
-					.ContentPadding(0)
+					.ContentPadding(0.0f)
 					.HasDownArrow(true)
 					.ButtonContent()
 					[
diff --git a/Engine/Source/Editor/Kismet/Private/SPinValueInspector.cpp b/Engine/Source/Editor/Kismet/Private/SPinValueInspector.cpp
--- a/Engine/Source/Editor/Kismet/Private/SPinValueInspector.cpp
+++ b/Engine/Source/Editor/Kismet/Private/SPinValueInspector.cpp
@@ -125,7 +125,7 @@ private:
 	TAttribute<TOptional<float>> MaxHeight;
 
 	// must be mutable to be changed in ComputeDesiredSize
-	mutable TOptional<float> CachedXSize;
+	mutable TOptional<double> CachedXSize;
 };
 
 void SPinValueInspector::Construct(const FArguments& InArgs)
diff --git a/Engine/Source/Editor/Kismet/Private/SSCSEditor.cpp b/Engine/Source/Editor/Kismet/Private/SSCSEditor.cpp
--- a/Engine/Source/Editor/Kismet/Private/SSCSEditor.cpp
+++ b/Engine/Source/Editor/Kismet/Private/SSCSEditor.cpp
@@ -1900,13 +1900,13 @@ void SSCS_RowWidget::AddToToolTipInfoBox(const TSharedRef<SVerticalBox>& InfoBox
 {
 	InfoBox->AddSlot()
 		.AutoHeight()
-		.Padding(0, 1)
+		.Padding(0.0f, 1.0f)
 		[
 			SNew(SHorizontalBox)
 
 			+ SHorizontalBox::Slot()
 			.AutoWidth()
-			.Padding(0, 0, 4, 0)
+			.Padding(0.0f, 0.0f, 4.0f, 0.0f)
 			[
 				SNew(STextBlock)
 				.TextStyle(FAppStyle::Get(), bImportant ? "SCSEditor.ComponentTooltip.ImportantLabel" : "SCSEditor.ComponentTooltip.Label")
@@ -1950,7 +1950,7 @@ TSharedRef<SToolTip> SSCS_RowWidget::CreateToolTipWidget() const
 				InfoBox->AddSlot()
 					.AutoHeight()
 					.HAlign(HAlign_Center)
-					.Padding(FMargin(0, 2, 0, 4))
+					.Padding(FMargin(0.0f, 2.0f, 0.0f, 4.0f))
 					[
 						SNew(STextBlock)
 						.TextStyle(FAppStyle::Get(), "SCSEditor.ComponentTooltip.ClassDescription")
@@ -1990,13 +1990,13 @@ TSharedRef<SToolTip> SSCS_RowWidget::CreateToolTipWidget() const
 
 	TSharedRef<SBorder> TooltipContent = SNew(SBorder)
 		.BorderImage(FAppStyle::GetBrush("NoBorder"))
-		.Padding(0)
+		.Padding(0.0f)
 		[
 			SNew(SVerticalBox)
 			
 			+SVerticalBox::Slot()
 			.AutoHeight()
-			.Padding(0, 0, 0, 4)
+			.Padding(0.0f, 0.0f, 0.0f, 4.0f)
 			[
 				SNew(SVerticalBox)
 
@@ -2008,7 +2008,7 @@ TSharedRef<SToolTip> SSCS_RowWidget::CreateToolTipWidget() const
 					+ SHorizontalBox::Slot()
 					.AutoWidth()
 					.VAlign(VAlign_Center)
-					.Padding(2)
+					.Padding(2.0f)
 					[
 						SNew(STextBlock)
 						.TextStyle(FAppStyle::Get(), "SCSEditor.ComponentTooltip.Title")
@@ -2022,7 +2022,7 @@ TSharedRef<SToolTip> SSCS_RowWidget::CreateToolTipWidget() const
 			[
 				SNew(SBorder)
 				.BorderImage(FAppStyle::GetBrush("NoBorder"))
-				.Padding(2)
+				.Padding(2.0f)
 				[
 					InfoBox
 				]
@@ -3637,13 +3637,13 @@ TSharedRef<SToolTip> SSCS_RowWidget_ActorRoot::CreateToolTipWidget() const
 
 	TSharedRef<SBorder> TooltipContent = SNew(SBorder)
 		.BorderImage(FAppStyle::GetBrush("NoBorder"))
-		.Padding(0)
+		.Padding(0.0f)
 		[
 			SNew(SVerticalBox)
 
 			+ SVerticalBox::Slot()
 			.AutoHeight()
-			.Padding(0, 0, 0, 4)
+			.Padding(0.0f, 0.0f, 0.0f, 4.0f)
 			[
 				SNew(SVerticalBox)
 
@@ -3655,7 +3655,7 @@ TSharedRef<SToolTip> SSCS_RowWidget_ActorRoot::CreateToolTipWidget() const
 					+ SHorizontalBox::Slot()
 					.AutoWidth()
 					.VAlign(VAlign_Center)
-					.Padding(4)
+					.Padding(4.0f)
 					[
 						SNew(STextBlock)
 						.TextStyle(FAppStyle::Get(), "SCSEditor.ComponentTooltip.Title")
@@ -3669,7 +3669,7 @@ TSharedRef<SToolTip> SSCS_RowWidget_ActorRoot::CreateToolTipWidget() const
 			[
 				SNew(SBorder)
 				.BorderImage(FAppStyle::GetBrush("NoBorder"))
-				.Padding(4)
+				.Padding(4.0f)
 				[
 					InfoBox
 				]
@@ -3992,7 +3992,7 @@ void SSCSEditor::Construct( const FArguments& InArgs )
 	ButtonBox = SNew(SHorizontalBox)
 	+ SHorizontalBox::Slot()
 	.VAlign(VAlign_Center)
-	.Padding(0,0,4,0)
+	.Padding(0.0f, 0.0f, 4.0f, 0.0f)
 	.AutoWidth()
 	[
 		SNew(SComponentClassCombo)
@@ -4004,7 +4004,7 @@ void SSCSEditor::Construct( const FArguments& InArgs )
 	]
 	+ SHorizontalBox::Slot()
 	.VAlign(VAlign_Center)
-	.Padding(0,0,4,0)
+	.Padding(0.0f, 0.0f, 4.0f, 0.0f)
 	.AutoWidth()
 	[
 		SAssignNew(ExtensionPanel, SExtensionPanel)
@@ -4016,7 +4016,7 @@ void SSCSEditor::Construct( const FArguments& InArgs )
 
 	+ SHorizontalBox::Slot()
 	.VAlign(VAlign_Center)
-	.Padding(0,0,4,0)
+	.Padding(0.0f, 0.0f, 4.0f, 0.0f)
 	.AutoWidth()
 	[
 		SNew(SPositiveActionButton)
diff --git a/Engine/Source/Editor/Kismet/Private/STimelineEditor.cpp b/Engine/Source/Editor/Kismet/Private/STimelineEditor.cpp
--- a/Engine/Source/Editor/Kismet/Private/STimelineEditor.cpp
+++ b/Engine/Source/Editor/Kismet/Private/STimelineEditor.cpp
@@ -125,8 +125,8 @@ void STimelineEdTrack::Construct(const FArguments& InArgs, TSharedPtr<FTimelineE
 	check(TimelineObj); // We shouldn't have any tracks if there is no track object!
 
 	// Get a pointer to the track this widget is for
-	CurveBasePtr = NULL;
-	FTTTrackBase* TrackBase = NULL;
+	CurveBasePtr = nullptr;
+	FTTTrackBase* TrackBase = nullptr;
 	bool bDrawCurve = true;
 
 	FTTTrackId TrackId = TimelineObj->GetDisplayTrackId(InTrack->DisplayIndex);
@@ -231,7 +231,7 @@ void STimelineEdTrack::Construct(const FArguments& InArgs, TSharedPtr<FTimelineE
 					// External Curve Label
 					+SVerticalBox::Slot()
 					.AutoHeight()
-					.Padding(2)
+					.Padding(2.0f)
 					[
 						SNew(STextBlock)
 						.Text(LOCTEXT("ExternalCurveLabel", "External Curve"))
@@ -451,7 +451,7 @@ void STimelineEdTrack::SwitchToExternalCurve(UCurveBase* AssetCurvePtr)
 		FTTTrackId TrackId = TimelineObj->GetDisplayTrackId(Track->DisplayIndex);
 		FTTTrackBase::ETrackType TrackType = (FTTTrackBase::ETrackType)TrackId.TrackType;
 
-		FTTTrackBase* TrackBase = NULL;
+		FTTTrackBase* TrackBase = nullptr;
 		if(TrackType == FTTTrackBase::TT_Event)
 		{
 			if(AssetCurvePtr->IsA(UCurveFloat::StaticClass()))
@@ -526,8 +526,8 @@ void STimelineEdTrack::UseInternalCurve( )
 		FTTTrackId TrackId = TimelineObj->GetDisplayTrackId(Track->DisplayIndex);
 		FTTTrackBase::ETrackType TrackType = (FTTTrackBase::ETrackType)TrackId.TrackType;
 
-		FTTTrackBase* TrackBase = NULL;
-		UCurveBase* CurveBase = NULL;
+		FTTTrackBase* TrackBase = nullptr;
+		UCurveBase* CurveBase = nullptr;
 
 		if(TrackType == FTTTrackBase::TT_Event)
 		{
@@ -644,7 +644,7 @@ FString STimelineEdTrack::GetExternalCurvePath( ) const
 
 UCurveBase* STimelineEdTrack::CreateCurveAsset()
 {
-	UCurveBase* AssetCurve = NULL;
+	UCurveBase* AssetCurve = nullptr;
 
 	TSharedPtr<STimelineEditor> TimelineEd = TimelineEdPtr.Pin();
 	check(TimelineEd.IsValid());
@@ -749,7 +749,7 @@ UCurveBase* STimelineEdTrack::CreateCurveAsset()
 		}
 	}
 
-	return NULL;
+	return nullptr;
 }
 
 
@@ -954,10 +954,10 @@ void STimelineEditor::Construct(const FArguments& InArgs, TSharedPtr<FBlueprintE
 	NewTrackPendingRename = NAME_None;
 
 	Kismet2Ptr = InKismet2;
-	TimelineObj = NULL;
+	TimelineObj = nullptr;
 
 	NominalTimelineDesiredHeight = 300.0f;
-	TimelineDesiredSize = FVector2D(128.0f, NominalTimelineDesiredHeight);
+	TimelineDesiredSize = FVector2f(128.0f, NominalTimelineDesiredHeight);
 
 	// Leave these uninitialized at first.  We'll zoom to fit the tracks which will set the correct values
 	ViewMinInput = 0.f;
@@ -1237,7 +1237,7 @@ void STimelineEditor::OnTimelineTickGroupChanged(TSharedPtr<FString> NewValue, E
 
 FText STimelineEditor::GetTimelineName() const
 {
-	if(TimelineObj != NULL)
+	if(TimelineObj != nullptr)
 	{
 		return FText::FromString(TimelineObj->GetVariableName().ToString());
 	}
@@ -1269,7 +1269,7 @@ float STimelineEditor::GetViewMinOutput() const
 
 float STimelineEditor::GetTimelineLength() const
 {
-	return (TimelineObj != NULL) ? TimelineObj->TimelineLength : 0.f;
+	return (TimelineObj != nullptr) ? TimelineObj->TimelineLength : 0.f;
 }
 
 void STimelineEditor::SetInputViewRange(float InViewMinInput, float InViewMaxInput)
@@ -1404,7 +1404,7 @@ UCurveBase* STimelineEditor::CreateNewCurve(FTTTrackBase::ETrackType Type )
 	UBlueprint* Blueprint = Kismet2->GetBlueprintObj();
 	UClass* OwnerClass = Blueprint->GeneratedClass;
 	check(OwnerClass);
-	UCurveBase* NewCurve = NULL;
+	UCurveBase* NewCurve = nullptr;
 	if(Type == FTTTrackBase::TT_Event)
 	{
 		NewCurve = NewObject<UCurveFloat>(OwnerClass, NAME_None, RF_Public);
@@ -1433,7 +1433,7 @@ bool STimelineEditor::CanDeleteSelectedTracks() const
 
 void STimelineEditor::OnDeleteSelectedTracks()
 {
-	if(TimelineObj != NULL)
+	if(TimelineObj != nullptr)
 	{
 		TSharedPtr<FBlueprintEditor> Kismet2 = Kismet2Ptr.Pin();
 		UBlueprint* Blueprint = Kismet2->GetBlueprintObj();
@@ -1507,7 +1507,7 @@ void STimelineEditor::OnTimelineChanged()
 	TSharedPtr<FTimelineEdTrack> NewlyCreatedTrack;
 
 	// If we have a timeline,
-	if(TimelineObj != NULL)
+	if(TimelineObj != nullptr)
 	{
 		// Iterate over tracks and create entries in the array that drives the list widget
 		for (int32 i = 0; i < TimelineObj->GetNumDisplayTracks(); ++i)
@@ -1664,7 +1664,7 @@ void STimelineEditor::OnIgnoreTimeDilationChanged(ECheckBoxState NewType)
 FText STimelineEditor::GetLengthString() const
 {
 	FString LengthString(TEXT("0.0"));
-	if(TimelineObj != NULL)
+	if(TimelineObj != nullptr)
 	{
 		LengthString = FString::Printf(TEXT("%.2f"), TimelineObj->TimelineLength);
 	}
@@ -1674,7 +1674,7 @@ FText STimelineEditor::GetLengthString() const
 void STimelineEditor::OnLengthStringChanged(const FText& NewString, ETextCommit::Type CommitInfo)
 {
 	bool bCommitted = (CommitInfo == ETextCommit::OnEnter) || (CommitInfo == ETextCommit::OnUserMovedFocus);
-	if(TimelineObj != NULL && bCommitted)
+	if(TimelineObj != nullptr && bCommitted)
 	{
 		float NewLength = FCString::Atof( *NewString.ToString() );
 		if(NewLength > KINDA_SMALL_NUMBER)
@@ -1769,7 +1769,7 @@ void STimelineEditor::OnTrackNameCommitted( const FText& StringName, ETextCommit
 
 void STimelineEditor::OnReorderTracks(int32 DisplayIndex, int32 DirectionDelta)
 {
-	if (TimelineObj != NULL)
+	if (TimelineObj != nullptr)
 	{
 		const FScopedTransaction Transaction(LOCTEXT("TimelineEditor_DeleteTrack", "Delete track"));
 
@@ -1956,7 +1956,7 @@ TSharedRef<SWidget> STimelineEditor::MakeAddButton()
 
 FVector2D STimelineEditor::GetTimelineDesiredSize() const
 {
-	return TimelineDesiredSize;
+	return FVector2D{ TimelineDesiredSize };
 }
 
 void STimelineEditor::SetSizeScaleValue(float NewValue)
diff --git a/Engine/Source/Editor/Kismet/Private/STimelineEditor.h b/Engine/Source/Editor/Kismet/Private/STimelineEditor.h
--- a/Engine/Source/Editor/Kismet/Private/STimelineEditor.h
+++ b/Engine/Source/Editor/Kismet/Private/STimelineEditor.h
@@ -245,7 +245,7 @@ private:
 	TSharedPtr< FUICommandList > CommandList;
 
 	/** The current desired size of the timeline */
-	FVector2D TimelineDesiredSize;
+	FVector2f TimelineDesiredSize;
 
 	/** The nominal desired height of a single timeline track at 1.0x height */
 	float NominalTimelineDesiredHeight;
diff --git a/Engine/Source/Editor/Kismet/Private/WatchPointViewer.cpp b/Engine/Source/Editor/Kismet/Private/WatchPointViewer.cpp
--- a/Engine/Source/Editor/Kismet/Private/WatchPointViewer.cpp
+++ b/Engine/Source/Editor/Kismet/Private/WatchPointViewer.cpp
@@ -508,7 +508,7 @@ void SWatchViewer::Construct(const FArguments& InArgs, TArray<TSharedRef<FWatchR
 	ChildSlot
 	[
 		SNew(SBorder)
-		.Padding(4)
+		.Padding(4.0f)
 		.BorderImage( FAppStyle::GetBrush("ToolPanel.GroupBorder") )
 		[
 			SNew(SOverlay)
@@ -673,7 +673,7 @@ void SWatchTreeWidgetItem::Construct(const FArguments& InArgs, SWatchViewer* InO
 	this->WatchRow = InArgs._WatchToVisualize;
 	Owner = InOwner;
 
-	SMultiColumnTableRow< TSharedRef<FWatchRow> >::Construct(SMultiColumnTableRow< TSharedRef<FWatchRow> >::FArguments().Padding(1), InOwnerTableView);
+	SMultiColumnTableRow<TSharedRef<FWatchRow>>::Construct(SMultiColumnTableRow<TSharedRef<FWatchRow>>::FArguments().Padding(1.0f), InOwnerTableView);
 }
 
 TSharedRef<SWidget> SWatchTreeWidgetItem::GenerateWidgetForColumn(const FName& ColumnName)
diff --git a/Engine/Source/Editor/KismetCompiler/Private/KismetCompiler.cpp b/Engine/Source/Editor/KismetCompiler/Private/KismetCompiler.cpp
--- a/Engine/Source/Editor/KismetCompiler/Private/KismetCompiler.cpp
+++ b/Engine/Source/Editor/KismetCompiler/Private/KismetCompiler.cpp
@@ -2813,7 +2813,7 @@ void FKismetCompilerContext::SetCalculatedMetaDataAndFlags(UFunction* Function,
 		if (Property->HasAnyPropertyFlags(CPF_Parm))
 		{
 			++Function->NumParms;
-			Function->ParmsSize = Property->GetOffset_ForUFunction() + Property->GetSize();
+			Function->ParmsSize = IntCastChecked<int16, int32>(Property->GetOffset_ForUFunction() + Property->GetSize());
 
 			if (Property->HasAnyPropertyFlags(CPF_OutParm))
 			{
@@ -2822,7 +2822,7 @@ void FKismetCompilerContext::SetCalculatedMetaDataAndFlags(UFunction* Function,
 
 			if (Property->HasAnyPropertyFlags(CPF_ReturnParm))
 			{
-				Function->ReturnValueOffset = Property->GetOffset_ForUFunction();
+				Function->ReturnValueOffset = IntCastChecked<uint16, int32>(Property->GetOffset_ForUFunction());
 			}
 		}
 		else
diff --git a/Engine/Source/Editor/KismetCompiler/Private/KismetCompilerVMBackend.cpp b/Engine/Source/Editor/KismetCompiler/Private/KismetCompilerVMBackend.cpp
--- a/Engine/Source/Editor/KismetCompiler/Private/KismetCompilerVMBackend.cpp
+++ b/Engine/Source/Editor/KismetCompiler/Private/KismetCompilerVMBackend.cpp
@@ -63,7 +63,7 @@ public:
 	
 	void Serialize( void* V, int64 Length ) override
 	{
-		int32 iStart = ScriptBuffer.AddUninitialized( Length );
+		int32 iStart = ScriptBuffer.AddUninitialized(IntCastChecked<int32, int64>(Length));
 		FMemory::Memcpy( &(ScriptBuffer[iStart]), V, Length );
 	}
 
@@ -125,14 +125,14 @@ public:
 	{
 		checkSlow(E < 0xFF);
 
-		uint8 B = E; 
+		uint8 B = static_cast<uint8>(E); 
 		Serialize(&B, 1); 
 		return *this;
 	}
 
 	FArchive& operator<<(ECastToken E)
 	{
-		uint8 B = E; 
+		uint8 B = static_cast<uint8>(E);
 		Serialize(&B, 1); 
 		return *this;
 	}
@@ -148,7 +148,7 @@ public:
 
 	FArchive& operator<<(EPropertyType E)
 	{
-		uint8 B = E; 
+		uint8 B = static_cast<uint8>(E);
 		Serialize(&B, 1); 
 		return *this;
 	}
@@ -800,11 +800,11 @@ public:
 				if (EnumPtr)
 				{
 					//Get value from enum string
-					Value = EnumPtr->GetValueByName(*(Term->Name));
+					Value = IntCastChecked<uint8, int64>(EnumPtr->GetValueByName(*(Term->Name)));
 				}
 				else
 				{
-					Value = FCString::Atoi(*(Term->Name));
+					Value = IntCastChecked<uint8, int32>(FCString::Atoi(*(Term->Name)));
 				}
 
 				Writer << EX_ByteConst;
@@ -1632,7 +1632,8 @@ public:
 		EmitTerm(DestinationExpression);
 
 		Writer << EX_Cast;
-		uint8 CastType = !bIsInterfaceCast ? CST_ObjectToBool : CST_InterfaceToBool;
+		ECastToken CastToken = !bIsInterfaceCast ? CST_ObjectToBool : CST_InterfaceToBool;
+		uint8 CastType = static_cast<uint8>(CastToken);
 		Writer << CastType;
 		
 		FProperty* TargetProperty = !bIsInterfaceCast ? ((FProperty*)(GetDefault<FObjectProperty>())) : ((FProperty*)(GetDefault<FInterfaceProperty>()));
@@ -1863,7 +1864,7 @@ public:
 
 		Writer << EX_SwitchValue;
 		// number of cases (without default)
-		uint16 NumCases = ((Statement.RHS.Num() - 2) / TermsPerCase);
+		uint16 NumCases = IntCastChecked<uint16, int32>((Statement.RHS.Num() - 2) / TermsPerCase);
 		Writer << NumCases;
 		// end goto index
 		CodeSkipSizeType PatchUpNeededAtOffset = Writer.EmitPlaceholderSkip();
diff --git a/Engine/Source/Editor/KismetWidgets/Private/CreateBlueprintFromActorDialog.cpp b/Engine/Source/Editor/KismetWidgets/Private/CreateBlueprintFromActorDialog.cpp
--- a/Engine/Source/Editor/KismetWidgets/Private/CreateBlueprintFromActorDialog.cpp
+++ b/Engine/Source/Editor/KismetWidgets/Private/CreateBlueprintFromActorDialog.cpp
@@ -548,7 +548,7 @@ void SSCreateBlueprintPicker::Construct(const FArguments& InArgs)
 				.AutoHeight()
 				.HAlign(HAlign_Right)
 				.VAlign(VAlign_Bottom)
-				.Padding(8)
+				.Padding(8.0f)
 				[
 					SNew(SUniformGridPanel)
 					.SlotPadding(FAppStyle::GetMargin("StandardDialog.SlotPadding"))
diff --git a/Engine/Source/Editor/KismetWidgets/Private/SKismetLinearExpression.cpp b/Engine/Source/Editor/KismetWidgets/Private/SKismetLinearExpression.cpp
--- a/Engine/Source/Editor/KismetWidgets/Private/SKismetLinearExpression.cpp
+++ b/Engine/Source/Editor/KismetWidgets/Private/SKismetLinearExpression.cpp
@@ -210,7 +210,7 @@ TSharedRef<SWidget> SKismetLinearExpression::MakeNodeWidget(const UEdGraphNode*
 				+SHorizontalBox::Slot()
 				.AutoWidth()
 				.VAlign(VAlign_Center)
-				.Padding(4)
+				.Padding(4.0f)
 				[
 					SNew(STextBlock).Text(LOCTEXT("BeginExpression", "("))
 				];
diff --git a/Engine/Source/Editor/KismetWidgets/Private/SPinTypeSelector.cpp b/Engine/Source/Editor/KismetWidgets/Private/SPinTypeSelector.cpp
--- a/Engine/Source/Editor/KismetWidgets/Private/SPinTypeSelector.cpp
+++ b/Engine/Source/Editor/KismetWidgets/Private/SPinTypeSelector.cpp
@@ -424,7 +424,7 @@ void SPinTypeSelector::Construct(const FArguments& InArgs, FGetPinTypeTree GetPi
 		// Only have a combo button with an icon
 		Widget = SAssignNew( TypeComboButton, SComboButton )
 			.OnGetMenuContent(this, &SPinTypeSelector::GetMenuContent, false)
-			.ContentPadding(0)
+			.ContentPadding(0.0f)
 			.ToolTipText(this, &SPinTypeSelector::GetToolTipForComboBoxType)
 			.HasDownArrow(false)
 			.ButtonStyle(FAppStyle::Get(),  "BlueprintEditor.CompactPinTypeSelector")
@@ -448,7 +448,7 @@ void SPinTypeSelector::Construct(const FArguments& InArgs, FGetPinTypeTree GetPi
 				.ComboButtonStyle(FAppStyle::Get(),"BlueprintEditor.CompactVariableTypeSelector")
 				.MenuPlacement(EMenuPlacement::MenuPlacement_ComboBoxRight)
 				.OnGetMenuContent(this, &SPinTypeSelector::GetPinContainerTypeMenuContent)
-				.ContentPadding(0)
+				.ContentPadding(0.0f)
 				.ToolTip(IDocumentation::Get()->CreateToolTip(TAttribute<FText>(this, &SPinTypeSelector::GetToolTipForContainerWidget), NULL, *PinTypeSelectorStatics::BigTooltipDocLink, TEXT("Containers")))
 				.IsEnabled(TargetPinType.Get().PinCategory != UEdGraphSchema_K2::PC_Exec)
 				.Visibility(InArgs._bAllowArrays ? EVisibility::Visible : EVisibility::Collapsed)
@@ -474,7 +474,7 @@ void SPinTypeSelector::Construct(const FArguments& InArgs, FGetPinTypeTree GetPi
 				SAssignNew(TypeComboButton, SComboButton)
 				.ComboButtonStyle(FAppStyle::Get(), "ComboButton")
 				.OnGetMenuContent(this, &SPinTypeSelector::GetMenuContent, false)
-				.ContentPadding(0)
+				.ContentPadding(0.0f)
 				.ToolTipText(this, &SPinTypeSelector::GetToolTipForComboBoxType)
 				.ForegroundColor(FSlateColor::UseForeground())
 				.ButtonContent()
@@ -523,7 +523,7 @@ void SPinTypeSelector::Construct(const FArguments& InArgs, FGetPinTypeTree GetPi
 				[
 					SAssignNew( SecondaryTypeComboButton, SComboButton )
 					.OnGetMenuContent(this, &SPinTypeSelector::GetMenuContent, true )
-					.ContentPadding(0)
+					.ContentPadding(0.0f)
 					.ToolTipText(this, &SPinTypeSelector::GetToolTipForComboBoxSecondaryType)
 					.ButtonContent()
 					[
@@ -1246,7 +1246,7 @@ TSharedRef<SWidget> SPinTypeSelector::GetPinContainerTypeMenuContent()
 				.ColorAndOpacity(this, &SPinTypeSelector::GetTypeIconColor)
 			]
 			+ SHorizontalBox::Slot()
-			.Padding(4.0f,2.0f)
+			.Padding(4.0f, 2.0f)
 			[
 				SNew(STextBlock)
 				.Text(PinTypeSelectorStatics::Labels[(int32)PinContainerType])
diff --git a/Engine/Source/Editor/KismetWidgets/Private/SScrubControlPanel.cpp b/Engine/Source/Editor/KismetWidgets/Private/SScrubControlPanel.cpp
--- a/Engine/Source/Editor/KismetWidgets/Private/SScrubControlPanel.cpp
+++ b/Engine/Source/Editor/KismetWidgets/Private/SScrubControlPanel.cpp
@@ -99,7 +99,7 @@ void SScrubControlPanel::Construct( const SScrubControlPanel::FArguments& InArgs
 		.AutoWidth()
 		[
 			SNew(SBorder)
-			.Padding(0)
+			.Padding(0.0f)
 			.BorderImage(FAppStyle::GetBrush("NoBorder"))
 			.Visibility(this, &SScrubControlPanel::GetRealtimeControlVisibility, false)
 			[
@@ -111,7 +111,7 @@ void SScrubControlPanel::Construct( const SScrubControlPanel::FArguments& InArgs
 		.AutoWidth()
 		[
 			SNew(SBorder)
-			.Padding(0)
+			.Padding(0.0f)
 			.BorderImage(FAppStyle::GetBrush("NoBorder"))
 			.Visibility(this, &SScrubControlPanel::GetRealtimeControlVisibility, true)
 			[
diff --git a/Engine/Source/Editor/KismetWidgets/Private/SScrubWidget.cpp b/Engine/Source/Editor/KismetWidgets/Private/SScrubWidget.cpp
--- a/Engine/Source/Editor/KismetWidgets/Private/SScrubWidget.cpp
+++ b/Engine/Source/Editor/KismetWidgets/Private/SScrubWidget.cpp
@@ -48,13 +48,13 @@ int32 SScrubWidget::GetDivider(float InputMinX, float InputMaxX, FVector2D Widge
 
 	float TimePerKey = Helper.TimePerKey();
 
-	float TotalWidgetWidth = TimeScaleInfo.WidgetSize.X;
-	int32 NumKeys = TimeScaleInfo.ViewInputRange / TimePerKey;
-	float KeyWidgetWidth = TotalWidgetWidth / (float)NumKeys;
+	float TotalWidgetWidth = static_cast<float>(TimeScaleInfo.WidgetSize.X);
+	int32 NumKeys = static_cast<int32>(TimeScaleInfo.ViewInputRange / TimePerKey);
+	float KeyWidgetWidth = TotalWidgetWidth / static_cast<float>(NumKeys);
 	int32 Divider = 1; 
 	if (KeyWidgetWidth > 0)
 	{
-		Divider = FMath::Max<int32>(50 / KeyWidgetWidth, 1);
+		Divider = FMath::Max<int32>(static_cast<int32>(50 / KeyWidgetWidth), 1);
 	}
 	return Divider;
 }
@@ -112,7 +112,7 @@ int32 SScrubWidget::OnPaint( const FPaintArgs& Args, const FGeometry& AllottedGe
 	const int32 TextLayer = BackgroundLayer + 1;
 
 	const FSlateBrush* StyleInfo = FAppStyle::Get().GetBrush("Brushes.Recessed");
-	const float GeomHeight = AllottedGeometry.GetLocalSize().Y;
+	const double GeomHeight = AllottedGeometry.GetLocalSize().Y;
 
 	const FTrackScaleInfo TimeScaleInfo(ViewInputMin.Get(), ViewInputMax.Get(), 0.f, 0.f, AllottedGeometry.GetLocalSize());
 
@@ -143,7 +143,7 @@ int32 SScrubWidget::OnPaint( const FPaintArgs& Args, const FGeometry& AllottedGe
 					InWidgetStyle.GetColorAndOpacityTint()
 					);
 
-				const int32 FrameNumber = KeyVal;
+				const int32 FrameNumber = static_cast<int32>(KeyVal);
 				const FString FrameString = FString::Printf(TEXT("%d"), (FrameNumber));
 				const FVector2D TextOffset(XPos+2.f, 0.f);
 
@@ -160,7 +160,7 @@ int32 SScrubWidget::OnPaint( const FPaintArgs& Args, const FGeometry& AllottedGe
 			}
  			else if (HalfDivider > 1.f)
  			{
- 				const float Height = GeomHeight;
+ 				const float Height = static_cast<float>(GeomHeight);
  				const FVector2D Offset(XPos, Height*0.25f);
  				const FVector2D Size(1, Height*0.5f);
  				// draw each box with key frame
@@ -181,8 +181,8 @@ int32 SScrubWidget::OnPaint( const FPaintArgs& Args, const FGeometry& AllottedGe
 	{
 		const int32 ArrowLayer = TextLayer + 1;
 		{
-			const float XPos = TimeScaleInfo.InputToLocalX(ValueAttribute.Get());
-			const float Height = AllottedGeometry.GetLocalSize().Y;
+			const double XPos = TimeScaleInfo.InputToLocalX(ValueAttribute.Get());
+			const double Height = AllottedGeometry.GetLocalSize().Y;
 			const FVector2D Offset( XPos - Height*0.25f, 0.f );
 
 			FPaintGeometry MyGeometry =	AllottedGeometry.ToPaintGeometry( Offset, FVector2D(Height * 0.5f, Height) );
@@ -271,8 +271,8 @@ FReply SScrubWidget::OnMouseButtonUp( const FGeometry& MyGeometry, const FPointe
 		bPanning = false;
 
 		FTrackScaleInfo TimeScaleInfo(ViewInputMin.Get(), ViewInputMax.Get(), 0.f, 0.f, MyGeometry.GetLocalSize());
-		FVector2D CursorPos = MyGeometry.AbsoluteToLocal(MouseEvent.GetLastScreenSpacePosition());
-		float NewValue = TimeScaleInfo.LocalXToInput(CursorPos.X);
+		float CursorPosX = static_cast<float>(MyGeometry.AbsoluteToLocal(MouseEvent.GetLastScreenSpacePosition()).X);
+		float NewValue = TimeScaleInfo.LocalXToInput(CursorPosX);
 
 		if( !bMouseMovedDuringPanning )
 		{
@@ -288,9 +288,9 @@ FReply SScrubWidget::OnMouseButtonUp( const FGeometry& MyGeometry, const FPointe
 
 			if(OnBarCommit.IsBound())
 			{
-				const FVector2D CursorPos = MyGeometry.AbsoluteToLocal( MouseEvent.GetScreenSpacePosition() );
+				const float CursorPosX = static_cast<float>(MyGeometry.AbsoluteToLocal(MouseEvent.GetScreenSpacePosition()).X);
 				const FTrackScaleInfo ScaleInfo(ViewInputMin.Get(), ViewInputMax.Get(), 0.f, 0.f, MyGeometry.GetLocalSize());
-				const float NewDataPos = FMath::Clamp( ScaleInfo.LocalXToInput(CursorPos.X), ViewInputMin.Get(), ViewInputMax.Get() );				
+				const float NewDataPos = FMath::Clamp( ScaleInfo.LocalXToInput(CursorPosX), ViewInputMin.Get(), ViewInputMax.Get() );				
 				OnBarCommit.ExecuteIfBound(DraggableBarIndex, NewDataPos);
 			}			
 		}
@@ -301,15 +301,14 @@ FReply SScrubWidget::OnMouseButtonUp( const FGeometry& MyGeometry, const FPointe
 		else
 		{
 			FTrackScaleInfo TimeScaleInfo(ViewInputMin.Get(), ViewInputMax.Get(), 0.f, 0.f, MyGeometry.GetLocalSize());
-			FVector2D CursorPos = MyGeometry.AbsoluteToLocal(MouseEvent.GetLastScreenSpacePosition());
-			float NewValue = TimeScaleInfo.LocalXToInput(CursorPos.X);
+			float CursorPosX = static_cast<float>(MyGeometry.AbsoluteToLocal(MouseEvent.GetLastScreenSpacePosition()).X);
+			float NewValue = static_cast<float>(TimeScaleInfo.LocalXToInput(CursorPosX));
 
-			CommitValue( NewValue, true, false );
+			CommitValue(NewValue, true, false);
 		}
 
 		bDragging = false;
 		return FReply::Handled().ReleaseMouseCapture();
-
 	}
 
 	return FReply::Unhandled();
@@ -321,15 +320,15 @@ FReply SScrubWidget::OnMouseMove( const FGeometry& MyGeometry, const FPointerEve
 	if(DraggingBar)
 	{
 		// Update bar if we are dragging
-		FVector2D CursorPos = MyGeometry.AbsoluteToLocal( MouseEvent.GetScreenSpacePosition() );
+		float CursorPosX = static_cast<float>(MyGeometry.AbsoluteToLocal(MouseEvent.GetScreenSpacePosition()).X);
 		FTrackScaleInfo ScaleInfo(ViewInputMin.Get(), ViewInputMax.Get(), 0.f, 0.f, MyGeometry.GetLocalSize());
-		float NewDataPos = FMath::Clamp( ScaleInfo.LocalXToInput(CursorPos.X), ViewInputMin.Get(), ViewInputMax.Get() );
+		float NewDataPos = FMath::Clamp(ScaleInfo.LocalXToInput(CursorPosX), ViewInputMin.Get(), ViewInputMax.Get());
 		OnBarDrag.ExecuteIfBound(DraggableBarIndex, NewDataPos);
 	}
 	else
 	{
 		// Update what bar we are hovering over
-		FVector2D CursorPos = MyGeometry.AbsoluteToLocal(MouseEvent.GetScreenSpacePosition());
+		float CursorPosX = static_cast<float>(MyGeometry.AbsoluteToLocal(MouseEvent.GetScreenSpacePosition()).X);
 		FTrackScaleInfo ScaleInfo(ViewInputMin.Get(), ViewInputMax.Get(), 0.f, 0.f, MyGeometry.GetLocalSize());
 		DraggableBarIndex = INDEX_NONE;
 		if ( DraggableBars.IsBound() )
@@ -337,7 +336,7 @@ FReply SScrubWidget::OnMouseMove( const FGeometry& MyGeometry, const FPointerEve
 			const TArray<float>& DraggableBarsVal = DraggableBars.Get();
 			for ( int32 I=0; I < DraggableBarsVal.Num(); I++ )
 			{
-				if( FMath::Abs( ScaleInfo.InputToLocalX(DraggableBarsVal[I]) - CursorPos.X ) < 10 )
+				if(FMath::Abs( ScaleInfo.InputToLocalX(DraggableBarsVal[I]) - CursorPosX ) < 10)
 				{
 					DraggableBarIndex = I;
 					break;
@@ -352,7 +351,7 @@ FReply SScrubWidget::OnMouseMove( const FGeometry& MyGeometry, const FPointerEve
 		{
 			FTrackScaleInfo ScaleInfo(ViewInputMin.Get(), ViewInputMax.Get(), 0.f, 0.f, MyGeometry.GetLocalSize());
 			FVector2D ScreenDelta = MouseEvent.GetCursorDelta();
-			float InputDeltaX = ScreenDelta.X/ScaleInfo.PixelsPerInput;
+			float InputDeltaX = static_cast<float>(ScreenDelta.X / ScaleInfo.PixelsPerInput);
 
 			bMouseMovedDuringPanning |= !ScreenDelta.IsNearlyZero(0.001f);
 
@@ -374,7 +373,8 @@ FReply SScrubWidget::OnMouseMove( const FGeometry& MyGeometry, const FPointerEve
 		}
 		else if (!bDragging)
 		{
-			DistanceDragged += FMath::Abs(MouseEvent.GetCursorDelta().X);
+			float CursorDeltaX = static_cast<float>(FMath::Abs(MouseEvent.GetCursorDelta().X));
+			DistanceDragged += CursorDeltaX;
 			if ( DistanceDragged > FSlateApplication::Get().GetDragTriggerDistance() )
 			{
 				bDragging = true;
@@ -387,8 +387,8 @@ FReply SScrubWidget::OnMouseMove( const FGeometry& MyGeometry, const FPointerEve
 		else if (bDragging)
 		{
 			FTrackScaleInfo TimeScaleInfo(ViewInputMin.Get(), ViewInputMax.Get(), 0.f, 0.f, MyGeometry.GetLocalSize());
-			FVector2D CursorPos = MyGeometry.AbsoluteToLocal(MouseEvent.GetLastScreenSpacePosition());
-			float NewValue = TimeScaleInfo.LocalXToInput(CursorPos.X);
+			float CursorPosX = static_cast<float>(MyGeometry.AbsoluteToLocal(MouseEvent.GetLastScreenSpacePosition()).X);
+			float NewValue = TimeScaleInfo.LocalXToInput(CursorPosX);
 
 			CommitValue( NewValue, true, false );
 		}
@@ -465,7 +465,7 @@ void SScrubWidget::CreateContextMenu(float CurrentFrameTime, const FPointerEvent
 		MenuBuilder.BeginSection("SequenceEditingContext", LOCTEXT("SequenceEditing", "Sequence Editing") );
 		{
 			float CurrentFrameFraction = CurrentFrameTime / SequenceLength.Get();
-			int32 CurrentFrameNumber = CurrentFrameFraction * NumOfKeys.Get();
+			int32 CurrentFrameNumber = static_cast<int32>(CurrentFrameFraction * NumOfKeys.Get());
 
 			FUIAction Action;
 			FText Label;
diff --git a/Engine/Source/Runtime/CoreUObject/Public/UObject/Script.h b/Engine/Source/Runtime/CoreUObject/Public/UObject/Script.h
--- a/Engine/Source/Runtime/CoreUObject/Public/UObject/Script.h
+++ b/Engine/Source/Runtime/CoreUObject/Public/UObject/Script.h
@@ -129,7 +129,7 @@ enum EFunctionFlags : uint32
 	FUNC_None				= 0x00000000,
 
 	FUNC_Final				= 0x00000001,	// Function is final (prebindable, non-overridable function).
-	FUNC_RequiredAPI			= 0x00000002,	// Indicates this function is DLL exported/imported.
+	FUNC_RequiredAPI		= 0x00000002,	// Indicates this function is DLL exported/imported.
 	FUNC_BlueprintAuthorityOnly= 0x00000004,   // Function will only run if the object has network authority
 	FUNC_BlueprintCosmetic	= 0x00000008,   // Function is cosmetic in nature and should not be invoked on dedicated servers
 	// FUNC_				= 0x00000010,   // unused.
@@ -181,7 +181,7 @@ ENUM_CLASS_FLAGS(EFunctionFlags)
 //
 // Evaluatable expression item types.
 //
-enum EExprToken
+enum EExprToken : uint8
 {
 	// Variable references.
 	EX_LocalVariable		= 0x00,	// A local variable.
@@ -294,10 +294,10 @@ enum EExprToken
 	EX_ArrayGetByRef		= 0x6B,
 	EX_ClassSparseDataVariable = 0x6C, // Sparse data variable
 	EX_FieldPathConst		= 0x6D,
-	EX_Max					= 0x100,
+	EX_Max					= 0xFF,
 };
 
-enum ECastToken
+enum ECastToken : uint8
 {
 	CST_ObjectToInterface		= 0x00,
 	CST_ObjectToBool			= 0x01,
