[Feature] Blueprints - Extract Details View of BP Diff tool and set the Fidd Override for UDataAsset types for this tool

diff --git a/Engine/Source/Developer/AssetTools/Private/AssetTypeActions/AssetTypeActions_DataAsset.cpp b/Engine/Source/Developer/AssetTools/Private/AssetTypeActions/AssetTypeActions_DataAsset.cpp
--- a/Engine/Source/Developer/AssetTools/Private/AssetTypeActions/AssetTypeActions_DataAsset.cpp
+++ b/Engine/Source/Developer/AssetTools/Private/AssetTypeActions/AssetTypeActions_DataAsset.cpp
@@ -7,6 +7,7 @@
 #include "Kismet2/SClassPickerDialog.h"
 #include "ClassViewerFilter.h"
 #include "ObjectTools.h"
+#include "SDetailsDiff.h"
 
 #define LOCTEXT_NAMESPACE "AssetTypeActions"
 
@@ -70,6 +71,18 @@ void FAssetTypeActions_DataAsset::GetActions(const TArray<UObject*>& InObjects,
 	);
 }
 
+void FAssetTypeActions_DataAsset::PerformAssetDiff(UObject* OldAsset, UObject* NewAsset, const FRevisionInfo& OldRevision, const FRevisionInfo& NewRevision) const
+{
+	// sometimes we're comparing different revisions of one single asset (other 
+	// times we're comparing two completely separate assets altogether)
+	const bool bIsSingleAsset = (OldAsset->GetName() == NewAsset->GetName());
+	static const FText BasicWindowTitle = LOCTEXT("NamelessDataAssetDiff", "DataAsset Diff");
+
+	const FText WindowTitle = !bIsSingleAsset ? BasicWindowTitle : FText::Format(LOCTEXT("DataAsset Diff", "{0} - DataAsset Diff"), FText::FromString(NewAsset->GetName()));
+
+	SDetailsDiff::CreateDiffWindow(WindowTitle, OldAsset, NewAsset, OldRevision, NewRevision);
+}
+
 void FAssetTypeActions_DataAsset::ExecuteChangeDataAssetClass(TArray<TWeakObjectPtr<UDataAsset>> InDataAssets)
 {
 	const FText TitleText = LOCTEXT("DataAsset_PickNewDataAssetClass", "Pick New DataAsset Class");
diff --git a/Engine/Source/Developer/AssetTools/Public/AssetTypeActions/AssetTypeActions_DataAsset.h b/Engine/Source/Developer/AssetTools/Public/AssetTypeActions/AssetTypeActions_DataAsset.h
--- a/Engine/Source/Developer/AssetTools/Public/AssetTypeActions/AssetTypeActions_DataAsset.h
+++ b/Engine/Source/Developer/AssetTools/Public/AssetTypeActions/AssetTypeActions_DataAsset.h
@@ -18,6 +18,7 @@ public:
 	virtual FText GetDisplayNameFromAssetData(const FAssetData& AssetData) const override;
 	virtual bool HasActions(const TArray<UObject*>& InObjects) const override { return true; }
 	virtual void GetActions(const TArray<UObject*>& InObjects, struct FToolMenuSection& Section) override;
+	virtual void PerformAssetDiff(UObject* OldAsset, UObject* NewAsset, const FRevisionInfo& OldRevision, const FRevisionInfo& NewRevision) const override;
 
 private:
 	void ExecuteChangeDataAssetClass(TArray<TWeakObjectPtr<UDataAsset>> InDataAssets);
diff --git a/Engine/Source/Editor/Kismet/Private/DiffControl.cpp b/Engine/Source/Editor/Kismet/Private/DiffControl.cpp
new file mode 100644
--- /dev/null
+++ b/Engine/Source/Editor/Kismet/Private/DiffControl.cpp
@@ -0,0 +1,610 @@
+#include "DiffControl.h"
+
+#include "DiffResults.h"
+#include "GraphDiffControl.h"
+#include "SBlueprintDiff.h"
+#include "Kismet2/BlueprintEditorUtils.h"
+#include "IDetailsView.h"
+
+#define LOCTEXT_NAMESPACE "SBlueprintDif"
+
+
+/////////////////////////////////////////////////////////////////////////////
+/// IDiffControl
+
+FText IDiffControl::RightRevision = LOCTEXT("OlderRevisionIdentifier", "Right Revision");
+
+TSharedRef<SWidget> IDiffControl::GenerateSimpleDiffWidget(FText DiffText)
+{
+	return SNew(STextBlock)
+		.Text(DiffText)
+		.ToolTipText(DiffText)
+		.ColorAndOpacity(DiffViewUtils::Differs());
+}
+
+TSharedRef<SWidget> IDiffControl::GenerateObjectDiffWidget(FSingleObjectDiffEntry DiffEntry, FText ObjectName)
+{
+	return SNew(STextBlock)
+		.Text(DiffViewUtils::PropertyDiffMessage(DiffEntry, ObjectName))
+		.ToolTipText(DiffViewUtils::PropertyDiffMessage(DiffEntry, ObjectName))
+		.ColorAndOpacity(DiffViewUtils::Differs());
+}
+
+
+/////////////////////////////////////////////////////////////////////////////
+/// FMyBlueprintDiffControl
+
+FMyBlueprintDiffControl::FMyBlueprintDiffControl(const UBlueprint* InOldBlueprint, const UBlueprint* InNewBlueprint, FOnDiffEntryFocused InSelectionCallback)
+	: SelectionCallback(MoveTemp(InSelectionCallback))
+	, OldBlueprint(InOldBlueprint)
+	, NewBlueprint(InNewBlueprint)
+{}
+
+void FMyBlueprintDiffControl::GenerateTreeEntries(TArray<TSharedPtr<FBlueprintDifferenceTreeEntry>>& OutTreeEntries, TArray<TSharedPtr<FBlueprintDifferenceTreeEntry>>& OutRealDifferences)
+{
+	TArray< TSharedPtr<FBlueprintDifferenceTreeEntry> > Children;
+
+	if (OldBlueprint && OldBlueprint->SkeletonGeneratedClass && NewBlueprint && NewBlueprint->SkeletonGeneratedClass)
+	{
+		for (TFieldIterator<FProperty> PropertyIt(OldBlueprint->SkeletonGeneratedClass); PropertyIt; ++PropertyIt)
+		{
+			FProperty* OldProperty = *PropertyIt;
+			FProperty* NewProperty = NewBlueprint->SkeletonGeneratedClass->FindPropertyByName(OldProperty->GetFName());
+
+			FText PropertyText = FText::FromString(OldProperty->GetAuthoredName());
+
+			if (NewProperty)
+			{
+				const int32 OldVarIndex = FBlueprintEditorUtils::FindNewVariableIndex(OldBlueprint, OldProperty->GetFName());
+				const int32 NewVarIndex = FBlueprintEditorUtils::FindNewVariableIndex(NewBlueprint, OldProperty->GetFName());
+
+				if (OldVarIndex != INDEX_NONE && NewVarIndex != INDEX_NONE)
+				{
+					TArray<FSingleObjectDiffEntry> DifferingProperties;
+					DiffUtils::CompareUnrelatedStructs(FBPVariableDescription::StaticStruct(), &OldBlueprint->NewVariables[OldVarIndex], FBPVariableDescription::StaticStruct(), &NewBlueprint->NewVariables[NewVarIndex], DifferingProperties);
+					for (const FSingleObjectDiffEntry& Difference : DifferingProperties)
+					{
+						TSharedPtr<FBlueprintDifferenceTreeEntry> Entry = MakeShared<FBlueprintDifferenceTreeEntry>(
+							SelectionCallback,
+							FGenerateDiffEntryWidget::CreateStatic(&GenerateObjectDiffWidget, Difference, PropertyText));
+						Children.Push(Entry);
+						OutRealDifferences.Push(Entry);
+					}
+				}	
+			}
+			else
+			{
+				FText DiffText = FText::Format(LOCTEXT("VariableRemoved", "Removed Variable {0}"), PropertyText);
+
+				TSharedPtr<FBlueprintDifferenceTreeEntry> Entry = MakeShared<FBlueprintDifferenceTreeEntry>(
+					SelectionCallback,
+					FGenerateDiffEntryWidget::CreateStatic(&GenerateSimpleDiffWidget, DiffText));
+
+				Children.Push(Entry);
+				OutRealDifferences.Push(Entry);
+			}
+		}
+
+		for (TFieldIterator<FProperty> PropertyIt(NewBlueprint->SkeletonGeneratedClass); PropertyIt; ++PropertyIt)
+		{
+			FProperty* NewProperty = *PropertyIt;
+			FProperty* OldProperty = OldBlueprint->SkeletonGeneratedClass->FindPropertyByName(NewProperty->GetFName());
+
+			if (!OldProperty)
+			{
+				FText DiffText = FText::Format(LOCTEXT("VariableAdded", "Added Variable {0}"), FText::FromString(NewProperty->GetAuthoredName()));
+
+				TSharedPtr<FBlueprintDifferenceTreeEntry> Entry = MakeShared<FBlueprintDifferenceTreeEntry>(
+					SelectionCallback,
+					FGenerateDiffEntryWidget::CreateStatic(&GenerateSimpleDiffWidget, DiffText));
+
+				Children.Push(Entry);
+				OutRealDifferences.Push(Entry);
+			}
+		}
+	}
+	const bool bHasDifferences = Children.Num() != 0;
+	if (!bHasDifferences)
+	{
+		// make one child informing the user that there are no differences:
+		Children.Push(FBlueprintDifferenceTreeEntry::NoDifferencesEntry());
+	}
+
+	static const FText MyBlueprintLabel = NSLOCTEXT("FBlueprintDifferenceTreeEntry", "MyBlueprintLabel", "My Blueprint");
+	static const FText MyBlueprintTooltip = NSLOCTEXT("FBlueprintDifferenceTreeEntry", "MyBlueprintTooltip", "The list of changes made to blueprint structure in the My Blueprint panel");
+	OutTreeEntries.Push(FBlueprintDifferenceTreeEntry::CreateCategoryEntry(
+		MyBlueprintLabel,
+		MyBlueprintTooltip,
+		SelectionCallback,
+		Children,
+		bHasDifferences
+	));
+}
+
+
+/////////////////////////////////////////////////////////////////////////////
+/// FSCSDiffControl
+
+FSCSDiffControl::FSCSDiffControl(const UBlueprint* InOldBlueprint, const UBlueprint* InNewBlueprint, FOnDiffEntryFocused InSelectionCallback)
+	: SelectionCallback(InSelectionCallback)
+	, OldSCS(InOldBlueprint)
+	, NewSCS(InNewBlueprint)
+{
+}
+
+void FSCSDiffControl::GenerateTreeEntries(TArray<TSharedPtr<FBlueprintDifferenceTreeEntry>>& OutTreeEntries, TArray<TSharedPtr<FBlueprintDifferenceTreeEntry>>& OutRealDifferences)
+{
+	TArray< FSCSResolvedIdentifier > OldHierarchy = OldSCS.GetDisplayedHierarchy();
+	TArray< FSCSResolvedIdentifier > NewHierarchy = NewSCS.GetDisplayedHierarchy();
+	DiffUtils::CompareUnrelatedSCS(OldSCS.GetBlueprint(), OldHierarchy, NewSCS.GetBlueprint(), NewHierarchy, DifferingProperties);
+
+	const auto FocusSCSDifferenceEntry = [](FSCSDiffEntry Entry, FOnDiffEntryFocused InSelectionCallback, FSCSDiffControl* Owner)
+	{
+		InSelectionCallback.ExecuteIfBound();
+		if (Entry.TreeIdentifier.Name != NAME_None)
+		{
+			Owner->OldSCS.HighlightProperty(Entry.TreeIdentifier.Name, FPropertyPath());
+			Owner->NewSCS.HighlightProperty(Entry.TreeIdentifier.Name, FPropertyPath());
+		}
+	};
+
+	const auto CreateSCSDifferenceWidget = [](FSCSDiffEntry Entry, FText ObjectName) -> TSharedRef<SWidget>
+	{
+		return SNew(STextBlock)
+				.Text(DiffViewUtils::SCSDiffMessage(Entry, ObjectName))
+				.ColorAndOpacity(DiffViewUtils::Differs());
+	};
+
+	TArray<TSharedPtr<FBlueprintDifferenceTreeEntry>> Children;
+	for (const FSCSDiffEntry& Difference : DifferingProperties.Entries)
+	{
+		TSharedPtr<FBlueprintDifferenceTreeEntry> Entry = MakeShared<FBlueprintDifferenceTreeEntry>(
+			FOnDiffEntryFocused::CreateStatic(FocusSCSDifferenceEntry, Difference, SelectionCallback, this),
+			FGenerateDiffEntryWidget::CreateStatic(CreateSCSDifferenceWidget, Difference, RightRevision));
+		Children.Push(Entry);
+		OutRealDifferences.Push(Entry);
+	}
+
+	const bool bHasDifferences = Children.Num() != 0;
+	if (!bHasDifferences)
+	{
+		// make one child informing the user that there are no differences:
+		Children.Push(FBlueprintDifferenceTreeEntry::NoDifferencesEntry());
+	}
+
+	static const FText SCSLabel = NSLOCTEXT("FBlueprintDifferenceTreeEntry", "SCSLabel", "Components");
+	static const FText SCSTooltip = NSLOCTEXT("FBlueprintDifferenceTreeEntry", "SCSTooltip", "The list of changes made in the Components panel");
+	OutTreeEntries.Push(FBlueprintDifferenceTreeEntry::CreateCategoryEntry(
+		SCSLabel,
+		SCSTooltip,
+		SelectionCallback,
+		Children,
+		bHasDifferences
+	));
+}
+
+
+/////////////////////////////////////////////////////////////////////////////
+/// FCDODiffControl
+
+FCDODiffControl::FCDODiffControl(const UObject* InOldObject, const UObject* InNewObject, FOnDiffEntryFocused InSelectionCallback): TDetailsDiffControl(InOldObject, InNewObject, InSelectionCallback)
+{
+}
+
+void FCDODiffControl::GenerateTreeEntries(TArray<TSharedPtr<FBlueprintDifferenceTreeEntry>>& OutTreeEntries, TArray<TSharedPtr<FBlueprintDifferenceTreeEntry>>& OutRealDifferences)
+{
+	TDetailsDiffControl::GenerateTreeEntries(OutTreeEntries, OutRealDifferences);
+
+	const bool bHasDifferences = Children.Num() != 0;
+	if (!bHasDifferences)
+	{
+		// make one child informing the user that there are no differences:
+		Children.Push(FBlueprintDifferenceTreeEntry::NoDifferencesEntry());
+	}
+
+	static const FText DefaultsLabel = NSLOCTEXT("FBlueprintDifferenceTreeEntry", "DefaultsLabel", "Defaults");
+	static const FText DefaultsTooltip = NSLOCTEXT("FBlueprintDifferenceTreeEntry", "DefaultsTooltip", "The list of changes made in the Defaults panel");
+	OutTreeEntries.Push(FBlueprintDifferenceTreeEntry::CreateCategoryEntry(
+		DefaultsLabel,
+		DefaultsTooltip,
+		SelectionCallback,
+		Children,
+		bHasDifferences
+	));
+}
+
+
+/////////////////////////////////////////////////////////////////////////////
+/// FClassSettingsDiffControl
+
+FClassSettingsDiffControl::FClassSettingsDiffControl(const UObject* InOldObject, const UObject* InNewObject, FOnDiffEntryFocused InSelectionCallback): TDetailsDiffControl(InOldObject, InNewObject, InSelectionCallback)
+{
+}
+
+void FClassSettingsDiffControl::GenerateTreeEntries(TArray<TSharedPtr<FBlueprintDifferenceTreeEntry>>& OutTreeEntries, TArray<TSharedPtr<FBlueprintDifferenceTreeEntry>>& OutRealDifferences)
+{
+	TDetailsDiffControl::GenerateTreeEntries(OutTreeEntries, OutRealDifferences);
+
+	// Check for parent class and interfaces here
+	const UBlueprint* OldBlueprint = Cast<UBlueprint>(OldDetails.GetDisplayedObject());
+	const UBlueprint* NewBlueprint = Cast<UBlueprint>(NewDetails.GetDisplayedObject());
+
+	if (OldBlueprint && NewBlueprint)
+	{
+		if (OldBlueprint->ParentClass != NewBlueprint->ParentClass)
+		{
+			FText DiffText = FText::Format(LOCTEXT("ParentChanged", "Parent Class changed from {0} to {1}"), FText::FromString(GetNameSafe(OldBlueprint->ParentClass)), FText::FromString(GetNameSafe(NewBlueprint->ParentClass)));
+
+			TSharedPtr<FBlueprintDifferenceTreeEntry> Entry = MakeShared<FBlueprintDifferenceTreeEntry>(
+				SelectionCallback,
+				FGenerateDiffEntryWidget::CreateStatic(&GenerateSimpleDiffWidget, DiffText));
+
+			Children.Push(Entry);
+			OutRealDifferences.Push(Entry);
+		}
+
+		FString OldInterfaces, NewInterfaces;
+		for (const FBPInterfaceDescription& Desc : OldBlueprint->ImplementedInterfaces)
+		{
+			if (!OldInterfaces.IsEmpty())
+			{
+				OldInterfaces += TEXT(", ");
+			}
+			OldInterfaces += GetNameSafe(*Desc.Interface);
+		}
+
+		for (const FBPInterfaceDescription& Desc : NewBlueprint->ImplementedInterfaces)
+		{
+			if (!NewInterfaces.IsEmpty())
+			{
+				NewInterfaces += TEXT(", ");
+			}
+			NewInterfaces += GetNameSafe(*Desc.Interface);
+		}
+			
+		if (OldInterfaces != NewInterfaces)
+		{
+			FText DiffText = FText::Format(LOCTEXT("InterfacesChanged", "Interfaces changed from '{0}' to '{1}'"), FText::FromString(OldInterfaces), FText::FromString(NewInterfaces));
+
+			TSharedPtr<FBlueprintDifferenceTreeEntry> Entry = MakeShared<FBlueprintDifferenceTreeEntry>(
+				SelectionCallback,
+				FGenerateDiffEntryWidget::CreateStatic(&GenerateSimpleDiffWidget, DiffText));
+
+			Children.Push(Entry);
+			OutRealDifferences.Push(Entry);
+		}
+	}
+
+	const bool bHasDifferences = Children.Num() != 0;
+	if (!bHasDifferences)
+	{
+		// make one child informing the user that there are no differences:
+		Children.Push(FBlueprintDifferenceTreeEntry::NoDifferencesEntry());
+	}
+
+	static const FText SettingsLabel = NSLOCTEXT("FBlueprintDifferenceTreeEntry", "SettingsLabel", "Class Settings");
+	static const FText SettingsTooltip = NSLOCTEXT("FBlueprintDifferenceTreeEntry", "SettingsTooltip", "The list of changes made in the Class Settings panel");
+	OutTreeEntries.Push(FBlueprintDifferenceTreeEntry::CreateCategoryEntry(
+		SettingsLabel,
+		SettingsTooltip,
+		SelectionCallback,
+		Children,
+		bHasDifferences
+	));
+}
+
+
+/////////////////////////////////////////////////////////////////////////////
+/// FBlueprintTypeDiffControl
+
+FBlueprintTypeDiffControl::FBlueprintTypeDiffControl(const UBlueprint* InBlueprintOld, const UBlueprint* InBlueprintNew, FOnDiffEntryFocused InSelectionCallback): BlueprintOld(InBlueprintOld), BlueprintNew(InBlueprintNew), SelectionCallback(InSelectionCallback), bDiffSucceeded(false)
+{
+	check(InBlueprintNew && InBlueprintOld);
+}
+
+void FBlueprintTypeDiffControl::GenerateTreeEntries(TArray< TSharedPtr<FBlueprintDifferenceTreeEntry> >& OutTreeEntries, TArray< TSharedPtr<FBlueprintDifferenceTreeEntry> >& OutRealDifferences)
+{
+	BuildDiffSourceArray();
+
+	TArray< TSharedPtr<FBlueprintDifferenceTreeEntry> > Children;
+	
+	bool bHasRealChange = false;
+
+	// First add manual diffs in main category
+	for (const TSharedPtr<FDiffResultItem>& Difference : DiffListSource)
+	{
+		TSharedPtr<FBlueprintDifferenceTreeEntry> ChildEntry = MakeShared<FBlueprintDifferenceTreeEntry>(
+			SelectionCallback,
+			FGenerateDiffEntryWidget::CreateSP(Difference.ToSharedRef(), &FDiffResultItem::GenerateWidget));
+		Children.Push(ChildEntry);
+		OutRealDifferences.Push(ChildEntry);
+
+		if (Difference->Result.IsRealDifference())
+		{
+			bHasRealChange = true;
+		}
+	}
+
+	if (Children.Num() == 0)
+	{
+		// Make one child informing the user that there are no differences, or that it is unknown
+		if (bDiffSucceeded)
+		{
+			Children.Push(FBlueprintDifferenceTreeEntry::NoDifferencesEntry());
+		}
+		else
+		{
+			Children.Push(FBlueprintDifferenceTreeEntry::UnknownDifferencesEntry());
+		}
+	}
+
+	TSharedPtr<FBlueprintDifferenceTreeEntry> CategoryEntry = MakeShared<FBlueprintDifferenceTreeEntry>(
+			SelectionCallback,
+			FGenerateDiffEntryWidget::CreateSP(AsShared(), &FBlueprintTypeDiffControl::GenerateCategoryWidget, bHasRealChange),
+			Children);
+	OutTreeEntries.Push(CategoryEntry);
+
+	// Now add subobject diffs, one category per object
+	for (const TSharedPtr<FSubObjectDiff>& SubObjectDiff : SubObjectDiffs)
+	{
+		Children.Reset();
+
+		Children.Append(SubObjectDiff->Diffs);
+		OutRealDifferences.Append(SubObjectDiff->Diffs);
+
+		TSharedPtr<FBlueprintDifferenceTreeEntry> SubObjectEntry = FBlueprintDifferenceTreeEntry::CreateCategoryEntry(
+			SubObjectDiff->SourceResult.DisplayString,
+			SubObjectDiff->SourceResult.ToolTip,
+			FOnDiffEntryFocused::CreateSP(AsShared(), &FBlueprintTypeDiffControl::OnSelectSubobjectDiff, FPropertySoftPath(), SubObjectDiff),
+			Children,
+			true);
+
+		OutTreeEntries.Push(SubObjectEntry);
+	}
+}
+
+TSharedRef<SWidget> FBlueprintTypeDiffControl::GenerateCategoryWidget(bool bHasRealDiffs)
+{
+	FLinearColor Color = FLinearColor::White;
+
+	if (bHasRealDiffs)
+	{
+		Color = DiffViewUtils::Differs();
+	}
+
+	FText Label = BlueprintNew->GetClass()->GetDisplayNameText();
+
+	return SNew(SHorizontalBox)
+		+ SHorizontalBox::Slot()
+		[
+			SNew(STextBlock)
+			.ColorAndOpacity(Color)
+		.Text(Label)
+		];
+}
+
+void FBlueprintTypeDiffControl::BuildDiffSourceArray()
+{
+	TArray<FDiffSingleResult> BlueprintDiffResults;
+	FDiffResults BlueprintDiffs(&BlueprintDiffResults);
+	if (BlueprintNew->FindDiffs(BlueprintOld, BlueprintDiffs))
+	{
+		bDiffSucceeded = true;
+
+		// Add manual diffs
+		for (const FDiffSingleResult& CurrentDiff : BlueprintDiffResults)
+		{
+			if (CurrentDiff.Diff == EDiffType::OBJECT_REQUEST_DIFF)
+			{
+				// Turn into a subobject diff
+
+				// Invert order, we want old then new
+				TSharedPtr<FSubObjectDiff> SubObjectDiff = MakeShared<FSubObjectDiff>(CurrentDiff, CurrentDiff.Object2, CurrentDiff.Object1);
+
+				TArray<FSingleObjectDiffEntry> DifferingProperties;
+				SubObjectDiff->OldDetails.DiffAgainst(SubObjectDiff->NewDetails, DifferingProperties, true);
+
+				if (DifferingProperties.Num() > 0)
+				{
+					// Actual differences, so add to tree
+					SubObjectDiffs.Add(SubObjectDiff);
+
+					for (const FSingleObjectDiffEntry& Difference : DifferingProperties)
+					{
+						TSharedPtr<FBlueprintDifferenceTreeEntry> Entry = MakeShared<FBlueprintDifferenceTreeEntry>(
+							FOnDiffEntryFocused::CreateSP(AsShared(), &FBlueprintTypeDiffControl::OnSelectSubobjectDiff, Difference.Identifier, SubObjectDiff),
+							FGenerateDiffEntryWidget::CreateStatic(&GenerateObjectDiffWidget, Difference, RightRevision));
+						SubObjectDiff->Diffs.Push(Entry);
+					}
+				}
+			}
+			else
+			{
+				DiffListSource.Add(MakeShared<FDiffResultItem>(CurrentDiff));
+			}
+		}
+
+		struct SortDiff
+		{
+			bool operator () (const TSharedPtr<FDiffResultItem>& A, const TSharedPtr<FDiffResultItem>& B) const
+			{
+				return A->Result.Diff < B->Result.Diff;
+			}
+		};
+
+		Sort(DiffListSource.GetData(), DiffListSource.Num(), SortDiff());
+	}
+}
+
+void FBlueprintTypeDiffControl::OnSelectSubobjectDiff(FPropertySoftPath Identifier, TSharedPtr<FSubObjectDiff> SubObjectDiff)
+{
+	// This allows the owning control to focus the correct tab (or do whatever else it likes):
+	SelectionCallback.ExecuteIfBound();
+
+	if (SubObjectDiff.IsValid())
+	{
+		SubObjectDiff->OldDetails.HighlightProperty(Identifier);
+		SubObjectDiff->NewDetails.HighlightProperty(Identifier);
+
+		OldDetailsBox->SetContent(SubObjectDiff->OldDetails.DetailsWidget());
+		NewDetailsBox->SetContent(SubObjectDiff->NewDetails.DetailsWidget());
+	}
+}
+
+
+/////////////////////////////////////////////////////////////////////////////
+/// FGraphToDiff
+
+FGraphToDiff::FGraphToDiff(SBlueprintDiff* InDiffWidget, UEdGraph* InGraphOld, UEdGraph* InGraphNew, const FRevisionInfo& InRevisionOld, const FRevisionInfo& InRevisionNew)
+	: FoundDiffs(MakeShared<TArray<FDiffSingleResult>>()), DiffWidget(InDiffWidget), GraphOld(InGraphOld), GraphNew(InGraphNew), RevisionOld(InRevisionOld), RevisionNew(InRevisionNew)
+{
+	check(InGraphOld || InGraphNew); //one of them needs to exist
+
+	//need to know when it is modified
+	if (InGraphNew)
+	{
+		OnGraphChangedDelegateHandle = InGraphNew->AddOnGraphChangedHandler( FOnGraphChanged::FDelegate::CreateRaw(this, &FGraphToDiff::OnGraphChanged));
+	}
+
+	BuildDiffSourceArray();
+}
+
+FGraphToDiff::~FGraphToDiff()
+{
+	if (GraphNew)
+	{
+		GraphNew->RemoveOnGraphChangedHandler( OnGraphChangedDelegateHandle);
+	}
+}
+
+void FGraphToDiff::GenerateTreeEntries(TArray< TSharedPtr<FBlueprintDifferenceTreeEntry> >& OutTreeEntries, TArray< TSharedPtr<FBlueprintDifferenceTreeEntry> >& OutRealDifferences)
+{
+	TArray< TSharedPtr<FBlueprintDifferenceTreeEntry> > Children;
+	for (const TSharedPtr<FDiffResultItem>& Difference : DiffListSource)
+	{
+		TSharedPtr<FBlueprintDifferenceTreeEntry> ChildEntry = MakeShared<FBlueprintDifferenceTreeEntry>(
+				FOnDiffEntryFocused::CreateRaw(DiffWidget, &SBlueprintDiff::OnDiffListSelectionChanged, Difference),
+				FGenerateDiffEntryWidget::CreateSP(Difference.ToSharedRef(), &FDiffResultItem::GenerateWidget));
+		Children.Push(ChildEntry);
+		OutRealDifferences.Push(ChildEntry);
+	}
+
+	if (Children.Num() == 0)
+	{
+		// make one child informing the user that there are no differences:
+		Children.Push(FBlueprintDifferenceTreeEntry::NoDifferencesEntry());
+	}
+
+	TSharedPtr<FBlueprintDifferenceTreeEntry> Entry = MakeShared<FBlueprintDifferenceTreeEntry>(
+			FOnDiffEntryFocused::CreateRaw(DiffWidget, &SBlueprintDiff::OnGraphSelectionChanged, TSharedPtr<FGraphToDiff>(AsShared()), ESelectInfo::Direct),
+			FGenerateDiffEntryWidget::CreateSP(AsShared(), &FGraphToDiff::GenerateCategoryWidget),
+			Children);
+	OutTreeEntries.Push(Entry);
+}
+
+UEdGraph* FGraphToDiff::GetGraphOld() const
+{
+	return GraphOld;
+}
+
+UEdGraph* FGraphToDiff::GetGraphNew() const
+{
+	return GraphNew;
+}
+
+FText FGraphToDiff::GetToolTip()
+{
+	if (GraphOld && GraphNew)
+	{
+		if (DiffListSource.Num() > 0)
+		{
+			return LOCTEXT("ContainsDifferences", "Revisions are different");
+		}
+		else
+		{
+			return LOCTEXT("GraphsIdentical", "Revisions appear to be identical");
+		}
+	}
+	else
+	{
+		UEdGraph* GoodGraph = GraphOld ? GraphOld : GraphNew;
+		check(GoodGraph);
+		const FRevisionInfo& Revision = GraphNew ? RevisionOld : RevisionNew;
+		FText RevisionText = LOCTEXT("CurrentRevision", "Current Revision");
+
+		if (!Revision.Revision.IsEmpty())
+		{
+			RevisionText = FText::Format(LOCTEXT("Revision Number", "Revision {0}"), FText::FromString(Revision.Revision));
+		}
+
+		return FText::Format(LOCTEXT("MissingGraph", "Graph '{0}' missing from {1}"), FText::FromString(GoodGraph->GetName()), RevisionText);
+	}
+}
+
+TSharedRef<SWidget> FGraphToDiff::GenerateCategoryWidget()
+{
+	const UEdGraph* Graph = GraphOld ? GraphOld : GraphNew;
+	check(Graph);
+	
+	FLinearColor Color = (GraphOld && GraphNew) ? FLinearColor::White : FLinearColor(0.3f,0.3f,1.f);
+
+	const bool bHasDiffs = DiffListSource.Num() > 0;
+
+	if (bHasDiffs)
+	{
+		Color = DiffViewUtils::Differs();
+	}
+
+	FText GraphName;
+	if (const UEdGraphSchema* Schema = Graph->GetSchema())
+	{
+		FGraphDisplayInfo DisplayInfo;
+		Schema->GetGraphDisplayInformation(*Graph, DisplayInfo);
+
+		GraphName = DisplayInfo.DisplayName;
+	}
+	else
+	{
+		GraphName = FText::FromName(Graph->GetFName());
+	}
+
+	return SNew(SHorizontalBox)
+	+ SHorizontalBox::Slot()
+	[
+		SNew(STextBlock)
+		.ColorAndOpacity(Color)
+		.Text(GraphName)
+		.ToolTipText(GetToolTip())
+	]
+	+ DiffViewUtils::Box( GraphOld != nullptr, Color )
+	+ DiffViewUtils::Box( GraphNew != nullptr, Color );
+}
+
+void FGraphToDiff::BuildDiffSourceArray()
+{
+	FoundDiffs->Empty();
+	FGraphDiffControl::DiffGraphs(GraphOld, GraphNew, *FoundDiffs);
+
+	struct SortDiff
+	{
+		bool operator () (const FDiffSingleResult& A, const FDiffSingleResult& B) const
+		{
+			return A.Diff < B.Diff;
+		}
+	};
+
+	Sort(FoundDiffs->GetData(), FoundDiffs->Num(), SortDiff());
+
+	DiffListSource.Empty();
+	for (const FDiffSingleResult& Diff : *FoundDiffs)
+	{
+		DiffListSource.Add(MakeShared<FDiffResultItem>(Diff));
+	}
+}
+
+void FGraphToDiff::OnGraphChanged( const FEdGraphEditAction& Action )
+{
+	DiffWidget->OnGraphChanged(this);
+}
+
+#undef LOCTEXT_NAMESPACE
\ No newline at end of file
diff --git a/Engine/Source/Editor/Kismet/Private/DiffControl.h b/Engine/Source/Editor/Kismet/Private/DiffControl.h
new file mode 100644
--- /dev/null
+++ b/Engine/Source/Editor/Kismet/Private/DiffControl.h
@@ -0,0 +1,256 @@
+// Copyright Epic Games, Inc. All Rights Reserved.
+#pragma once
+
+#include "SCSDiff.h"
+#include "CoreMinimal.h"
+#include "DetailsDiff.h"
+#include "DiffResults.h"
+#include "IAssetTypeActions.h"
+#include "Algo/Transform.h"
+
+struct FDiffResultItem;
+class SBlueprintDiff;
+class UEdGraph;
+struct FEdGraphEditAction;
+
+/** Interface responsible for generating FBlueprintDifferenceTreeEntry's for visual diff tools */
+class IDiffControl
+{
+public:
+	virtual ~IDiffControl() {}
+
+	/** Adds widgets to the tree of differences to show */
+	virtual void GenerateTreeEntries(TArray< TSharedPtr<FBlueprintDifferenceTreeEntry> >& OutTreeEntries, TArray< TSharedPtr<FBlueprintDifferenceTreeEntry> >& OutRealDifferences) = 0;
+	
+	static FText RightRevision;
+	static TSharedRef<SWidget> GenerateSimpleDiffWidget(FText DiffText);
+	static TSharedRef<SWidget> GenerateObjectDiffWidget(FSingleObjectDiffEntry DiffEntry, FText ObjectName);
+	
+};
+
+/** Shows all differences for the blueprint structure itself that aren't picked up elsewhere */
+class FMyBlueprintDiffControl : public TSharedFromThis<FMyBlueprintDiffControl>, public IDiffControl
+{
+public:
+	FMyBlueprintDiffControl(const UBlueprint* InOldBlueprint, const UBlueprint* InNewBlueprint, FOnDiffEntryFocused InSelectionCallback);
+
+	virtual void GenerateTreeEntries(TArray< TSharedPtr<FBlueprintDifferenceTreeEntry> >& OutTreeEntries, TArray< TSharedPtr<FBlueprintDifferenceTreeEntry> >& OutRealDifferences) override;
+
+private:
+	FOnDiffEntryFocused SelectionCallback;
+	const UBlueprint* OldBlueprint;
+	const UBlueprint* NewBlueprint;
+};
+
+/** 
+ * Each difference in the tree will either be a tree node that is added in one Blueprint 
+ * or a tree node and an FName of a property that has been added or edited in one Blueprint
+ */
+class FSCSDiffControl : public TSharedFromThis<FSCSDiffControl>, public IDiffControl
+{
+public:
+	FSCSDiffControl(const UBlueprint* InOldBlueprint, const UBlueprint* InNewBlueprint, FOnDiffEntryFocused InSelectionCallback);
+
+	virtual void GenerateTreeEntries(TArray< TSharedPtr<FBlueprintDifferenceTreeEntry> >& OutTreeEntries, TArray< TSharedPtr<FBlueprintDifferenceTreeEntry> >& OutRealDifferences) override;
+
+	TSharedRef<SWidget> OldTreeWidget() { return OldSCS.TreeWidget(); }
+	TSharedRef<SWidget> NewTreeWidget() { return NewSCS.TreeWidget(); }
+
+private:
+	FOnDiffEntryFocused SelectionCallback;
+	FSCSDiffRoot DifferingProperties;
+
+	FSCSDiff OldSCS;
+	FSCSDiff NewSCS;
+};
+
+/** Generic wrapper around a details view template parameter determines whether TreeEntries are populated */
+template<bool bPopulateOutTreeEntries>
+class TDetailsDiffControl : public TSharedFromThis<TDetailsDiffControl<bPopulateOutTreeEntries>>, public IDiffControl
+{
+public:
+	TDetailsDiffControl(const UObject* InOldObject, const UObject* InNewObject, FOnDiffEntryFocused InSelectionCallback)
+		: SelectionCallback(InSelectionCallback)
+		, OldDetails(InOldObject, FDetailsDiff::FOnDisplayedPropertiesChanged())
+		, NewDetails(InNewObject, FDetailsDiff::FOnDisplayedPropertiesChanged())
+	{
+		OldDetails.DiffAgainst(NewDetails, DifferingProperties, true);
+
+		TSet<FPropertyPath> PropertyPaths;
+		Algo::Transform(DifferingProperties, PropertyPaths,
+			[&InOldObject](const FSingleObjectDiffEntry& DiffEntry)
+			{
+				return DiffEntry.Identifier.ResolvePath(InOldObject);
+			});
+
+		OldDetails.DetailsWidget()->UpdatePropertyAllowList(PropertyPaths);
+
+		PropertyPaths.Reset();
+		Algo::Transform(DifferingProperties, PropertyPaths,
+			[&InNewObject](const FSingleObjectDiffEntry& DiffEntry)
+			{
+				return DiffEntry.Identifier.ResolvePath(InNewObject);
+			});
+
+		NewDetails.DetailsWidget()->UpdatePropertyAllowList(PropertyPaths);
+	}
+
+	virtual void GenerateTreeEntries(TArray< TSharedPtr<FBlueprintDifferenceTreeEntry> >& OutTreeEntries, TArray< TSharedPtr<FBlueprintDifferenceTreeEntry> >& OutRealDifferences) override
+	{
+		for (const FSingleObjectDiffEntry& Difference : DifferingProperties)
+		{
+			
+			TSharedPtr<FBlueprintDifferenceTreeEntry> Entry = MakeShared<FBlueprintDifferenceTreeEntry>(
+				FOnDiffEntryFocused::CreateSP(TSharedFromThis<TDetailsDiffControl<bPopulateOutTreeEntries>>::AsShared(), &TDetailsDiffControl::OnSelectDiffEntry, Difference.Identifier),
+				FGenerateDiffEntryWidget::CreateStatic(&GenerateObjectDiffWidget, Difference, RightRevision));
+			Children.Push(Entry);
+			OutRealDifferences.Push(Entry);
+			if constexpr (bPopulateOutTreeEntries)
+			{
+				OutTreeEntries.Push(Entry);
+			}
+		}
+	}
+
+	TSharedRef<SWidget> OldDetailsWidget() { return OldDetails.DetailsWidget(); }
+	TSharedRef<SWidget> NewDetailsWidget() { return NewDetails.DetailsWidget(); }
+
+protected:
+	virtual void OnSelectDiffEntry(FPropertySoftPath PropertyName)
+	{
+		SelectionCallback.ExecuteIfBound();
+		OldDetails.HighlightProperty(PropertyName);
+		NewDetails.HighlightProperty(PropertyName);
+	}
+
+	FOnDiffEntryFocused SelectionCallback;
+	FDetailsDiff OldDetails;
+	FDetailsDiff NewDetails;
+
+	TArray<FSingleObjectDiffEntry> DifferingProperties;
+	TArray< TSharedPtr<FBlueprintDifferenceTreeEntry> > Children;
+};
+
+using FDetailsDiffControl = TDetailsDiffControl<true>;
+
+/** Override for CDO special case */
+class FCDODiffControl : public TDetailsDiffControl<false>
+{
+public:
+	FCDODiffControl(const UObject* InOldObject, const UObject* InNewObject, FOnDiffEntryFocused InSelectionCallback);
+
+	virtual void GenerateTreeEntries(TArray< TSharedPtr<FBlueprintDifferenceTreeEntry> >& OutTreeEntries, TArray< TSharedPtr<FBlueprintDifferenceTreeEntry> >& OutRealDifferences) override;
+};
+
+/** Override for class class settings */
+class FClassSettingsDiffControl : public TDetailsDiffControl<false>
+{
+public:
+	FClassSettingsDiffControl(const UObject* InOldObject, const UObject* InNewObject, FOnDiffEntryFocused InSelectionCallback);
+
+	virtual void GenerateTreeEntries(TArray< TSharedPtr<FBlueprintDifferenceTreeEntry> >& OutTreeEntries, TArray< TSharedPtr<FBlueprintDifferenceTreeEntry> >& OutRealDifferences) override;
+};
+
+/** Diff control to handle finding type-specific differences */
+struct FBlueprintTypeDiffControl : public TSharedFromThis<FBlueprintTypeDiffControl>, public IDiffControl
+{
+	struct FSubObjectDiff
+	{
+		FDiffSingleResult SourceResult;
+		FDetailsDiff OldDetails;
+		FDetailsDiff NewDetails;
+		TArray<TSharedPtr<FBlueprintDifferenceTreeEntry>> Diffs;
+
+		FSubObjectDiff(const FDiffSingleResult& InSourceResult, const UObject* OldObject, const UObject* NewObject)
+			: SourceResult(InSourceResult)
+			, OldDetails(OldObject, FDetailsDiff::FOnDisplayedPropertiesChanged())
+			, NewDetails(NewObject, FDetailsDiff::FOnDisplayedPropertiesChanged())
+		{}
+	};
+
+	FBlueprintTypeDiffControl(const UBlueprint* InBlueprintOld, const UBlueprint* InBlueprintNew, FOnDiffEntryFocused InSelectionCallback);
+
+	/** Generate difference tree widgets */
+	virtual void GenerateTreeEntries(TArray< TSharedPtr<FBlueprintDifferenceTreeEntry> >& OutTreeEntries, TArray< TSharedPtr<FBlueprintDifferenceTreeEntry> >& OutRealDifferences) override;
+
+	/** The old blueprint (left) */
+	const UBlueprint* BlueprintOld;
+
+	/** The new blueprint(right) */
+	const UBlueprint* BlueprintNew;
+
+	/** Boxes that will display the details diffs */
+	TSharedPtr<SBox> OldDetailsBox;
+	TSharedPtr<SBox> NewDetailsBox;
+
+private:
+	/** Generate Widget for top category */
+	TSharedRef<SWidget> GenerateCategoryWidget(bool bHasRealDiffs);
+
+	/** Build up the Diff Source Array*/
+	void BuildDiffSourceArray();
+
+	/** Handle selecting a diff */
+	void OnSelectSubobjectDiff(FPropertySoftPath Identifier, TSharedPtr<FSubObjectDiff> SubObjectDiff);
+
+	/** List of objects with differences */
+	TArray<TSharedPtr<FSubObjectDiff>> SubObjectDiffs;
+
+	/** Source for list view */
+	TArray<TSharedPtr<FDiffResultItem>> DiffListSource;
+
+	/** Selection callback */
+	FOnDiffEntryFocused SelectionCallback;
+
+	/** Did diff generation succeed? */
+	bool bDiffSucceeded;
+};
+
+
+/** Category list item for a graph*/
+struct FGraphToDiff	: public TSharedFromThis<FGraphToDiff>, IDiffControl
+{
+	FGraphToDiff(SBlueprintDiff* DiffWidget, UEdGraph* GraphOld, UEdGraph* GraphNew, const FRevisionInfo& RevisionOld, const FRevisionInfo& RevisionNew);
+	virtual ~FGraphToDiff() override;
+
+	/** Add widgets to the differences tree */
+	virtual void GenerateTreeEntries(TArray< TSharedPtr<FBlueprintDifferenceTreeEntry> >& OutTreeEntries, TArray< TSharedPtr<FBlueprintDifferenceTreeEntry> >& OutRealDifferences) override;
+
+	/** Get old(left) graph*/
+	UEdGraph* GetGraphOld() const;
+
+	/** Get new(right) graph*/
+	UEdGraph* GetGraphNew() const;
+
+	/** Source for list view */
+	TArray<TSharedPtr<FDiffResultItem>> DiffListSource;
+	TSharedPtr<TArray<FDiffSingleResult>> FoundDiffs;
+
+private:
+	/** Get tooltip for category */
+	FText GetToolTip();
+
+	/** Generate Widget for category list */
+	TSharedRef<SWidget> GenerateCategoryWidget();
+
+	/** Called when the Newer Graph is modified*/
+	void OnGraphChanged(const FEdGraphEditAction& Action);
+
+	/** Build up the Diff Source Array*/
+	void BuildDiffSourceArray();
+
+	/** Diff widget */
+	class SBlueprintDiff* DiffWidget;
+
+	/** The old graph(left)*/
+	UEdGraph* GraphOld;
+
+	/** The new graph(right)*/
+	UEdGraph* GraphNew;
+
+	/** Description of Old and new graph*/
+	FRevisionInfo	RevisionOld, RevisionNew;
+
+	/** Handle to the registered OnGraphChanged delegate. */
+	FDelegateHandle OnGraphChangedDelegateHandle;
+};
\ No newline at end of file
diff --git a/Engine/Source/Editor/Kismet/Private/DiffUtils.cpp b/Engine/Source/Editor/Kismet/Private/DiffUtils.cpp
--- a/Engine/Source/Editor/Kismet/Private/DiffUtils.cpp
+++ b/Engine/Source/Editor/Kismet/Private/DiffUtils.cpp
@@ -375,7 +375,7 @@ void DiffUtils::CompareUnrelatedSCS(const UBlueprint* Old, const TArray< FSCSRes
 	{
 		for (const auto& Node : InArray)
 		{
-			if (Node.Identifier.Name == Value->Name )
+			if (Node.Identifier.Name == Value->Name)
 			{
 				return &Node;
 			}
@@ -1004,7 +1004,7 @@ FText DiffViewUtils::SCSDiffMessage(const FSCSDiffEntry& Difference, FText Objec
 	return Text;
 }
 
-FText DiffViewUtils::GetPanelLabel(const UBlueprint* Blueprint, const FRevisionInfo& Revision, FText Label )
+FText DiffViewUtils::GetPanelLabel(const UObject* Asset, const FRevisionInfo& Revision, FText Label)
 {
 	if( !Revision.Revision.IsEmpty() )
 	{
@@ -1027,32 +1027,32 @@ FText DiffViewUtils::GetPanelLabel(const UBlueprint* Blueprint, const FRevisionI
 		if (Label.IsEmpty())
 		{
 			return FText::Format(NSLOCTEXT("DiffViewUtils", "RevisionLabelTwoLines", "{0}\n{1}")
-				, FText::FromString(Blueprint->GetName())
+				, FText::FromString(Asset->GetName())
 				, RevisionData);
 		}
 		else
 		{
 			return FText::Format(NSLOCTEXT("DiffViewUtils", "RevisionLabel", "{0}\n{1}\n{2}")
 				, Label
-				, FText::FromString(Blueprint->GetName())
+				, FText::FromString(Asset->GetName())
 				, RevisionData);
 		}
 	}
 	else
 	{
-		if( Blueprint )
+		if( Asset )
 		{
 			if (Label.IsEmpty())
 			{
 				return FText::Format(NSLOCTEXT("DiffViewUtils", "RevisionLabelTwoLines", "{0}\n{1}")
-					, FText::FromString(Blueprint->GetName())
+					, FText::FromString(Asset->GetName())
 					, NSLOCTEXT("DiffViewUtils", "LocalRevisionLabel", "Local Revision"));
 			}
 			else
 			{
 				return FText::Format(NSLOCTEXT("DiffViewUtils", "RevisionLabel", "{0}\n{1}\n{2}")
 					, Label
-					, FText::FromString(Blueprint->GetName())
+					, FText::FromString(Asset->GetName())
 					, NSLOCTEXT("DiffViewUtils", "LocalRevisionLabel", "Local Revision"));
 			}
 		}
diff --git a/Engine/Source/Editor/Kismet/Private/SBlueprintDiff.cpp b/Engine/Source/Editor/Kismet/Private/SBlueprintDiff.cpp
--- a/Engine/Source/Editor/Kismet/Private/SBlueprintDiff.cpp
+++ b/Engine/Source/Editor/Kismet/Private/SBlueprintDiff.cpp
@@ -3,7 +3,6 @@
 #include "SBlueprintDiff.h"
 #include "Widgets/Layout/SSplitter.h"
 #include "Widgets/SOverlay.h"
-#include "Engine/GameViewportClient.h"
 #include "SlateOptMacros.h"
 #include "Widgets/Layout/SSpacer.h"
 #include "Framework/MultiBox/MultiBoxDefs.h"
@@ -13,42 +12,29 @@
 #include "K2Node_MathExpression.h"
 #include "Kismet2/BlueprintEditorUtils.h"
 #include "Kismet2/KismetEditorUtilities.h"
-#include "BlueprintEditorModes.h"
 #include "DetailsDiff.h"
 #include "EdGraphUtilities.h"
 #include "GraphDiffControl.h"
 #include "SMyBlueprint.h"
-#include "SCSDiff.h"
 #include "WorkflowOrientedApp/SModeWidget.h"
 #include "Framework/Commands/GenericCommands.h"
 #include "WidgetBlueprint.h"
 #include "HAL/PlatformApplicationMisc.h"
 #include "Framework/Application/SlateApplication.h"
-#include "Algo/Transform.h"
+#include "DiffControl.h"
 
 #include "Subsystems/AssetEditorSubsystem.h"
 
 #define LOCTEXT_NAMESPACE "SBlueprintDif"
 
-class IDiffControl
-{
-public:
-	virtual ~IDiffControl() {}
-
-	/** Adds widgets to the tree of differences to show */
-	virtual void GenerateTreeEntries(TArray< TSharedPtr<FBlueprintDifferenceTreeEntry> >& OutTreeEntries, TArray< TSharedPtr<FBlueprintDifferenceTreeEntry> >& OutRealDifferences) = 0;
-};
-
-FText RightRevision = LOCTEXT("OlderRevisionIdentifier", "Right Revision");
-
 typedef TMap< FName, const FProperty* > FNamePropertyMap;
 
-const FName BlueprintTypeMode = FName(TEXT("BlueprintTypeMode"));
-const FName MyBlueprintMode = FName(TEXT("MyBlueprintMode"));
-const FName DefaultsMode = FName(TEXT("DefaultsMode"));
-const FName ClassSettingsMode = FName(TEXT("ClassSettingsMode"));
-const FName ComponentsMode = FName(TEXT("ComponentsMode"));
-const FName GraphMode = FName(TEXT("GraphMode"));
+static const FName BlueprintTypeMode = FName(TEXT("BlueprintTypeMode"));
+static const FName MyBlueprintMode = FName(TEXT("MyBlueprintMode"));
+static const FName DefaultsMode = FName(TEXT("DefaultsMode"));
+static const FName ClassSettingsMode = FName(TEXT("ClassSettingsMode"));
+static const FName ComponentsMode = FName(TEXT("ComponentsMode"));
+static const FName GraphMode = FName(TEXT("GraphMode"));
 
 TSharedRef<SWidget>	FDiffResultItem::GenerateWidget() const
 {
@@ -66,759 +52,6 @@ TSharedRef<SWidget>	FDiffResultItem::GenerateWidget() const
 		.Text(Text);
 }
 
-static TSharedRef<SWidget> GenerateObjectDiffWidget(FSingleObjectDiffEntry DiffEntry, FText ObjectName)
-{
-	return SNew(STextBlock)
-		.Text(DiffViewUtils::PropertyDiffMessage(DiffEntry, ObjectName))
-		.ToolTipText(DiffViewUtils::PropertyDiffMessage(DiffEntry, ObjectName))
-		.ColorAndOpacity(DiffViewUtils::Differs());
-}
-
-static TSharedRef<SWidget> GenerateSimpleDiffWidget(FText DiffText)
-{
-	return SNew(STextBlock)
-		.Text(DiffText)
-		.ToolTipText(DiffText)
-		.ColorAndOpacity(DiffViewUtils::Differs());
-};
-
-/** Shows all differences for the blueprint structure itself that aren't picked up elsewhere */
-class FMyBlueprintDiffControl : public TSharedFromThis<FMyBlueprintDiffControl>, public IDiffControl
-{
-public:
-	FMyBlueprintDiffControl(const UBlueprint* InOldBlueprint, const UBlueprint* InNewBlueprint, FOnDiffEntryFocused InSelectionCallback)
-		: SelectionCallback(MoveTemp(InSelectionCallback)), OldBlueprint(InOldBlueprint), NewBlueprint(InNewBlueprint)
-	{
-	}
-
-	virtual void GenerateTreeEntries(TArray< TSharedPtr<FBlueprintDifferenceTreeEntry> >& OutTreeEntries, TArray< TSharedPtr<FBlueprintDifferenceTreeEntry> >& OutRealDifferences) override
-	{
-		TArray< TSharedPtr<FBlueprintDifferenceTreeEntry> > Children;
-
-		if (OldBlueprint && OldBlueprint->SkeletonGeneratedClass && NewBlueprint && NewBlueprint->SkeletonGeneratedClass)
-		{
-			for (TFieldIterator<FProperty> PropertyIt(OldBlueprint->SkeletonGeneratedClass); PropertyIt; ++PropertyIt)
-			{
-				FProperty* OldProperty = *PropertyIt;
-				FProperty* NewProperty = NewBlueprint->SkeletonGeneratedClass->FindPropertyByName(OldProperty->GetFName());
-
-				FText PropertyText = FText::FromString(OldProperty->GetAuthoredName());
-
-				if (NewProperty)
-				{
-					const int32 OldVarIndex = FBlueprintEditorUtils::FindNewVariableIndex(OldBlueprint, OldProperty->GetFName());
-					const int32 NewVarIndex = FBlueprintEditorUtils::FindNewVariableIndex(NewBlueprint, OldProperty->GetFName());
-
-					if (OldVarIndex != INDEX_NONE && NewVarIndex != INDEX_NONE)
-					{
-						TArray<FSingleObjectDiffEntry> DifferingProperties;
-						DiffUtils::CompareUnrelatedStructs(FBPVariableDescription::StaticStruct(), &OldBlueprint->NewVariables[OldVarIndex], FBPVariableDescription::StaticStruct(), &NewBlueprint->NewVariables[NewVarIndex], DifferingProperties);
-						for (const FSingleObjectDiffEntry& Difference : DifferingProperties)
-						{
-							TSharedPtr<FBlueprintDifferenceTreeEntry> Entry = MakeShared<FBlueprintDifferenceTreeEntry>(
-								SelectionCallback,
-								FGenerateDiffEntryWidget::CreateStatic(&GenerateObjectDiffWidget, Difference, PropertyText));
-							Children.Push(Entry);
-							OutRealDifferences.Push(Entry);
-						}
-					}	
-				}
-				else
-				{
-					FText DiffText = FText::Format(LOCTEXT("VariableRemoved", "Removed Variable {0}"), PropertyText);
-
-					TSharedPtr<FBlueprintDifferenceTreeEntry> Entry = MakeShared<FBlueprintDifferenceTreeEntry>(
-						SelectionCallback,
-						FGenerateDiffEntryWidget::CreateStatic(&GenerateSimpleDiffWidget, DiffText));
-
-					Children.Push(Entry);
-					OutRealDifferences.Push(Entry);
-				}
-			}
-
-			for (TFieldIterator<FProperty> PropertyIt(NewBlueprint->SkeletonGeneratedClass); PropertyIt; ++PropertyIt)
-			{
-				FProperty* NewProperty = *PropertyIt;
-				FProperty* OldProperty = OldBlueprint->SkeletonGeneratedClass->FindPropertyByName(NewProperty->GetFName());
-
-				if (!OldProperty)
-				{
-					FText DiffText = FText::Format(LOCTEXT("VariableAdded", "Added Variable {0}"), FText::FromString(NewProperty->GetAuthoredName()));
-
-					TSharedPtr<FBlueprintDifferenceTreeEntry> Entry = MakeShared<FBlueprintDifferenceTreeEntry>(
-						SelectionCallback,
-						FGenerateDiffEntryWidget::CreateStatic(&GenerateSimpleDiffWidget, DiffText));
-
-					Children.Push(Entry);
-					OutRealDifferences.Push(Entry);
-				}
-			}
-		}
-		const bool bHasDifferences = Children.Num() != 0;
-		if (!bHasDifferences)
-		{
-			// make one child informing the user that there are no differences:
-			Children.Push(FBlueprintDifferenceTreeEntry::NoDifferencesEntry());
-		}
-
-		OutTreeEntries.Push(FBlueprintDifferenceTreeEntry::CreateCategoryEntry(
-			NSLOCTEXT("FBlueprintDifferenceTreeEntry", "MyBlueprintLabel", "My Blueprint"),
-			NSLOCTEXT("FBlueprintDifferenceTreeEntry", "MyBlueprintTooltip", "The list of changes made to blueprint structure in the My Blueprint panel"),
-			SelectionCallback,
-			Children,
-			bHasDifferences
-		));
-	}
-
-private:
-	FOnDiffEntryFocused SelectionCallback;
-	const UBlueprint* OldBlueprint;
-	const UBlueprint* NewBlueprint;
-};
-
-/** 
- * Each difference in the tree will either be a tree node that is added in one Blueprint 
- * or a tree node and an FName of a property that has been added or edited in one Blueprint
- */
-class FSCSDiffControl : public TSharedFromThis<FSCSDiffControl>, public IDiffControl
-{
-public:
-	FSCSDiffControl(const UBlueprint* InOldBlueprint, const UBlueprint* InNewBlueprint, FOnDiffEntryFocused InSelectionCallback)
-		: SelectionCallback(InSelectionCallback)
-		, OldSCS(InOldBlueprint)
-		, NewSCS(InNewBlueprint)
-	{
-	}
-
-	virtual void GenerateTreeEntries(TArray< TSharedPtr<FBlueprintDifferenceTreeEntry> >& OutTreeEntries, TArray< TSharedPtr<FBlueprintDifferenceTreeEntry> >& OutRealDifferences) override
-	{
-		TArray< FSCSResolvedIdentifier > OldHierarchy = OldSCS.GetDisplayedHierarchy();
-		TArray< FSCSResolvedIdentifier > NewHierarchy = NewSCS.GetDisplayedHierarchy();
-		DiffUtils::CompareUnrelatedSCS(OldSCS.GetBlueprint(), OldHierarchy, NewSCS.GetBlueprint(), NewHierarchy, DifferingProperties);
-
-		const auto FocusSCSDifferenceEntry = [](FSCSDiffEntry Entry, FOnDiffEntryFocused InSelectionCallback, FSCSDiffControl* Owner)
-		{
-			InSelectionCallback.ExecuteIfBound();
-			if (Entry.TreeIdentifier.Name != NAME_None)
-			{
-				Owner->OldSCS.HighlightProperty(Entry.TreeIdentifier.Name, FPropertyPath());
-				Owner->NewSCS.HighlightProperty(Entry.TreeIdentifier.Name, FPropertyPath());
-			}
-		};
-
-		const auto CreateSCSDifferenceWidget = [](FSCSDiffEntry Entry, FText ObjectName) -> TSharedRef<SWidget>
-		{
-			return SNew(STextBlock)
-				.Text(DiffViewUtils::SCSDiffMessage(Entry, ObjectName))
-				.ColorAndOpacity(DiffViewUtils::Differs());
-		};
-
-		TArray< TSharedPtr<FBlueprintDifferenceTreeEntry> > Children;
-		for (const FSCSDiffEntry& Difference : DifferingProperties.Entries)
-		{
-			TSharedPtr<FBlueprintDifferenceTreeEntry> Entry = MakeShared<FBlueprintDifferenceTreeEntry>(
-				FOnDiffEntryFocused::CreateStatic(FocusSCSDifferenceEntry, Difference, SelectionCallback, this),
-				FGenerateDiffEntryWidget::CreateStatic(CreateSCSDifferenceWidget, Difference, RightRevision));
-			Children.Push(Entry);
-			OutRealDifferences.Push(Entry);
-		}
-
-		const bool bHasDifferences = Children.Num() != 0;
-		if (!bHasDifferences)
-		{
-			// make one child informing the user that there are no differences:
-			Children.Push(FBlueprintDifferenceTreeEntry::NoDifferencesEntry());
-		}
-
-		OutTreeEntries.Push(FBlueprintDifferenceTreeEntry::CreateCategoryEntry(
-			NSLOCTEXT("FBlueprintDifferenceTreeEntry", "SCSLabel", "Components"),
-			NSLOCTEXT("FBlueprintDifferenceTreeEntry", "SCSTooltip", "The list of changes made in the Components panel"),
-			SelectionCallback,
-			Children,
-			bHasDifferences
-		));
-	}
-
-	TSharedRef<SWidget> OldTreeWidget() { return OldSCS.TreeWidget(); }
-	TSharedRef<SWidget> NewTreeWidget() { return NewSCS.TreeWidget(); }
-
-private:
-	FOnDiffEntryFocused SelectionCallback;
-	FSCSDiffRoot DifferingProperties;
-
-	FSCSDiff OldSCS;
-	FSCSDiff NewSCS;
-};
-
-/** Generic wrapper around a details view, this does not actually fill out OutTreeEntries */
-class FDetailsDiffControl : public TSharedFromThis<FDetailsDiffControl>, public IDiffControl
-{
-public:
-	FDetailsDiffControl(const UObject* InOldObject, const UObject* InNewObject, FOnDiffEntryFocused InSelectionCallback)
-		: SelectionCallback(InSelectionCallback)
-		, OldDetails(InOldObject, FDetailsDiff::FOnDisplayedPropertiesChanged())
-		, NewDetails(InNewObject, FDetailsDiff::FOnDisplayedPropertiesChanged())
-	{
-		OldDetails.DiffAgainst(NewDetails, DifferingProperties, true);
-
-		TSet<FPropertyPath> PropertyPaths;
-		Algo::Transform(DifferingProperties, PropertyPaths,
-			[&InOldObject](const FSingleObjectDiffEntry& DiffEntry)
-			{
-				return DiffEntry.Identifier.ResolvePath(InOldObject);
-			});
-
-		OldDetails.DetailsWidget()->UpdatePropertyAllowList(PropertyPaths);
-
-		PropertyPaths.Reset();
-		Algo::Transform(DifferingProperties, PropertyPaths,
-			[&InNewObject](const FSingleObjectDiffEntry& DiffEntry)
-			{
-				return DiffEntry.Identifier.ResolvePath(InNewObject);
-			});
-
-		NewDetails.DetailsWidget()->UpdatePropertyAllowList(PropertyPaths);
-	}
-
-	virtual void GenerateTreeEntries(TArray< TSharedPtr<FBlueprintDifferenceTreeEntry> >& OutTreeEntries, TArray< TSharedPtr<FBlueprintDifferenceTreeEntry> >& OutRealDifferences) override
-	{
-		for (const FSingleObjectDiffEntry& Difference : DifferingProperties)
-		{
-			TSharedPtr<FBlueprintDifferenceTreeEntry> Entry = MakeShared<FBlueprintDifferenceTreeEntry>(
-				FOnDiffEntryFocused::CreateSP(AsShared(), &FDetailsDiffControl::OnSelectDiffEntry, Difference.Identifier),
-				FGenerateDiffEntryWidget::CreateStatic(&GenerateObjectDiffWidget, Difference, RightRevision));
-			Children.Push(Entry);
-			OutRealDifferences.Push(Entry);
-		}
-	}
-
-	TSharedRef<SWidget> OldDetailsWidget() { return OldDetails.DetailsWidget(); }
-	TSharedRef<SWidget> NewDetailsWidget() { return NewDetails.DetailsWidget(); }
-
-protected:
-	virtual void OnSelectDiffEntry(FPropertySoftPath PropertyName)
-	{
-		SelectionCallback.ExecuteIfBound();
-		OldDetails.HighlightProperty(PropertyName);
-		NewDetails.HighlightProperty(PropertyName);
-	}
-
-	FOnDiffEntryFocused SelectionCallback;
-	FDetailsDiff OldDetails;
-	FDetailsDiff NewDetails;
-
-	TArray<FSingleObjectDiffEntry> DifferingProperties;
-	TArray< TSharedPtr<FBlueprintDifferenceTreeEntry> > Children;
-};
-
-/** Override for CDO special case */
-class FCDODiffControl : public FDetailsDiffControl
-{
-public:
-	FCDODiffControl(const UObject* InOldObject, const UObject* InNewObject, FOnDiffEntryFocused InSelectionCallback)
-		: FDetailsDiffControl(InOldObject, InNewObject, InSelectionCallback)
-	{
-	}
-	
-	virtual void GenerateTreeEntries(TArray< TSharedPtr<FBlueprintDifferenceTreeEntry> >& OutTreeEntries, TArray< TSharedPtr<FBlueprintDifferenceTreeEntry> >& OutRealDifferences) override
-	{
-		FDetailsDiffControl::GenerateTreeEntries(OutTreeEntries, OutRealDifferences);
-
-		const bool bHasDifferences = Children.Num() != 0;
-		if (!bHasDifferences)
-		{
-			// make one child informing the user that there are no differences:
-			Children.Push(FBlueprintDifferenceTreeEntry::NoDifferencesEntry());
-		}
-
-		OutTreeEntries.Push(FBlueprintDifferenceTreeEntry::CreateCategoryEntry(
-			NSLOCTEXT("FBlueprintDifferenceTreeEntry", "DefaultsLabel", "Defaults"),
-			NSLOCTEXT("FBlueprintDifferenceTreeEntry", "DefaultsTooltip", "The list of changes made in the Defaults panel"),
-			SelectionCallback,
-			Children,
-			bHasDifferences
-		));
-	}
-};
-
-/** Override for class class settings */
-class FClassSettingsDiffControl : public FDetailsDiffControl
-{
-public:
-	FClassSettingsDiffControl(const UObject* InOldObject, const UObject* InNewObject, FOnDiffEntryFocused InSelectionCallback)
-		: FDetailsDiffControl(InOldObject, InNewObject, InSelectionCallback)
-	{
-	}
-
-	virtual void GenerateTreeEntries(TArray< TSharedPtr<FBlueprintDifferenceTreeEntry> >& OutTreeEntries, TArray< TSharedPtr<FBlueprintDifferenceTreeEntry> >& OutRealDifferences) override
-	{
-		FDetailsDiffControl::GenerateTreeEntries(OutTreeEntries, OutRealDifferences);
-
-		// Check for parent class and interfaces here
-		const UBlueprint* OldBlueprint = Cast<UBlueprint>(OldDetails.GetDisplayedObject());
-		const UBlueprint* NewBlueprint = Cast<UBlueprint>(NewDetails.GetDisplayedObject());
-
-		if (OldBlueprint && NewBlueprint)
-		{
-			if (OldBlueprint->ParentClass != NewBlueprint->ParentClass)
-			{
-				FText DiffText = FText::Format(LOCTEXT("ParentChanged", "Parent Class changed from {0} to {1}"), FText::FromString(GetNameSafe(OldBlueprint->ParentClass)), FText::FromString(GetNameSafe(NewBlueprint->ParentClass)));
-
-				TSharedPtr<FBlueprintDifferenceTreeEntry> Entry = MakeShared<FBlueprintDifferenceTreeEntry>(
-					SelectionCallback,
-					FGenerateDiffEntryWidget::CreateStatic(&GenerateSimpleDiffWidget, DiffText));
-
-				Children.Push(Entry);
-				OutRealDifferences.Push(Entry);
-			}
-
-			FString OldInterfaces, NewInterfaces;
-			for (const FBPInterfaceDescription& Desc : OldBlueprint->ImplementedInterfaces)
-			{
-				if (!OldInterfaces.IsEmpty())
-				{
-					OldInterfaces += TEXT(", ");
-				}
-				OldInterfaces += GetNameSafe(*Desc.Interface);
-			}
-
-			for (const FBPInterfaceDescription& Desc : NewBlueprint->ImplementedInterfaces)
-			{
-				if (!NewInterfaces.IsEmpty())
-				{
-					NewInterfaces += TEXT(", ");
-				}
-				NewInterfaces += GetNameSafe(*Desc.Interface);
-			}
-			
-			if (OldInterfaces != NewInterfaces)
-			{
-				FText DiffText = FText::Format(LOCTEXT("InterfacesChanged", "Interfaces changed from '{0}' to '{1}'"), FText::FromString(OldInterfaces), FText::FromString(NewInterfaces));
-
-				TSharedPtr<FBlueprintDifferenceTreeEntry> Entry = MakeShared<FBlueprintDifferenceTreeEntry>(
-					SelectionCallback,
-					FGenerateDiffEntryWidget::CreateStatic(&GenerateSimpleDiffWidget, DiffText));
-
-				Children.Push(Entry);
-				OutRealDifferences.Push(Entry);
-			}
-		}
-
-		const bool bHasDifferences = Children.Num() != 0;
-		if (!bHasDifferences)
-		{
-			// make one child informing the user that there are no differences:
-			Children.Push(FBlueprintDifferenceTreeEntry::NoDifferencesEntry());
-		}
-
-		OutTreeEntries.Push(FBlueprintDifferenceTreeEntry::CreateCategoryEntry(
-			NSLOCTEXT("FBlueprintDifferenceTreeEntry", "SettingsLabel", "Class Settings"),
-			NSLOCTEXT("FBlueprintDifferenceTreeEntry", "SettingsTooltip", "The list of changes made in the Class Settings panel"),
-			SelectionCallback,
-			Children,
-			bHasDifferences
-		));
-	}
-};
-
-/** Diff control to handle finding type-specific differences */
-struct FBlueprintTypeDiffControl : public TSharedFromThis<FBlueprintTypeDiffControl>, public IDiffControl
-{
-	struct FSubObjectDiff
-	{
-		FDiffSingleResult SourceResult;
-		FDetailsDiff OldDetails;
-		FDetailsDiff NewDetails;
-		TArray<TSharedPtr<FBlueprintDifferenceTreeEntry>> Diffs;
-
-		FSubObjectDiff(const FDiffSingleResult& InSourceResult, const UObject* OldObject, const UObject* NewObject)
-			: SourceResult(InSourceResult)
-			, OldDetails(OldObject, FDetailsDiff::FOnDisplayedPropertiesChanged())
-			, NewDetails(NewObject, FDetailsDiff::FOnDisplayedPropertiesChanged())
-		{}
-	};
-
-	FBlueprintTypeDiffControl(const UBlueprint* InBlueprintOld, const UBlueprint* InBlueprintNew, FOnDiffEntryFocused InSelectionCallback)
-		: BlueprintOld(InBlueprintOld), BlueprintNew(InBlueprintNew), SelectionCallback(InSelectionCallback), bDiffSucceeded(false)
-	{
-		check(InBlueprintNew && InBlueprintOld);
-	}
-
-	/** Generate difference tree widgets */
-	virtual void GenerateTreeEntries(TArray< TSharedPtr<FBlueprintDifferenceTreeEntry> >& OutTreeEntries, TArray< TSharedPtr<FBlueprintDifferenceTreeEntry> >& OutRealDifferences) override;
-
-	/** The old blueprint (left) */
-	const UBlueprint* BlueprintOld;
-
-	/** The new blueprint(right) */
-	const UBlueprint* BlueprintNew;
-
-	/** Boxes that will display the details diffs */
-	TSharedPtr<SBox> OldDetailsBox;
-	TSharedPtr<SBox> NewDetailsBox;
-
-private:
-	/** Generate Widget for top category */
-	TSharedRef<SWidget> GenerateCategoryWidget(bool bHasRealDiffs);
-
-	/** Build up the Diff Source Array*/
-	void BuildDiffSourceArray();
-
-	/** Handle selecting a diff */
-	void OnSelectSubobjectDiff(FPropertySoftPath Identifier, TSharedPtr<FSubObjectDiff> SubObjectDiff);
-
-	/** List of objects with differences */
-	TArray<TSharedPtr<FSubObjectDiff>> SubObjectDiffs;
-
-	/** Source for list view */
-	TArray<TSharedPtr<FDiffResultItem>> DiffListSource;
-
-	/** Selection callback */
-	FOnDiffEntryFocused SelectionCallback;
-
-	/** Did diff generation succeed? */
-	bool bDiffSucceeded;
-};
-
-TSharedRef<SWidget> FBlueprintTypeDiffControl::GenerateCategoryWidget(bool bHasRealDiffs)
-{
-	FLinearColor Color = FLinearColor::White;
-
-	if (bHasRealDiffs)
-	{
-		Color = DiffViewUtils::Differs();
-	}
-
-	FText Label = BlueprintNew->GetClass()->GetDisplayNameText();
-
-	return SNew(SHorizontalBox)
-		+ SHorizontalBox::Slot()
-		[
-			SNew(STextBlock)
-			.ColorAndOpacity(Color)
-		.Text(Label)
-		];
-}
-
-void FBlueprintTypeDiffControl::OnSelectSubobjectDiff(FPropertySoftPath Identifier, TSharedPtr<FSubObjectDiff> SubObjectDiff)
-{
-	// This allows the owning control to focus the correct tab (or do whatever else it likes):
-	SelectionCallback.ExecuteIfBound();
-
-	if (SubObjectDiff.IsValid())
-	{
-		SubObjectDiff->OldDetails.HighlightProperty(Identifier);
-		SubObjectDiff->NewDetails.HighlightProperty(Identifier);
-
-		OldDetailsBox->SetContent(SubObjectDiff->OldDetails.DetailsWidget());
-		NewDetailsBox->SetContent(SubObjectDiff->NewDetails.DetailsWidget());
-	}
-}
-
-void FBlueprintTypeDiffControl::BuildDiffSourceArray()
-{
-	TArray<FDiffSingleResult> BlueprintDiffResults;
-	FDiffResults BlueprintDiffs(&BlueprintDiffResults);
-	if (BlueprintNew->FindDiffs(BlueprintOld, BlueprintDiffs))
-	{
-		bDiffSucceeded = true;
-
-		// Add manual diffs
-		for (const FDiffSingleResult& CurrentDiff : BlueprintDiffResults)
-		{
-			if (CurrentDiff.Diff == EDiffType::OBJECT_REQUEST_DIFF)
-			{
-				// Turn into a subobject diff
-
-				// Invert order, we want old then new
-				TSharedPtr<FSubObjectDiff> SubObjectDiff = MakeShared<FSubObjectDiff>(CurrentDiff, CurrentDiff.Object2, CurrentDiff.Object1);
-
-				TArray<FSingleObjectDiffEntry> DifferingProperties;
-				SubObjectDiff->OldDetails.DiffAgainst(SubObjectDiff->NewDetails, DifferingProperties, true);
-
-				if (DifferingProperties.Num() > 0)
-				{
-					// Actual differences, so add to tree
-					SubObjectDiffs.Add(SubObjectDiff);
-
-					for (const FSingleObjectDiffEntry& Difference : DifferingProperties)
-					{
-						TSharedPtr<FBlueprintDifferenceTreeEntry> Entry = MakeShared<FBlueprintDifferenceTreeEntry>(
-							FOnDiffEntryFocused::CreateSP(AsShared(), &FBlueprintTypeDiffControl::OnSelectSubobjectDiff, Difference.Identifier, SubObjectDiff),
-							FGenerateDiffEntryWidget::CreateStatic(&GenerateObjectDiffWidget, Difference, RightRevision));
-						SubObjectDiff->Diffs.Push(Entry);
-					}
-				}
-			}
-			else
-			{
-				DiffListSource.Add(MakeShared<FDiffResultItem>(CurrentDiff));
-			}
-		}
-
-		struct SortDiff
-		{
-			bool operator () (const TSharedPtr<FDiffResultItem>& A, const TSharedPtr<FDiffResultItem>& B) const
-			{
-				return A->Result.Diff < B->Result.Diff;
-			}
-		};
-
-		Sort(DiffListSource.GetData(), DiffListSource.Num(), SortDiff());
-	}
-}
-
-void FBlueprintTypeDiffControl::GenerateTreeEntries(TArray< TSharedPtr<FBlueprintDifferenceTreeEntry> >& OutTreeEntries, TArray< TSharedPtr<FBlueprintDifferenceTreeEntry> >& OutRealDifferences)
-{
-	BuildDiffSourceArray();
-
-	TArray< TSharedPtr<FBlueprintDifferenceTreeEntry> > Children;
-	
-	bool bHasRealChange = false;
-
-	// First add manual diffs in main category
-	for (const TSharedPtr<FDiffResultItem>& Difference : DiffListSource)
-	{
-		TSharedPtr<FBlueprintDifferenceTreeEntry> ChildEntry = MakeShared<FBlueprintDifferenceTreeEntry>(
-			SelectionCallback,
-			FGenerateDiffEntryWidget::CreateSP(Difference.ToSharedRef(), &FDiffResultItem::GenerateWidget));
-		Children.Push(ChildEntry);
-		OutRealDifferences.Push(ChildEntry);
-
-		if (Difference->Result.IsRealDifference())
-		{
-			bHasRealChange = true;
-		}
-	}
-
-	if (Children.Num() == 0)
-	{
-		// Make one child informing the user that there are no differences, or that it is unknown
-		if (bDiffSucceeded)
-		{
-			Children.Push(FBlueprintDifferenceTreeEntry::NoDifferencesEntry());
-		}
-		else
-		{
-			Children.Push(FBlueprintDifferenceTreeEntry::UnknownDifferencesEntry());
-		}
-	}
-
-	TSharedPtr<FBlueprintDifferenceTreeEntry> CategoryEntry = MakeShared<FBlueprintDifferenceTreeEntry>(
-			SelectionCallback,
-			FGenerateDiffEntryWidget::CreateSP(AsShared(), &FBlueprintTypeDiffControl::GenerateCategoryWidget, bHasRealChange),
-			Children);
-	OutTreeEntries.Push(CategoryEntry);
-
-	// Now add subobject diffs, one category per object
-	for (const TSharedPtr<FSubObjectDiff>& SubObjectDiff : SubObjectDiffs)
-	{
-		Children.Reset();
-
-		Children.Append(SubObjectDiff->Diffs);
-		OutRealDifferences.Append(SubObjectDiff->Diffs);
-
-		TSharedPtr<FBlueprintDifferenceTreeEntry> SubObjectEntry = FBlueprintDifferenceTreeEntry::CreateCategoryEntry(
-			SubObjectDiff->SourceResult.DisplayString,
-			SubObjectDiff->SourceResult.ToolTip,
-			FOnDiffEntryFocused::CreateSP(AsShared(), &FBlueprintTypeDiffControl::OnSelectSubobjectDiff, FPropertySoftPath(), SubObjectDiff),
-			Children,
-			true);
-
-		OutTreeEntries.Push(SubObjectEntry);
-	}
-}
-
-/** Category list item for a graph*/
-struct FGraphToDiff	: public TSharedFromThis<FGraphToDiff>, IDiffControl
-{
-	FGraphToDiff(SBlueprintDiff* DiffWidget, UEdGraph* GraphOld, UEdGraph* GraphNew, const FRevisionInfo& RevisionOld, const FRevisionInfo& RevisionNew);
-	virtual ~FGraphToDiff();
-
-	/** Add widgets to the differences tree */
-	virtual void GenerateTreeEntries(TArray< TSharedPtr<FBlueprintDifferenceTreeEntry> >& OutTreeEntries, TArray< TSharedPtr<FBlueprintDifferenceTreeEntry> >& OutRealDifferences) override;
-
-	/** Get old(left) graph*/
-	UEdGraph* GetGraphOld() const { return GraphOld; }
-
-	/** Get new(right) graph*/
-	UEdGraph* GetGraphNew() const { return GraphNew; }
-
-	/** Source for list view */
-	TArray<TSharedPtr<FDiffResultItem>> DiffListSource;
-	TSharedPtr<TArray<FDiffSingleResult>> FoundDiffs;
-
-private:
-	/** Get tooltip for category */
-	FText GetToolTip();
-
-	/** Generate Widget for category list */
-	TSharedRef<SWidget> GenerateCategoryWidget();
-
-	/** Called when the Newer Graph is modified*/
-	void OnGraphChanged(const FEdGraphEditAction& Action);
-
-	/** Build up the Diff Source Array*/
-	void BuildDiffSourceArray();
-
-	/** Diff widget */
-	class SBlueprintDiff* DiffWidget;
-
-	/** The old graph(left)*/
-	UEdGraph* GraphOld;
-
-	/** The new graph(right)*/
-	UEdGraph* GraphNew;
-
-	/** Description of Old and new graph*/
-	FRevisionInfo	RevisionOld, RevisionNew;
-
-	/** Handle to the registered OnGraphChanged delegate. */
-	FDelegateHandle OnGraphChangedDelegateHandle;
-};
-
-FGraphToDiff::FGraphToDiff(SBlueprintDiff* InDiffWidget, UEdGraph* InGraphOld, UEdGraph* InGraphNew, const FRevisionInfo& InRevisionOld, const FRevisionInfo& InRevisionNew)
-	: FoundDiffs(MakeShared<TArray<FDiffSingleResult>>()), DiffWidget(InDiffWidget), GraphOld(InGraphOld), GraphNew(InGraphNew), RevisionOld(InRevisionOld), RevisionNew(InRevisionNew)
-{
-	check(InGraphOld || InGraphNew); //one of them needs to exist
-
-	//need to know when it is modified
-	if(InGraphNew)
-	{
-		OnGraphChangedDelegateHandle = InGraphNew->AddOnGraphChangedHandler( FOnGraphChanged::FDelegate::CreateRaw(this, &FGraphToDiff::OnGraphChanged));
-	}
-
-	BuildDiffSourceArray();
-}
-
-FGraphToDiff::~FGraphToDiff()
-{
-	if(GraphNew)
-	{
-		GraphNew->RemoveOnGraphChangedHandler( OnGraphChangedDelegateHandle);
-	}
-}
-
-TSharedRef<SWidget> FGraphToDiff::GenerateCategoryWidget()
-{
-	const UEdGraph* Graph = GraphOld ? GraphOld : GraphNew;
-	check(Graph);
-	
-	FLinearColor Color = (GraphOld && GraphNew) ? FLinearColor::White : FLinearColor(0.3f,0.3f,1.f);
-
-	const bool bHasDiffs = DiffListSource.Num() > 0;
-
-	if(bHasDiffs)
-	{
-		Color = DiffViewUtils::Differs();
-	}
-
-	FText GraphName;
-	if (const UEdGraphSchema* Schema = Graph->GetSchema())
-	{
-		FGraphDisplayInfo DisplayInfo;
-		Schema->GetGraphDisplayInformation(*Graph, DisplayInfo);
-
-		GraphName = DisplayInfo.DisplayName;
-	}
-	else
-	{
-		GraphName = FText::FromName(Graph->GetFName());
-	}
-
-	return SNew(SHorizontalBox)
-	+ SHorizontalBox::Slot()
-	[
-		SNew(STextBlock)
-		.ColorAndOpacity(Color)
-		.Text(GraphName)
-		.ToolTipText(GetToolTip())
-	]
-	+ DiffViewUtils::Box( GraphOld != nullptr, Color )
-	+ DiffViewUtils::Box( GraphNew != nullptr, Color );
-}
-
-FText FGraphToDiff::GetToolTip()
-{
-	if (GraphOld && GraphNew)
-	{
-		if (DiffListSource.Num() > 0)
-		{
-			return LOCTEXT("ContainsDifferences", "Revisions are different");
-		}
-		else
-		{
-			return LOCTEXT("GraphsIdentical", "Revisions appear to be identical");
-		}
-	}
-	else
-	{
-		UEdGraph* GoodGraph = GraphOld ? GraphOld : GraphNew;
-		check(GoodGraph);
-		const FRevisionInfo& Revision = GraphNew ? RevisionOld : RevisionNew;
-		FText RevisionText = LOCTEXT("CurrentRevision", "Current Revision");
-
-		if (!Revision.Revision.IsEmpty())
-		{
-			RevisionText = FText::Format(LOCTEXT("Revision Number", "Revision {0}"), FText::FromString(Revision.Revision));
-		}
-
-		return FText::Format(LOCTEXT("MissingGraph", "Graph '{0}' missing from {1}"), FText::FromString(GoodGraph->GetName()), RevisionText);
-	}
-}
-
-void FGraphToDiff::GenerateTreeEntries(TArray< TSharedPtr<FBlueprintDifferenceTreeEntry> >& OutTreeEntries, TArray< TSharedPtr<FBlueprintDifferenceTreeEntry> >& OutRealDifferences)
-{
-	TArray< TSharedPtr<FBlueprintDifferenceTreeEntry> > Children;
-	for (const TSharedPtr<FDiffResultItem>& Difference : DiffListSource)
-	{
-		TSharedPtr<FBlueprintDifferenceTreeEntry> ChildEntry = MakeShared<FBlueprintDifferenceTreeEntry>(
-				FOnDiffEntryFocused::CreateRaw(DiffWidget, &SBlueprintDiff::OnDiffListSelectionChanged, Difference),
-				FGenerateDiffEntryWidget::CreateSP(Difference.ToSharedRef(), &FDiffResultItem::GenerateWidget));
-		Children.Push(ChildEntry);
-		OutRealDifferences.Push(ChildEntry);
-	}
-
-	if (Children.Num() == 0)
-	{
-		// make one child informing the user that there are no differences:
-		Children.Push(FBlueprintDifferenceTreeEntry::NoDifferencesEntry());
-	}
-
-	TSharedPtr<FBlueprintDifferenceTreeEntry> Entry = MakeShared<FBlueprintDifferenceTreeEntry>(
-			FOnDiffEntryFocused::CreateRaw(DiffWidget, &SBlueprintDiff::OnGraphSelectionChanged, TSharedPtr<FGraphToDiff>(AsShared()), ESelectInfo::Direct),
-			FGenerateDiffEntryWidget::CreateSP(AsShared(), &FGraphToDiff::GenerateCategoryWidget),
-			Children);
-	OutTreeEntries.Push(Entry);
-}
-
-void FGraphToDiff::BuildDiffSourceArray()
-{
-	FoundDiffs->Empty();
-	FGraphDiffControl::DiffGraphs(GraphOld, GraphNew, *FoundDiffs);
-
-	struct SortDiff
-	{
-		bool operator () (const FDiffSingleResult& A, const FDiffSingleResult& B) const
-		{
-			return A.Diff < B.Diff;
-		}
-	};
-
-	Sort(FoundDiffs->GetData(), FoundDiffs->Num(), SortDiff());
-
-	DiffListSource.Empty();
-	for (const FDiffSingleResult& Diff : *FoundDiffs)
-	{
-		DiffListSource.Add(MakeShared<FDiffResultItem>(Diff));
-	}
-}
-
-void FGraphToDiff::OnGraphChanged( const FEdGraphEditAction& Action )
-{
-	DiffWidget->OnGraphChanged(this);
-}
-
 FDiffPanel::FDiffPanel()
 {
 	Blueprint = nullptr;
@@ -855,7 +88,7 @@ static int32 GetCurrentIndex( SListView< TSharedPtr< FDiffSingleResult> > const&
 void DiffWidgetUtils::SelectNextRow( SListView< TSharedPtr< FDiffSingleResult> >& ListView, const TArray< TSharedPtr< FDiffSingleResult > >& ListViewSource )
 {
 	int32 CurrentIndex = GetCurrentIndex(ListView, ListViewSource);
-	if( CurrentIndex == ListViewSource.Num() - 1 )
+	if (CurrentIndex == ListViewSource.Num() - 1)
 	{
 		return;
 	}
@@ -1120,7 +353,7 @@ void SBlueprintDiff::CreateGraphEntry( UEdGraph* GraphOld, UEdGraph* GraphNew )
 
 void SBlueprintDiff::OnGraphSelectionChanged(TSharedPtr<FGraphToDiff> Item, ESelectInfo::Type SelectionType)
 {
-	if(!Item.IsValid())
+	if (!Item.IsValid())
 	{
 		return;
 	}
@@ -1131,7 +364,7 @@ void SBlueprintDiff::OnGraphSelectionChanged(TSharedPtr<FGraphToDiff> Item, ESel
 
 void SBlueprintDiff::OnGraphChanged(FGraphToDiff* Diff)
 {
-	if(PanelNew.GraphEditor.IsValid() && PanelNew.GraphEditor.Pin()->GetCurrentGraph() == Diff->GetGraphNew())
+	if (PanelNew.GraphEditor.IsValid() && PanelNew.GraphEditor.Pin()->GetCurrentGraph() == Diff->GetGraphNew())
 	{
 		FocusOnGraphRevisions(Diff);
 	}
@@ -1213,7 +446,7 @@ bool SBlueprintDiff::HasPrevDiff() const
 
 FGraphToDiff* SBlueprintDiff::FindGraphToDiffEntry(const FString& GraphPath)
 {
-	for(const TSharedPtr<FGraphToDiff>& Graph : Graphs)
+	for (const TSharedPtr<FGraphToDiff>& Graph : Graphs)
 	{
 		FString SearchGraphPath = Graph->GetGraphOld() ? FGraphDiffControl::GetGraphPath(Graph->GetGraphOld()) : FGraphDiffControl::GetGraphPath(Graph->GetGraphNew());
 		if (SearchGraphPath.Equals(GraphPath, ESearchCase::CaseSensitive))
@@ -1244,7 +477,7 @@ void SBlueprintDiff::OnDiffListSelectionChanged(TSharedPtr<FDiffResultItem> TheD
 	const auto SafeClearSelection = []( TWeakPtr<SGraphEditor> GraphEditor )
 	{
 		TSharedPtr<SGraphEditor> GraphEditorPtr = GraphEditor.Pin();
-		if( GraphEditorPtr.IsValid())
+		if (GraphEditorPtr.IsValid())
 		{
 			GraphEditorPtr->ClearSelectionSet();
 		}
@@ -1281,7 +514,7 @@ void SBlueprintDiff::OnToggleSplitViewMode()
 {
 	bVerticalSplitGraphMode = !bVerticalSplitGraphMode;
 
-	if(SSplitter* DiffGraphSplitterPtr = DiffGraphSplitter.Get())
+	if (SSplitter* DiffGraphSplitterPtr = DiffGraphSplitter.Get())
 	{
 		DiffGraphSplitterPtr->SetOrientation(bVerticalSplitGraphMode ? Orient_Horizontal : Orient_Vertical);
 	}
@@ -1299,9 +532,9 @@ FSlateIcon SBlueprintDiff::GetSplitViewModeImage() const
 
 void SBlueprintDiff::ResetGraphEditors()
 {
-	if(PanelOld.GraphEditor.IsValid() && PanelNew.GraphEditor.IsValid())
+	if (PanelOld.GraphEditor.IsValid() && PanelNew.GraphEditor.IsValid())
 	{
-		if(bLockViews)
+		if (bLockViews)
 		{
 			PanelOld.GraphEditor.Pin()->LockToGraphEditor(PanelNew.GraphEditor);
 			PanelNew.GraphEditor.Pin()->LockToGraphEditor(PanelOld.GraphEditor);
@@ -1323,7 +556,7 @@ void FDiffPanel::GeneratePanel(UEdGraph* NewGraph, UEdGraph* OldGraph )
 
 void FDiffPanel::GeneratePanel(UEdGraph* Graph, TSharedPtr<TArray<FDiffSingleResult>> DiffResults, TAttribute<int32> FocusedDiffResult)
 {
-	if( GraphEditor.IsValid() && GraphEditor.Pin()->GetCurrentGraph() == Graph )
+	if (GraphEditor.IsValid() && GraphEditor.Pin()->GetCurrentGraph() == Graph)
 	{
 		return;
 	}
@@ -1335,7 +568,7 @@ void FDiffPanel::GeneratePanel(UEdGraph* Graph, TSharedPtr<TArray<FDiffSingleRes
 									SNew(STextBlock).Text( LOCTEXT("BPDifPanelNoGraphTip", "Graph does not exist in this revision"))
 								];
 
-	if(Graph)
+	if (Graph)
 	{
 		SGraphEditor::FGraphEditorEvents InEvents;
 		{
@@ -1354,7 +587,7 @@ void FDiffPanel::GeneratePanel(UEdGraph* Graph, TSharedPtr<TArray<FDiffSingleRes
 			InEvents.OnCreateNodeOrPinMenu = SGraphEditor::FOnCreateNodeOrPinMenu::CreateStatic(ContextMenuHandler);
 		}
 
-		if ( !GraphEditorCommands.IsValid() )
+		if (!GraphEditorCommands.IsValid())
 		{
 			GraphEditorCommands = MakeShared<FUICommandList>();
 
@@ -1851,7 +1084,7 @@ TSharedRef<SBox> SBlueprintDiff::GenerateRevisionInfoWidgetForPanel(TSharedPtr<S
 
 void SBlueprintDiff::SetCurrentMode(FName NewMode)
 {
-	if( CurrentMode == NewMode )
+	if (CurrentMode == NewMode)
 	{
 		return;
 	}
diff --git a/Engine/Source/Editor/Kismet/Private/SDetailsDiff.cpp b/Engine/Source/Editor/Kismet/Private/SDetailsDiff.cpp
new file mode 100644
--- /dev/null
+++ b/Engine/Source/Editor/Kismet/Private/SDetailsDiff.cpp
@@ -0,0 +1,390 @@
+// Copyright Epic Games, Inc. All Rights Reserved.
+
+#include "SDetailsDiff.h"
+#include "Editor.h"
+#include "Widgets/Layout/SSplitter.h"
+#include "Widgets/SOverlay.h"
+#include "SlateOptMacros.h"
+#include "Widgets/Layout/SSpacer.h"
+#include "Framework/MultiBox/MultiBoxDefs.h"
+#include "Framework/MultiBox/MultiBoxBuilder.h"
+#include "Styling/AppStyle.h"
+#include "K2Node_MathExpression.h"
+#include "Kismet2/KismetEditorUtilities.h"
+#include "DetailsDiff.h"
+#include "GraphDiffControl.h"
+#include "WidgetBlueprint.h"
+#include "HAL/PlatformApplicationMisc.h"
+#include "Framework/Application/SlateApplication.h"
+#include "SBlueprintDiff.h"
+#include "DiffControl.h"
+#include "PropertyEditorModule.h"
+#include "Modules/ModuleManager.h"
+#include "IDetailsView.h"
+
+#include "Subsystems/AssetEditorSubsystem.h"
+
+#define LOCTEXT_NAMESPACE "SDetailsDif"
+
+typedef TMap< FName, const FProperty* > FNamePropertyMap;
+
+static const FName DetailsMode = FName(TEXT("DetailsMode"));
+
+BEGIN_SLATE_FUNCTION_BUILD_OPTIMIZATION
+void SDetailsDiff::Construct( const FArguments& InArgs)
+{
+	check(InArgs._AssetOld && InArgs._AssetNew);
+	PanelOld.Object = InArgs._AssetOld;
+	PanelNew.Object = InArgs._AssetNew;
+	PanelOld.RevisionInfo = InArgs._OldRevision;
+	PanelNew.RevisionInfo = InArgs._NewRevision;
+
+	// sometimes we want to clearly identify the assets being diffed (when it's
+	// not the same asset in each panel)
+	PanelOld.bShowAssetName = InArgs._ShowAssetNames;
+	PanelNew.bShowAssetName = InArgs._ShowAssetNames;
+
+	bLockViews = true;
+
+	if (InArgs._ParentWindow.IsValid())
+	{
+		WeakParentWindow = InArgs._ParentWindow;
+
+		AssetEditorCloseDelegate = GEditor->GetEditorSubsystem<UAssetEditorSubsystem>()->OnAssetEditorRequestClose().AddSP(this, &SDetailsDiff::OnCloseAssetEditor);
+	}
+
+	FToolBarBuilder NavToolBarBuilder(TSharedPtr< const FUICommandList >(), FMultiBoxCustomization::None);
+	NavToolBarBuilder.AddToolBarButton(
+		FUIAction(
+			FExecuteAction::CreateSP(this, &SDetailsDiff::PrevDiff),
+			FCanExecuteAction::CreateSP( this, &SDetailsDiff::HasPrevDiff)
+		)
+		, NAME_None
+		, LOCTEXT("PrevDiffLabel", "Prev")
+		, LOCTEXT("PrevDiffTooltip", "Go to previous difference")
+		, FSlateIcon(FAppStyle::GetAppStyleSetName(), "BlueprintDif.PrevDiff")
+	);
+	NavToolBarBuilder.AddToolBarButton(
+		FUIAction(
+			FExecuteAction::CreateSP(this, &SDetailsDiff::NextDiff),
+			FCanExecuteAction::CreateSP(this, &SDetailsDiff::HasNextDiff)
+		)
+		, NAME_None
+		, LOCTEXT("NextDiffLabel", "Next")
+		, LOCTEXT("NextDiffTooltip", "Go to next difference")
+		, FSlateIcon(FAppStyle::GetAppStyleSetName(), "BlueprintDif.NextDiff")
+	);
+
+	DifferencesTreeView = DiffTreeView::CreateTreeView(&MasterDifferencesList);
+
+	GenerateDifferencesList();
+
+	const auto TextBlock = [](FText Text) -> TSharedRef<SWidget>
+	{
+		return SNew(SBox)
+		.Padding(FMargin(4.0f,10.0f))
+		.VAlign(VAlign_Center)
+		.HAlign(HAlign_Left)
+		[
+			SNew(STextBlock)
+			.Visibility(EVisibility::HitTestInvisible)
+			.TextStyle(FAppStyle::Get(), "DetailsView.CategoryTextStyle")
+			.Text(Text)
+		];
+	};
+
+	TopRevisionInfoWidget =
+		SNew(SSplitter)
+		.Visibility(EVisibility::HitTestInvisible)
+		+ SSplitter::Slot()
+		.Value(.2f)
+		[
+			SNew(SBox)
+		]
+		+ SSplitter::Slot()
+		.Value(.8f)
+		[
+			SNew(SSplitter)
+			.PhysicalSplitterHandleSize(10.0f)
+			+ SSplitter::Slot()
+			.Value(.5f)
+			[
+				TextBlock(DiffViewUtils::GetPanelLabel(PanelOld.Object, PanelOld.RevisionInfo, FText()))
+			]
+			+ SSplitter::Slot()
+			.Value(.5f)
+			[
+				TextBlock(DiffViewUtils::GetPanelLabel(PanelNew.Object, PanelNew.RevisionInfo, FText()))
+			]
+		];
+
+	this->ChildSlot
+		[
+			SNew(SBorder)
+			.BorderImage(FAppStyle::GetBrush( "Docking.Tab", ".ContentAreaBrush" ))
+			[
+				SNew(SOverlay)
+				+ SOverlay::Slot()
+				.VAlign(VAlign_Top)
+				[
+					TopRevisionInfoWidget.ToSharedRef()		
+				]
+				+ SOverlay::Slot()
+				[
+					SNew(SVerticalBox)
+					+SVerticalBox::Slot()
+					.AutoHeight()
+					.Padding(0.0f, 2.0f, 0.0f, 2.0f)
+					[
+						SNew(SHorizontalBox)
+						+ SHorizontalBox::Slot()
+						.Padding(4.f)
+						.AutoWidth()
+						[
+							NavToolBarBuilder.MakeWidget()
+						]
+						+ SHorizontalBox::Slot()
+						[
+							SNew(SSpacer)
+						]
+					]
+					+ SVerticalBox::Slot()
+					[
+						SNew(SSplitter)
+						+ SSplitter::Slot()
+						.Value(.2f)
+						[
+							SNew(SBorder)
+							.BorderImage(FAppStyle::GetBrush("ToolPanel.GroupBorder"))
+							[
+								DifferencesTreeView.ToSharedRef()
+							]
+						]
+						+ SSplitter::Slot()
+						.Value(.8f)
+						[
+							SAssignNew(ModeContents, SBox)
+						]
+					]
+				]
+			]
+		];
+
+	SetCurrentMode(DetailsMode);
+}
+END_SLATE_FUNCTION_BUILD_OPTIMIZATION
+
+SDetailsDiff::~SDetailsDiff()
+{
+	if (AssetEditorCloseDelegate.IsValid())
+	{
+		GEditor->GetEditorSubsystem<UAssetEditorSubsystem>()->OnAssetEditorRequestClose().Remove(AssetEditorCloseDelegate);
+	}
+}
+
+void SDetailsDiff::OnCloseAssetEditor(UObject* Asset, EAssetEditorCloseReason CloseReason)
+{
+	if (PanelOld.Object == Asset || PanelNew.Object == Asset || CloseReason == EAssetEditorCloseReason::CloseAllAssetEditors)
+	{
+		// Tell our window to close and set our selves to collapsed to try and stop it from ticking
+		SetVisibility(EVisibility::Collapsed);
+
+		if (AssetEditorCloseDelegate.IsValid())
+		{
+			GEditor->GetEditorSubsystem<UAssetEditorSubsystem>()->OnAssetEditorRequestClose().Remove(AssetEditorCloseDelegate);
+		}
+
+		if (WeakParentWindow.IsValid())
+		{
+			WeakParentWindow.Pin()->RequestDestroyWindow();
+		}
+	}
+}
+
+TSharedRef<SWidget> SDetailsDiff::DefaultEmptyPanel()
+{
+	return SNew(SHorizontalBox)
+		+ SHorizontalBox::Slot()
+		.HAlign(HAlign_Center)
+		.VAlign(VAlign_Center)
+		[
+			SNew(STextBlock)
+			.Text(LOCTEXT("BlueprintDifGraphsToolTip", "Select Graph to Diff"))
+		];
+}
+
+TSharedPtr<SWindow> SDetailsDiff::CreateDiffWindow(FText WindowTitle, UObject* OldObject, UObject* NewObject, const FRevisionInfo& OldRevision, const FRevisionInfo& NewRevision)
+{
+	// sometimes we're comparing different revisions of one single asset (other 
+	// times we're comparing two completely separate assets altogether)
+	bool bIsSingleAsset = (NewObject->GetName() == OldObject->GetName());
+
+	TSharedPtr<SWindow> Window = SNew(SWindow)
+		.Title(WindowTitle)
+		.ClientSize(FVector2D(1000.f, 800.f));
+
+	Window->SetContent(SNew(SDetailsDiff)
+		.AssetOld(OldObject)
+		.AssetNew(NewObject)
+		.OldRevision(OldRevision)
+		.NewRevision(NewRevision)
+		.ShowAssetNames(!bIsSingleAsset)
+		.ParentWindow(Window));
+
+	// Make this window a child of the modal window if we've been spawned while one is active.
+	TSharedPtr<SWindow> ActiveModal = FSlateApplication::Get().GetActiveModalWindow();
+	if (ActiveModal.IsValid())
+	{
+		FSlateApplication::Get().AddWindowAsNativeChild(Window.ToSharedRef(), ActiveModal.ToSharedRef());
+	}
+	else
+	{
+		FSlateApplication::Get().AddWindow(Window.ToSharedRef());
+	}
+
+	return Window;
+}
+
+void SDetailsDiff::NextDiff()
+{
+	DiffTreeView::HighlightNextDifference(DifferencesTreeView.ToSharedRef(), RealDifferences, MasterDifferencesList);
+}
+
+void SDetailsDiff::PrevDiff()
+{
+	DiffTreeView::HighlightPrevDifference(DifferencesTreeView.ToSharedRef(), RealDifferences, MasterDifferencesList);
+}
+
+bool SDetailsDiff::HasNextDiff() const
+{
+	return DiffTreeView::HasNextDifference(DifferencesTreeView.ToSharedRef(), RealDifferences);
+}
+
+bool SDetailsDiff::HasPrevDiff() const
+{
+	return DiffTreeView::HasPrevDifference(DifferencesTreeView.ToSharedRef(), RealDifferences);
+}
+
+void SDetailsDiff::OnDiffListSelectionChanged(TSharedPtr<FDiffResultItem> TheDiff )
+{
+	check( !TheDiff->Result.OwningObjectPath.IsEmpty() );
+	// TODO: What do I put here?
+}
+
+void SDetailsDiff::GenerateDifferencesList()
+{
+	MasterDifferencesList.Empty();
+	RealDifferences.Empty();
+	ModePanels.Empty();
+	
+	const auto CreateInspector = [](const UObject* Object) {
+		FPropertyEditorModule& EditModule = FModuleManager::Get().GetModuleChecked<FPropertyEditorModule>("PropertyEditor");
+	
+		FNotifyHook* NotifyHook = nullptr;
+	
+		FDetailsViewArgs DetailsViewArgs;
+		DetailsViewArgs.NameAreaSettings = FDetailsViewArgs::HideNameArea;
+		DetailsViewArgs.bHideSelectionTip = true;
+		DetailsViewArgs.NotifyHook = NotifyHook;
+		DetailsViewArgs.ViewIdentifier = FName("ObjectInspector");
+		TSharedRef<IDetailsView> DetailsView = EditModule.CreateDetailView( DetailsViewArgs );
+		DetailsView->SetObject(const_cast<UObject*>(Object)); // TODO: @jordan.hoffmann apologize to the const gods for your sins
+		
+		return DetailsView;
+	};
+
+	// TODO: construct DetailsView of PanelOld and PanelNew
+	PanelOld.DetailsView = CreateInspector(PanelOld.Object);
+	PanelNew.DetailsView = CreateInspector(PanelOld.Object);
+
+	// Now that we have done the diffs, create the panel widgets
+	// (we're currently only generating the details panel but we can add more as needed)
+	ModePanels.Add(DetailsMode, GenerateDetailsPanel());
+
+	DifferencesTreeView->RebuildList();
+}
+
+SDetailsDiff::FDiffControl SDetailsDiff::GenerateDetailsPanel()
+{
+	TSharedPtr<FDetailsDiffControl> NewDiffControl = MakeShared<FDetailsDiffControl>(PanelOld.Object, PanelNew.Object, FOnDiffEntryFocused::CreateRaw(this, &SDetailsDiff::SetCurrentMode, DetailsMode));
+	NewDiffControl->GenerateTreeEntries(MasterDifferencesList, RealDifferences);
+
+	SDetailsDiff::FDiffControl Ret;
+	Ret.DiffControl = NewDiffControl;
+	Ret.Widget = SNew(SSplitter)
+		.PhysicalSplitterHandleSize(10.0f)
+		+ SSplitter::Slot()
+		.Value(0.5f)
+		[
+			NewDiffControl->OldDetailsWidget()
+		]
+		+ SSplitter::Slot()
+		.Value(0.5f)
+		[
+			NewDiffControl->NewDetailsWidget()
+		];
+
+	return Ret;
+}
+
+TSharedRef<SBox> SDetailsDiff::GenerateRevisionInfoWidgetForPanel(TSharedPtr<SWidget>& OutGeneratedWidget, const FText& InRevisionText) const
+{
+	return SAssignNew(OutGeneratedWidget,SBox)
+		.Padding(FMargin(4.0f, 10.0f))
+		.VAlign(VAlign_Center)
+		.HAlign(HAlign_Left)
+		[
+			SNew(STextBlock)
+			.TextStyle(FAppStyle::Get(), "DetailsView.CategoryTextStyle")
+			.Text(InRevisionText)
+			.ShadowColorAndOpacity(FColor::Black)
+			.ShadowOffset(FVector2D(1.4,1.4))
+		];
+}
+
+void SDetailsDiff::SetCurrentMode(FName NewMode)
+{
+	if (CurrentMode == NewMode)
+	{
+		return;
+	}
+
+	CurrentMode = NewMode;
+
+	FDiffControl* FoundControl = ModePanels.Find(NewMode);
+
+	if (FoundControl)
+	{
+		// Reset inspector view
+		PanelOld.DetailsView->SetObject(nullptr);
+		PanelNew.DetailsView->SetObject(nullptr);
+
+		ModeContents->SetContent(FoundControl->Widget.ToSharedRef());
+	}
+	else
+	{
+		ensureMsgf(false, TEXT("Diff panel does not support mode %s"), *NewMode.ToString() );
+	}
+
+	OnModeChanged(NewMode);
+}
+
+void SDetailsDiff::UpdateTopSectionVisibility(const FName& InNewViewMode) const
+{
+	SSplitter* TopRevisionInfoWidgetPtr = TopRevisionInfoWidget.Get();
+
+	if (!TopRevisionInfoWidgetPtr)
+	{
+		return;
+	}
+	
+	TopRevisionInfoWidgetPtr->SetVisibility(EVisibility::HitTestInvisible);
+}
+
+void SDetailsDiff::OnModeChanged(const FName& InNewViewMode) const
+{
+	UpdateTopSectionVisibility(InNewViewMode);
+}
+
+#undef LOCTEXT_NAMESPACE
+
diff --git a/Engine/Source/Editor/Kismet/Public/DiffUtils.h b/Engine/Source/Editor/Kismet/Public/DiffUtils.h
--- a/Engine/Source/Editor/Kismet/Public/DiffUtils.h
+++ b/Engine/Source/Editor/Kismet/Public/DiffUtils.h
@@ -75,7 +75,7 @@ struct FPropertySoftPath
 
 	FPropertySoftPath( FPropertyPath InPropertyPath )
 	{
-		for( int32 i = 0, end = InPropertyPath.GetNumProperties(); i != end; ++i )
+		for (int32 i = 0, end = InPropertyPath.GetNumProperties(); i != end; ++i)
 		{
 			PropertyChain.Push(FChainElement(InPropertyPath.GetPropertyInfo(i).Property.Get()));
 		}
@@ -193,7 +193,7 @@ FORCEINLINE bool operator!=(const FSCSIdentifier& A, const FSCSIdentifier& B)
 FORCEINLINE uint32 GetTypeHash( FPropertySoftPath const& Path )
 {
 	uint32 Ret = 0;
-	for( const FPropertySoftPath::FChainElement& PropertyElement : Path.PropertyChain )
+	for (const FPropertySoftPath::FChainElement& PropertyElement : Path.PropertyChain)
 	{
 		Ret = Ret ^ GetTypeHash(PropertyElement.PropertyName);
 	}
@@ -332,7 +332,7 @@ namespace DiffViewUtils
 
 	KISMET_API FText PropertyDiffMessage(FSingleObjectDiffEntry Difference, FText ObjectName);
 	KISMET_API FText SCSDiffMessage(const FSCSDiffEntry& Difference, FText ObjectName);
-	KISMET_API FText GetPanelLabel(const UBlueprint* Blueprint, const FRevisionInfo& Revision, FText Label);
+	KISMET_API FText GetPanelLabel(const UObject* Asset, const FRevisionInfo& Revision, FText Label);
 
 	KISMET_API SHorizontalBox::FSlot::FSlotArguments Box(bool bIsPresent, FLinearColor Color);
 }
diff --git a/Engine/Source/Editor/Kismet/Public/SDetailsDiff.h b/Engine/Source/Editor/Kismet/Public/SDetailsDiff.h
new file mode 100644
--- /dev/null
+++ b/Engine/Source/Editor/Kismet/Public/SDetailsDiff.h
@@ -0,0 +1,139 @@
+// Copyright Epic Games, Inc. All Rights Reserved.
+#pragma once
+
+#include "CoreMinimal.h"
+#include "Widgets/DeclarativeSyntaxSupport.h"
+#include "Widgets/SWindow.h"
+#include "Widgets/SCompoundWidget.h"
+#include "DiffUtils.h"
+#include "Developer/AssetTools/Public/IAssetTypeActions.h"
+
+class FSpawnTabArgs;
+class FTabManager;
+class IDiffControl;
+class FUICommandList;
+enum class EAssetEditorCloseReason : uint8;
+
+/** Panel used to display the details */
+struct KISMET_API FDetailsDiffPanel
+{
+	/** The asset that owns the details view we are showing */
+	const UObject*				Object = nullptr;
+
+	/** The details view associated with the asset */
+	TSharedPtr<class IDetailsView>	DetailsView;
+
+	/** Revision information for this asset */
+	FRevisionInfo					RevisionInfo;
+
+	/** True if we should show a name identifying which asset this panel is displaying */
+	bool							bShowAssetName = true;
+
+	/** The widget that contains the revision info in graph mode */
+	TSharedPtr<SWidget>				OverlayRevisionInfo;
+private:
+	/** Command list for this diff panel */
+	TSharedPtr<FUICommandList> GraphEditorCommands;
+};
+
+/* Visual Diff between two Assets */
+class  KISMET_API SDetailsDiff: public SCompoundWidget
+{
+public:
+	SLATE_BEGIN_ARGS( SDetailsDiff ){}
+			SLATE_ARGUMENT( const class UObject*, AssetOld )
+			SLATE_ARGUMENT( const class UObject*, AssetNew )
+			SLATE_ARGUMENT( struct FRevisionInfo, OldRevision )
+			SLATE_ARGUMENT( struct FRevisionInfo, NewRevision )
+			SLATE_ARGUMENT( bool, ShowAssetNames )
+			SLATE_ARGUMENT(TSharedPtr<SWindow>, ParentWindow)
+	SLATE_END_ARGS()
+
+	void Construct(const FArguments& InArgs);
+	virtual ~SDetailsDiff();
+
+	/** Called when user clicks on an entry in the listview of differences */
+	void OnDiffListSelectionChanged(TSharedPtr<struct FDiffResultItem> TheDiff);
+
+	/** Helper function for generating an empty widget */
+	static TSharedRef<SWidget> DefaultEmptyPanel();
+
+	/** Helper function to create a window that holds a diff widget */
+	static TSharedPtr<SWindow> CreateDiffWindow(FText WindowTitle, UObject* OldObject, UObject* NewObject, const struct FRevisionInfo& OldRevision, const struct FRevisionInfo& NewRevision);
+
+protected:
+	/** Called when user clicks button to go to next difference */
+	void NextDiff();
+
+	/** Called when user clicks button to go to prev difference */
+	void PrevDiff();
+
+	/** Called to determine whether we have a list of differences to cycle through */
+	bool HasNextDiff() const;
+	bool HasPrevDiff() const;
+	
+	/** Function used to generate the list of differences and the widgets needed to calculate that list */
+	void GenerateDifferencesList();
+
+	/** Called when editor may need to be closed */
+	void OnCloseAssetEditor(UObject* Asset, EAssetEditorCloseReason CloseReason);
+
+	struct FDiffControl
+	{
+		FDiffControl()
+		: Widget()
+		, DiffControl(nullptr)
+		{
+		}
+
+		TSharedPtr<SWidget> Widget;
+		TSharedPtr< class IDiffControl > DiffControl;
+	};
+
+	FDiffControl GenerateDetailsPanel();
+
+	TSharedRef<SBox> GenerateRevisionInfoWidgetForPanel(TSharedPtr<SWidget>& OutGeneratedWidget,const FText& InRevisionText) const;
+
+	/** Accessor and event handler for toggling between diff view modes (defaults, components, graph view, interface, macro): */
+	void SetCurrentMode(FName NewMode);
+	FName GetCurrentMode() const { return CurrentMode; }
+	void OnModeChanged(const FName& InNewViewMode) const;
+
+	void UpdateTopSectionVisibility(const FName& InNewViewMode) const;
+
+	FName CurrentMode;
+
+	/*The two panels used to show the old & new revision*/ 
+	FDetailsDiffPanel PanelOld;
+	FDetailsDiffPanel PanelNew;
+	
+	/** If the two views should be locked */
+	bool bLockViews;
+
+	/** Contents widget that we swap when mode changes (defaults, components, etc) */
+	TSharedPtr<SBox> ModeContents;
+
+	TSharedPtr<SSplitter> TopRevisionInfoWidget;
+
+	/** We can't use the global tab manager because we need to instance the diff control, so we have our own tab manager: */
+	TSharedPtr<FTabManager> TabManager;
+
+	/** Tree of differences collected across all panels: */
+	TArray< TSharedPtr<class FBlueprintDifferenceTreeEntry> > MasterDifferencesList;
+
+	/** List of all differences, cached so that we can iterate only the differences and not labels, etc: */
+	TArray< TSharedPtr<class FBlueprintDifferenceTreeEntry> > RealDifferences;
+
+	/** Tree view that displays the differences, cached for the buttons that iterate the differences: */
+	TSharedPtr< STreeView< TSharedPtr< FBlueprintDifferenceTreeEntry > > > DifferencesTreeView;
+
+	/** Stored references to widgets used to display various parts of an object, from the mode name */
+	TMap<FName, FDiffControl> ModePanels;
+
+	/** A pointer to the window holding this */
+	TWeakPtr<SWindow> WeakParentWindow;
+
+	FDelegateHandle AssetEditorCloseDelegate;
+};
+
+
