[BugFix] BP Diff pins not highlighting when their diff is selected sometimes

diff --git a/Engine/Source/Editor/Kismet/Private/DiffControl.cpp b/Engine/Source/Editor/Kismet/Private/DiffControl.cpp
--- a/Engine/Source/Editor/Kismet/Private/DiffControl.cpp
+++ b/Engine/Source/Editor/Kismet/Private/DiffControl.cpp
@@ -213,7 +213,6 @@ void FCDODiffControl::GenerateTreeEntries(TArray<TSharedPtr<FBlueprintDifference
 	));
 }
 
-
 /////////////////////////////////////////////////////////////////////////////
 /// FClassSettingsDiffControl
 
@@ -479,6 +478,11 @@ FGraphToDiff::~FGraphToDiff()
 
 void FGraphToDiff::GenerateTreeEntries(TArray< TSharedPtr<FBlueprintDifferenceTreeEntry> >& OutTreeEntries, TArray< TSharedPtr<FBlueprintDifferenceTreeEntry> >& OutRealDifferences)
 {
+	if (!DiffListSource.IsEmpty())
+	{
+		RealDifferencesStartIndex = OutRealDifferences.Num();
+	}
+	
 	TArray< TSharedPtr<FBlueprintDifferenceTreeEntry> > Children;
 	for (const TSharedPtr<FDiffResultItem>& Difference : DiffListSource)
 	{
diff --git a/Engine/Source/Editor/Kismet/Private/DiffControl.h b/Engine/Source/Editor/Kismet/Private/DiffControl.h
--- a/Engine/Source/Editor/Kismet/Private/DiffControl.h
+++ b/Engine/Source/Editor/Kismet/Private/DiffControl.h
@@ -226,6 +226,9 @@ struct FGraphToDiff	: public TSharedFromThis<FGraphToDiff>, IDiffControl
 	TArray<TSharedPtr<FDiffResultItem>> DiffListSource;
 	TSharedPtr<TArray<FDiffSingleResult>> FoundDiffs;
 
+	/** Index of the first item in RealDifferences that was generated by this graph */
+	int32 RealDifferencesStartIndex = INDEX_NONE;
+
 private:
 	/** Get tooltip for category */
 	FText GetToolTip();
diff --git a/Engine/Source/Editor/Kismet/Private/SBlueprintDiff.cpp b/Engine/Source/Editor/Kismet/Private/SBlueprintDiff.cpp
--- a/Engine/Source/Editor/Kismet/Private/SBlueprintDiff.cpp
+++ b/Engine/Source/Editor/Kismet/Private/SBlueprintDiff.cpp
@@ -694,6 +694,7 @@ void SBlueprintDiff::HandleGraphChanged( const FString& GraphPath )
 	UEdGraph* GraphOld = nullptr;
 	UEdGraph* GraphNew = nullptr;
 	TSharedPtr<TArray<FDiffSingleResult>> DiffResults;
+	int32 RealDifferencesStartIndex = INDEX_NONE;
 	for (const TSharedPtr<FGraphToDiff>& GraphToDiff : Graphs)
 	{
 		UEdGraph* NewGraph = GraphToDiff->GetGraphNew();
@@ -704,32 +705,22 @@ void SBlueprintDiff::HandleGraphChanged( const FString& GraphPath )
 			GraphNew = NewGraph;
 			GraphOld = OldGraph;
 			DiffResults = GraphToDiff->FoundDiffs;
+			RealDifferencesStartIndex = GraphToDiff->RealDifferencesStartIndex;
 			break;
 		}
 	}
 	
 	const TAttribute<int32> FocusedDiffResult = TAttribute<int32>::CreateLambda(
-        [this, GraphPath]()
+        [this, RealDifferencesStartIndex]()
         {
-        	// because DifferencesTreeView has differences from several graphs, we need to find the index of the first entry in this graph
-        	// that way we can return the selected index of just this graph's entries
-        	int32 startIndex = 0;
-			for (const TSharedPtr<FGraphToDiff>& GraphToDiff : Graphs)
-			{
-				UEdGraph* NewGraph = GraphToDiff->GetGraphNew();
-				UEdGraph* OldGraph = GraphToDiff->GetGraphOld();
-				const FString OtherGraphPath = NewGraph ? FGraphDiffControl::GetGraphPath(NewGraph) : FGraphDiffControl::GetGraphPath(OldGraph);
-				if (GraphPath.Equals(OtherGraphPath))
-				{
-					break;
-				}
-				if (GraphToDiff->FoundDiffs.IsValid())
-				{
-					startIndex += GraphToDiff->FoundDiffs->Num();
-				}
+        	int32 FocusedDiffResult = INDEX_NONE;
+        	if (RealDifferencesStartIndex != INDEX_NONE)
+        	{
+				FocusedDiffResult = DiffTreeView::CurrentDifference(DifferencesTreeView.ToSharedRef(), RealDifferences) - RealDifferencesStartIndex;
 			}
+        	
 			// find selected index in all the graphs, and subtract the index of the first entry in this graph
-			return DiffTreeView::CurrentDifference(DifferencesTreeView.ToSharedRef(), RealDifferences) - startIndex;
+			return FocusedDiffResult;
         });
 
 	// only regenerate PanelOld if the old graph has changed
