Add missing imports after dragging one or more non-imported namespaced type assets into the components tree in the Blueprint editor.

Change summary:
- Removed SSubobjectBlueprintEditor::FOnNewSubobjectAdded typedef and OnNewSubobjectAdded delegate members. Deprecation was not needed as these were added after the 5.0 release.
- Restored SSubobjectBlueprintEditor::AddNewSubobject() back to its original implementation (prior to adding the now-defunct OnNewSubobjectAdded event).
- Added USubobjectDataSubsystem::OnNewSubobjectAdded() to allow potential listeners to register for and receive successful subobject "add" event notifications.
- Modified USubobjectDataSubsystem::AddNewSubobject() to broadcast new subobject data to any registered listener(s) after a successful add.
- Added FBlueprintEditor::DeferredNamespaceImports to stack auto-imports for events that can be triggered more than once within a single frame.
- Modified FBlueprintEditor::OnComponentAddedToBlueprint() to handle subobject add events that target editor Blueprint objects. Also fixed to query the subobject class for its default import set instead of the subobject itself.
- Modified the FBlueprintEditor dtor to unregister the OnBlueprintUnloaded() event handler as I noticed this was not previously being done (benign as the shared ptr was skipped for being invalid, but this now removes us from the broadcast chain).
- Fixed up minor formatting inconsistencies in some legacy FBlueprintEditor analytics-related code paths.

diff --git a/Engine/Source/Editor/Kismet/Private/BlueprintEditor.cpp b/Engine/Source/Editor/Kismet/Private/BlueprintEditor.cpp
--- a/Engine/Source/Editor/Kismet/Private/BlueprintEditor.cpp
+++ b/Engine/Source/Editor/Kismet/Private/BlueprintEditor.cpp
@@ -917,9 +917,9 @@ UEdGraphNode* FBlueprintEditor::GetSingleSelectedNode() const
 	return (SelectedNodes.Num() == 1) ? Cast<UEdGraphNode>(*SelectedNodes.CreateConstIterator()) : nullptr;
 }
 
-void FBlueprintEditor::AnalyticsTrackNodeEvent( UBlueprint* Blueprint, UEdGraphNode *GraphNode, bool bNodeDelete ) const
+void FBlueprintEditor::AnalyticsTrackNodeEvent(UBlueprint* Blueprint, UEdGraphNode *GraphNode, bool bNodeDelete) const
 {
-	if( Blueprint && GraphNode && FEngineAnalytics::IsAvailable() )
+	if(Blueprint && GraphNode && FEngineAnalytics::IsAvailable())
 	{
 		// we'd like to see if this was happening in normal blueprint editor or persona 
 		const FString EditorName = Cast<UAnimBlueprint>(Blueprint) != nullptr ? TEXT("Persona") : TEXT("BlueprintEditor");
@@ -927,39 +927,39 @@ void FBlueprintEditor::AnalyticsTrackNodeEvent( UBlueprint* Blueprint, UEdGraphN
 		// Build Node Details
 		const UGeneralProjectSettings& ProjectSettings = *GetDefault<UGeneralProjectSettings>();
 		FString ProjectID = ProjectSettings.ProjectID.ToString();
-		TArray< FAnalyticsEventAttribute > NodeAttributes;
-		NodeAttributes.Add( FAnalyticsEventAttribute( TEXT( "ProjectId" ), ProjectID ) );
-		NodeAttributes.Add( FAnalyticsEventAttribute( TEXT( "BlueprintId" ), Blueprint->GetBlueprintGuid().ToString() ) );
-		TArray< TKeyValuePair<FString, FString> > Attributes;
+		TArray<FAnalyticsEventAttribute> NodeAttributes;
+		NodeAttributes.Add(FAnalyticsEventAttribute(TEXT("ProjectId"), ProjectID));
+		NodeAttributes.Add(FAnalyticsEventAttribute(TEXT("BlueprintId"), Blueprint->GetBlueprintGuid().ToString()));
+		TArray<TKeyValuePair<FString, FString>> Attributes;
 
-		if( UK2Node* K2Node = Cast<UK2Node>( GraphNode ))
+		if (UK2Node* K2Node = Cast<UK2Node>(GraphNode))
 		{
-			K2Node->GetNodeAttributes( Attributes );
+			K2Node->GetNodeAttributes(Attributes);
 		}
-		else if( UEdGraphNode_Comment* CommentNode = Cast<UEdGraphNode_Comment>( GraphNode ))
+		else if (UEdGraphNode_Comment* CommentNode = Cast<UEdGraphNode_Comment>(GraphNode))
 		{
-			Attributes.Add( TKeyValuePair<FString, FString>( TEXT( "Type" ), TEXT( "Comment" ) ));
-			Attributes.Add( TKeyValuePair<FString, FString>( TEXT( "Class" ), CommentNode->GetClass()->GetName() ));
-			Attributes.Add( TKeyValuePair<FString, FString>( TEXT( "Name" ), CommentNode->GetName() ));
+			Attributes.Add(TKeyValuePair<FString, FString>(TEXT("Type"), TEXT("Comment")));
+			Attributes.Add(TKeyValuePair<FString, FString>(TEXT("Class"), CommentNode->GetClass()->GetName()));
+			Attributes.Add(TKeyValuePair<FString, FString>(TEXT("Name"), CommentNode->GetName()));
 		}
-		if( Attributes.Num() > 0 )
+		if (Attributes.Num() > 0)
 		{
 			// Build Node Attributes
 			for (const TKeyValuePair<FString, FString>& Attribute : Attributes)
 			{
-				NodeAttributes.Add( FAnalyticsEventAttribute( Attribute.Key, Attribute.Value ));
+				NodeAttributes.Add(FAnalyticsEventAttribute(Attribute.Key, Attribute.Value));
 			}
 			// Send Analytics event 
-			FString EventType = bNodeDelete ?	FString::Printf(TEXT( "Editor.Usage.%s.NodeDeleted" ), *EditorName) :
-												FString::Printf(TEXT( "Editor.Usage.%s.NodeCreated" ), *EditorName);
-			FEngineAnalytics::GetProvider().RecordEvent( EventType, NodeAttributes );
+			FString EventType = bNodeDelete ?	FString::Printf(TEXT("Editor.Usage.%s.NodeDeleted"), *EditorName) :
+												FString::Printf(TEXT("Editor.Usage.%s.NodeCreated"), *EditorName);
+			FEngineAnalytics::GetProvider().RecordEvent(EventType, NodeAttributes);
 		}
 	}
 }
 
-void FBlueprintEditor::AnalyticsTrackCompileEvent( UBlueprint* Blueprint, int32 NumErrors, int32 NumWarnings ) const
+void FBlueprintEditor::AnalyticsTrackCompileEvent(UBlueprint* Blueprint, int32 NumErrors, int32 NumWarnings) const
 {
-	if(Blueprint && FEngineAnalytics::IsAvailable())
+	if (Blueprint && FEngineAnalytics::IsAvailable())
 	{
 		// we'd like to see if this was happening in normal blueprint editor or persona 
 		const FString EditorName = Cast<UAnimBlueprint>(Blueprint) != nullptr ? TEXT("Persona") : TEXT("BlueprintEditor");
@@ -969,10 +969,10 @@ void FBlueprintEditor::AnalyticsTrackCompileEvent( UBlueprint* Blueprint, int32
 		FString ProjectID = ProjectSettings.ProjectID.ToString();
 
 		const bool bSuccess = NumErrors == 0;
-		TArray< FAnalyticsEventAttribute > Attributes;
+		TArray<FAnalyticsEventAttribute> Attributes;
 		Attributes.Add(FAnalyticsEventAttribute(TEXT("ProjectId"), ProjectID));
 		Attributes.Add(FAnalyticsEventAttribute(TEXT("BlueprintId"), Blueprint->GetBlueprintGuid().ToString()));
-		Attributes.Add(FAnalyticsEventAttribute(TEXT("Success"), bSuccess? TEXT("True") : TEXT("False") ));
+		Attributes.Add(FAnalyticsEventAttribute(TEXT("Success"), bSuccess? TEXT("True") : TEXT("False")));
 		Attributes.Add(FAnalyticsEventAttribute(TEXT("NumErrors"), FString::FromInt(NumErrors)));
 		Attributes.Add(FAnalyticsEventAttribute(TEXT("NumWarnings"), FString::FromInt(NumWarnings)));
 
@@ -1284,13 +1284,34 @@ void FBlueprintEditor::OnComponentDoubleClicked(TSharedPtr<FSubobjectEditorTreeN
 
 void FBlueprintEditor::OnComponentAddedToBlueprint(const FSubobjectData& NewSubobjectData)
 {
+	// Determine if we've added to a Blueprint instance that we're editing within this context.
+	bool bIsEditingEventTarget = false;
+	if (const UBlueprint* TargetBlueprint = NewSubobjectData.GetBlueprint())
+	{
+		for (UObject* EditorObject : GetEditingObjects())
+		{
+			if (EditorObject == TargetBlueprint)
+			{
+				bIsEditingEventTarget = true;
+				break;
+			}
+		}
+	}
+
+	// Only handle add events if the editor context includes its targeted Blueprint object.
+	if (!bIsEditingEventTarget)
+	{
+		return;
+	}
+
 	const UObject* NewSubobject = NewSubobjectData.GetObject();
 	check(NewSubobject);
 
-	// If necessary, auto-import the default namespace set associated with the new subobject's class.
-	FImportNamespaceExParameters Params;
-	FBlueprintNamespaceUtilities::GetDefaultImportsForObject(NewSubobject, Params.NamespacesToImport);
-	ImportNamespaceEx(Params);
+	// Get the default namespace set associated with the new subobject's class. Because we
+	// might receive multiple events within a single frame (e.g. dragging multiple component
+	// Blueprint class assets into the components tree will result in multiple notifications),
+	// we'll add these into the deferred list for now and auto-import them all on the next tick.
+	FBlueprintNamespaceUtilities::GetDefaultImportsForObject(NewSubobject->GetClass(), DeferredNamespaceImports);
 }
 
 TSharedRef<SWidget> FBlueprintEditor::CreateGraphTitleBarWidget(TSharedRef<FTabInfo> InTabInfo, UEdGraph* InGraph)
@@ -2774,54 +2795,61 @@ FBlueprintEditor::~FBlueprintEditor()
 	UEditorEngine* Editor = (UEditorEngine*)GEngine;
 	if (Editor)
 	{
-		Editor->UnregisterForUndo( this );
+		Editor->UnregisterForUndo(this);
 	}
 
 	CloseMergeTool();
 
 	if (GetBlueprintObj())
 	{
-		GetBlueprintObj()->OnChanged().RemoveAll( this );
-		GetBlueprintObj()->OnCompiled().RemoveAll( this );
-		GetBlueprintObj()->OnSetObjectBeingDebugged().RemoveAll( this );
+		GetBlueprintObj()->OnChanged().RemoveAll(this);
+		GetBlueprintObj()->OnCompiled().RemoveAll(this);
+		GetBlueprintObj()->OnSetObjectBeingDebugged().RemoveAll(this);
+	}
+
+	if (USubobjectDataSubsystem* SubobjectDataSubsystem = USubobjectDataSubsystem::Get())
+	{
+		SubobjectDataSubsystem->OnNewSubobjectAdded().RemoveAll(this);
 	}
 
-	FGlobalTabmanager::Get()->OnActiveTabChanged_Unsubscribe( OnActiveTabChangedDelegateHandle );
+	FKismetEditorUtilities::OnBlueprintUnloaded.RemoveAll(this);
+
+	FGlobalTabmanager::Get()->OnActiveTabChanged_Unsubscribe(OnActiveTabChangedDelegateHandle);
 
 	if (FEngineAnalytics::IsAvailable())
 	{
 		const UGeneralProjectSettings& ProjectSettings = *GetDefault<UGeneralProjectSettings>();
 		FString ProjectID = ProjectSettings.ProjectID.ToString();
 
-		TArray< FAnalyticsEventAttribute > BPEditorAttribs;
-		BPEditorAttribs.Add( FAnalyticsEventAttribute( FString( "GraphActionMenusExecuted.NonContextSensitive" ), AnalyticsStats.GraphActionMenusNonCtxtSensitiveExecCount ));
-		BPEditorAttribs.Add( FAnalyticsEventAttribute( FString( "GraphActionMenusExecuted.ContextSensitive" ), AnalyticsStats.GraphActionMenusCtxtSensitiveExecCount ));
-		BPEditorAttribs.Add( FAnalyticsEventAttribute( FString( "GraphActionMenusClosed" ), AnalyticsStats.GraphActionMenusCancelledCount ));
+		TArray<FAnalyticsEventAttribute> BPEditorAttribs;
+		BPEditorAttribs.Add(FAnalyticsEventAttribute(TEXT("GraphActionMenusExecuted.NonContextSensitive"), AnalyticsStats.GraphActionMenusNonCtxtSensitiveExecCount));
+		BPEditorAttribs.Add(FAnalyticsEventAttribute(TEXT("GraphActionMenusExecuted.ContextSensitive"), AnalyticsStats.GraphActionMenusCtxtSensitiveExecCount));
+		BPEditorAttribs.Add(FAnalyticsEventAttribute(TEXT("GraphActionMenusClosed"), AnalyticsStats.GraphActionMenusCancelledCount));
 
-		BPEditorAttribs.Add( FAnalyticsEventAttribute( FString( "MyBlueprintDragPlacedNodesCreated" ), AnalyticsStats.MyBlueprintNodeDragPlacementCount ));
-		BPEditorAttribs.Add( FAnalyticsEventAttribute( FString( "BlueprintPaletteDragPlacedNodesCreated" ), AnalyticsStats.PaletteNodeDragPlacementCount ));
-		BPEditorAttribs.Add( FAnalyticsEventAttribute( FString( "GraphContextNodesCreated" ), AnalyticsStats.NodeGraphContextCreateCount ));
-		BPEditorAttribs.Add( FAnalyticsEventAttribute( FString( "GraphPinContextNodesCreated" ), AnalyticsStats.NodePinContextCreateCount ));
-		BPEditorAttribs.Add( FAnalyticsEventAttribute( FString( "KeymapNodesCreated" ), AnalyticsStats.NodeKeymapCreateCount ));
-		BPEditorAttribs.Add( FAnalyticsEventAttribute( FString( "PastedNodesCreated" ), AnalyticsStats.NodePasteCreateCount ));
+		BPEditorAttribs.Add(FAnalyticsEventAttribute(TEXT("MyBlueprintDragPlacedNodesCreated"), AnalyticsStats.MyBlueprintNodeDragPlacementCount));
+		BPEditorAttribs.Add(FAnalyticsEventAttribute(TEXT("BlueprintPaletteDragPlacedNodesCreated"), AnalyticsStats.PaletteNodeDragPlacementCount));
+		BPEditorAttribs.Add(FAnalyticsEventAttribute(TEXT("GraphContextNodesCreated" ), AnalyticsStats.NodeGraphContextCreateCount));
+		BPEditorAttribs.Add(FAnalyticsEventAttribute(TEXT("GraphPinContextNodesCreated" ), AnalyticsStats.NodePinContextCreateCount));
+		BPEditorAttribs.Add(FAnalyticsEventAttribute(TEXT("KeymapNodesCreated"), AnalyticsStats.NodeKeymapCreateCount));
+		BPEditorAttribs.Add(FAnalyticsEventAttribute(TEXT("PastedNodesCreated"), AnalyticsStats.NodePasteCreateCount));
 
-		BPEditorAttribs.Add( FAnalyticsEventAttribute( FString( "ProjectId" ), ProjectID ) );
-		FEngineAnalytics::GetProvider().RecordEvent( FString( "Editor.Usage.BlueprintEditorSummary" ), BPEditorAttribs );
+		BPEditorAttribs.Add(FAnalyticsEventAttribute(TEXT("ProjectId"), ProjectID));
+		FEngineAnalytics::GetProvider().RecordEvent(TEXT("Editor.Usage.BlueprintEditorSummary"), BPEditorAttribs);
 
 		for (auto Iter = AnalyticsStats.GraphDisallowedPinConnections.CreateConstIterator(); Iter; ++Iter)
 		{
-			TArray< FAnalyticsEventAttribute > BPEditorPinConnectAttribs;
-			BPEditorPinConnectAttribs.Add( FAnalyticsEventAttribute( FString( "FromPin.Category" ), Iter->PinTypeCategoryA ));
-			BPEditorPinConnectAttribs.Add( FAnalyticsEventAttribute( FString( "FromPin.IsArray" ), Iter->bPinIsArrayA ));
-			BPEditorPinConnectAttribs.Add( FAnalyticsEventAttribute( FString( "FromPin.IsReference" ), Iter->bPinIsReferenceA ));
-			BPEditorPinConnectAttribs.Add( FAnalyticsEventAttribute( FString( "FromPin.IsWeakPointer" ), Iter->bPinIsWeakPointerA ));
-			BPEditorPinConnectAttribs.Add( FAnalyticsEventAttribute( FString( "ToPin.Category" ), Iter->PinTypeCategoryB ));
-			BPEditorPinConnectAttribs.Add( FAnalyticsEventAttribute( FString( "ToPin.IsArray" ), Iter->bPinIsArrayB ));
-			BPEditorPinConnectAttribs.Add( FAnalyticsEventAttribute( FString( "ToPin.IsReference" ), Iter->bPinIsReferenceB ));
-			BPEditorPinConnectAttribs.Add( FAnalyticsEventAttribute( FString( "ToPin.IsWeakPointer" ), Iter->bPinIsWeakPointerB ));
-			BPEditorPinConnectAttribs.Add( FAnalyticsEventAttribute( FString( "ProjectId" ), ProjectID ) );
+			TArray<FAnalyticsEventAttribute> BPEditorPinConnectAttribs;
+			BPEditorPinConnectAttribs.Add(FAnalyticsEventAttribute(TEXT("FromPin.Category"), Iter->PinTypeCategoryA));
+			BPEditorPinConnectAttribs.Add(FAnalyticsEventAttribute(TEXT("FromPin.IsArray"), Iter->bPinIsArrayA));
+			BPEditorPinConnectAttribs.Add(FAnalyticsEventAttribute(TEXT("FromPin.IsReference"), Iter->bPinIsReferenceA));
+			BPEditorPinConnectAttribs.Add(FAnalyticsEventAttribute(TEXT("FromPin.IsWeakPointer"), Iter->bPinIsWeakPointerA));
+			BPEditorPinConnectAttribs.Add(FAnalyticsEventAttribute(TEXT("ToPin.Category"), Iter->PinTypeCategoryB));
+			BPEditorPinConnectAttribs.Add(FAnalyticsEventAttribute(TEXT("ToPin.IsArray"), Iter->bPinIsArrayB));
+			BPEditorPinConnectAttribs.Add(FAnalyticsEventAttribute(TEXT("ToPin.IsReference"), Iter->bPinIsReferenceB));
+			BPEditorPinConnectAttribs.Add(FAnalyticsEventAttribute(TEXT("ToPin.IsWeakPointer"), Iter->bPinIsWeakPointerB));
+			BPEditorPinConnectAttribs.Add(FAnalyticsEventAttribute(TEXT("ProjectId"), ProjectID));
 
-			FEngineAnalytics::GetProvider().RecordEvent( FString( "Editor.Usage.BPDisallowedPinConnection" ), BPEditorPinConnectAttribs );
+			FEngineAnalytics::GetProvider().RecordEvent(TEXT("Editor.Usage.BPDisallowedPinConnection"), BPEditorPinConnectAttribs);
 		}
 	}
 
@@ -2920,11 +2948,15 @@ void FBlueprintEditor::CreateSubobjectEditors()
 		.AllowEditing(this, &FBlueprintEditor::InEditingMode)
 		.OnSelectionUpdated(this, &FBlueprintEditor::OnSelectionUpdated)
 		.OnItemDoubleClicked(this, &FBlueprintEditor::OnComponentDoubleClicked)
-		.OnNewSubobjectAdded(this, &FBlueprintEditor::OnComponentAddedToBlueprint)
 		.SubobjectClassListFilters(ClassFilters);
 	
 	SubobjectViewport = SAssignNew(SubobjectViewport, SSCSEditorViewport)
 		.BlueprintEditor(SharedThis(this));
+
+	if (USubobjectDataSubsystem* SubobjectDataSubsystem = USubobjectDataSubsystem::Get())
+	{
+		SubobjectDataSubsystem->OnNewSubobjectAdded().AddSP(this, &FBlueprintEditor::OnComponentAddedToBlueprint);
+	}
 }
 
 void FBlueprintEditor::OnLogTokenClicked(const TSharedRef<IMessageToken>& Token)
@@ -7403,7 +7435,7 @@ void FBlueprintEditor::PasteNodesHere(class UEdGraph* DestinationGraph, const FV
 			{
 				for (const UStruct* ExternalDependency : ExternalDependencies)
 				{
-					FBlueprintNamespaceUtilities::GetDefaultImportsForObject(ExternalDependency, NamespacesToImport);
+					FBlueprintNamespaceUtilities::GetDefaultImportsForObject(ExternalDependency, DeferredNamespaceImports);
 				}
 			}
 
@@ -7440,15 +7472,6 @@ void FBlueprintEditor::PasteNodesHere(class UEdGraph* DestinationGraph, const FV
 			// Log new node created to analytics
 			AnalyticsTrackNodeEvent(GetBlueprintObj(), Node, false);
 		}
-
-		// Auto-import any namespace dependencies.
-		const bool bShouldAutoImportNamespaceDependencies = GetDefault<UBlueprintEditorSettings>()->bEnableNamespaceImportingFeatures;
-		if (bShouldAutoImportNamespaceDependencies && NamespacesToImport.Num() > 0)
-		{
-			FImportNamespaceExParameters Params;
-			Params.NamespacesToImport = MoveTemp(NamespacesToImport);
-			ImportNamespaceEx(Params);
-		}
 	}
 
 	if (bNeedToModifyStructurally)
@@ -9001,6 +9024,19 @@ void FBlueprintEditor::Tick(float DeltaTime)
 		check(EditorInst != nullptr);
 		EditorInst->CloseWindow();
 	}
+	else
+	{
+		// Auto-import any namespaces we've collected as a result of compound events that may have occurred within this frame.
+		if (!DeferredNamespaceImports.IsEmpty())
+		{
+			FImportNamespaceExParameters Params;
+			Params.NamespacesToImport = MoveTemp(DeferredNamespaceImports);
+			ImportNamespaceEx(Params);
+
+			// Assert that this was reset by the move above.
+			check(DeferredNamespaceImports.IsEmpty());
+		}
+	}
 }
 TStatId FBlueprintEditor::GetStatId() const
 {
diff --git a/Engine/Source/Editor/Kismet/Public/BlueprintEditor.h b/Engine/Source/Editor/Kismet/Public/BlueprintEditor.h
--- a/Engine/Source/Editor/Kismet/Public/BlueprintEditor.h
+++ b/Engine/Source/Editor/Kismet/Public/BlueprintEditor.h
@@ -1425,6 +1425,9 @@ protected:
 
 	/** Filter used to compose multiple pin type filters together */
 	TSharedPtr<class IPinTypeSelectorFilter> CompositePinTypeSelectorFilter;
+
+	/** Pending set of namespaces to be auto-imported on the next tick */
+	TSet<FString> DeferredNamespaceImports;
 	
 	FOnSetPinVisibility OnSetPinVisibility;
 
diff --git a/Engine/Source/Editor/SubobjectDataInterface/Private/SubobjectDataSubsystem.cpp b/Engine/Source/Editor/SubobjectDataInterface/Private/SubobjectDataSubsystem.cpp
--- a/Engine/Source/Editor/SubobjectDataInterface/Private/SubobjectDataSubsystem.cpp
+++ b/Engine/Source/Editor/SubobjectDataInterface/Private/SubobjectDataSubsystem.cpp
@@ -1073,6 +1073,12 @@ FSubobjectDataHandle USubobjectDataSubsystem::AddNewSubobject(const FAddNewSubob
 			}
 		}
 	}
+
+	// If the new handle is valid, broadcast to any listeners that a new subobject was added.
+	if (const FSubobjectData* NewData = NewDataHandle.GetData())
+	{
+		OnNewSubobjectAdded_Delegate.Broadcast(*NewData);
+	}
 	
 	return NewDataHandle;
 }
diff --git a/Engine/Source/Editor/SubobjectDataInterface/Public/SubobjectDataSubsystem.h b/Engine/Source/Editor/SubobjectDataInterface/Public/SubobjectDataSubsystem.h
--- a/Engine/Source/Editor/SubobjectDataInterface/Public/SubobjectDataSubsystem.h
+++ b/Engine/Source/Editor/SubobjectDataInterface/Public/SubobjectDataSubsystem.h
@@ -3,6 +3,7 @@
 #pragma once
 
 #include "CoreMinimal.h"
+#include "Delegates/Delegate.h"
 #include "Subsystems/EngineSubsystem.h"
 #include "SubobjectData.h"
 #include "ModuleDescriptor.h"
@@ -12,6 +13,8 @@
 class UActorComponent;
 class FScopedTransaction;
 
+DECLARE_MULTICAST_DELEGATE_OneParam(FOnNewSubobjectAdded, const FSubobjectData&);
+
 /** Options when adding a new subobject */
 USTRUCT(BlueprintType, Category = "SubobjectDataSubsystem")
 struct FAddNewSubobjectParams
@@ -114,6 +117,14 @@ public:
 	*/
 	void GatherSubobjectData(UObject* Context, TArray<FSubobjectDataHandle>& OutArray);
 
+	////////////////////////////////////////////
+	// System events
+	
+	/** Delegate invoked when a new subobject is successfully added */
+	FOnNewSubobjectAdded& OnNewSubobjectAdded() { return OnNewSubobjectAdded_Delegate; };
+
+	////////////////////////////////////////////
+
 	/**
 	* Gather all subobjects that the given Blueprint context has. Populates an array of
 	* handles that will have the given context and all it's subobjects.
@@ -410,6 +421,8 @@ private:
 	FSubobjectDataHandle FindParentForNewSubobject(const UObject* NewSubobject, const FSubobjectDataHandle& SelectedParent);
 	
 	FSubobjectFactoryManager* FactoryManager = nullptr;
+
+	FOnNewSubobjectAdded OnNewSubobjectAdded_Delegate;
 	
 public:
 	FSubobjectFactoryManager* GetSubobjectFactoryManager() const { return FactoryManager; }
diff --git a/Engine/Source/Editor/SubobjectEditor/Private/SSubobjectBlueprintEditor.cpp b/Engine/Source/Editor/SubobjectEditor/Private/SSubobjectBlueprintEditor.cpp
--- a/Engine/Source/Editor/SubobjectEditor/Private/SSubobjectBlueprintEditor.cpp
+++ b/Engine/Source/Editor/SubobjectEditor/Private/SSubobjectBlueprintEditor.cpp
@@ -44,7 +44,6 @@ void SSubobjectBlueprintEditor::Construct(const FArguments& InArgs)
 	OnSelectionUpdated = InArgs._OnSelectionUpdated;
 	OnItemDoubleClicked = InArgs._OnItemDoubleClicked;
 	OnHighlightPropertyInDetailsView = InArgs._OnHighlightPropertyInDetailsView;
-	OnNewSubobjectAdded = InArgs._OnNewSubobjectAdded;
 	AllowEditing = InArgs._AllowEditing;
 	HideComponentClassCombo = InArgs._HideComponentClassCombo;
 	bAllowTreeUpdates = true;
@@ -526,17 +525,7 @@ FSubobjectDataHandle SSubobjectBlueprintEditor::AddNewSubobject(const FSubobject
     USubobjectDataSubsystem* System = USubobjectDataSubsystem::Get();
     check(System);
 	
-	FSubobjectDataHandle NewHandle = System->AddNewSubobject(Params,OutFailReason);
-
-	if (NewHandle.IsValid())
-	{
-		if (const FSubobjectData* NewSubobjectData = NewHandle.GetData())
-		{
-			OnNewSubobjectAdded.ExecuteIfBound(*NewSubobjectData);
-		}
-	}
-	
-    return NewHandle;	
+	return System->AddNewSubobject(Params, OutFailReason);
 }
 
 void SSubobjectBlueprintEditor::PopulateContextMenuImpl(UToolMenu* InMenu, TArray<FSubobjectEditorTreeNodePtrType>& InSelectedItems, bool bIsChildActorSubtreeNodeSelected)
diff --git a/Engine/Source/Editor/SubobjectEditor/Public/SSubobjectBlueprintEditor.h b/Engine/Source/Editor/SubobjectEditor/Public/SSubobjectBlueprintEditor.h
--- a/Engine/Source/Editor/SubobjectEditor/Public/SSubobjectBlueprintEditor.h
+++ b/Engine/Source/Editor/SubobjectEditor/Public/SSubobjectBlueprintEditor.h
@@ -13,7 +13,6 @@ class SUBOBJECTEDITOR_API SSubobjectBlueprintEditor final : public SSubobjectEdi
 {
 public:
 	DECLARE_DELEGATE_OneParam(FOnHighlightPropertyInDetailsView, const class FPropertyPath&);
-	DECLARE_DELEGATE_OneParam(FOnNewSubobjectAdded, const FSubobjectData&);
 
 private:	
 	SLATE_BEGIN_ARGS(SSubobjectBlueprintEditor)
@@ -23,7 +22,6 @@ private:
 		, _HideComponentClassCombo(false)
         , _OnSelectionUpdated()
 		, _OnHighlightPropertyInDetailsView()
-		, _OnNewSubobjectAdded()
 		{}
 
 		SLATE_ATTRIBUTE(UObject*, ObjectContext)
@@ -33,7 +31,6 @@ private:
 	    SLATE_EVENT(FOnSelectionUpdated, OnSelectionUpdated)
 	    SLATE_EVENT(FOnItemDoubleClicked, OnItemDoubleClicked)
 		SLATE_EVENT(FOnHighlightPropertyInDetailsView, OnHighlightPropertyInDetailsView)
-		SLATE_EVENT(FOnNewSubobjectAdded, OnNewSubobjectAdded)
 		SLATE_ARGUMENT(TArray<TSharedRef<IClassViewerFilter>>, SubobjectClassListFilters)
 	    
 	SLATE_END_ARGS()
@@ -73,9 +70,6 @@ public:
 
 	/** Delegate to invoke when the given property should be highlighted in the details view (e.g. diff). */
 	FOnHighlightPropertyInDetailsView OnHighlightPropertyInDetailsView;
-
-	/** Delegate to invoke when a subobject type is chosen to add a new subobject entry. */
-	FOnNewSubobjectAdded OnNewSubobjectAdded;
 	
 	/**
 	* Fills out an events section in ui.
