Fix instanced variable defaults not duplicating on copy/paste

diff --git a/Engine/Source/Editor/Kismet/Private/SMyBlueprint.cpp b/Engine/Source/Editor/Kismet/Private/SMyBlueprint.cpp
--- a/Engine/Source/Editor/Kismet/Private/SMyBlueprint.cpp
+++ b/Engine/Source/Editor/Kismet/Private/SMyBlueprint.cpp
@@ -3676,7 +3676,7 @@ void SMyBlueprint::OnPasteVariable()
 	FBPVariableDescription Description;
 	FStringOutputDevice Errors;
 	const TCHAR* Import = ClipboardText.GetCharArray().GetData() + FCString::Strlen(VAR_PREFIX);
-	FBPVariableDescription::StaticStruct()->ImportText(Import, &Description, nullptr, 0, &Errors, FBPVariableDescription::StaticStruct()->GetName());
+	FBPVariableDescription::StaticStruct()->ImportText(Import, &Description, nullptr, PPF_InstanceSubobjects, &Errors, FBPVariableDescription::StaticStruct()->GetName());
 	if (Errors.IsEmpty())
 	{
 		FBPVariableDescription NewVar = FBlueprintEditorUtils::DuplicateVariableDescription(Blueprint, Description);
