Another IWYU pass on Engine/Source/Editor/...

diff --git a/Engine/Source/Editor/AnimationBlueprintEditor/Private/AnimationBlueprintEditor.h b/Engine/Source/Editor/AnimationBlueprintEditor/Private/AnimationBlueprintEditor.h
--- a/Engine/Source/Editor/AnimationBlueprintEditor/Private/AnimationBlueprintEditor.h
+++ b/Engine/Source/Editor/AnimationBlueprintEditor/Private/AnimationBlueprintEditor.h
@@ -13,7 +13,6 @@
 #include "GraphEditor.h"
 #include "HAL/Platform.h"
 #include "IAnimationBlueprintEditor.h"
-#include "Internationalization/Internationalization.h"
 #include "Internationalization/Text.h"
 #include "Math/Color.h"
 #include "Misc/AssertionMacros.h"
diff --git a/Engine/Source/Editor/AnimationBlueprintEditor/Private/AnimationBlueprintEditorModule.cpp b/Engine/Source/Editor/AnimationBlueprintEditor/Private/AnimationBlueprintEditorModule.cpp
--- a/Engine/Source/Editor/AnimationBlueprintEditor/Private/AnimationBlueprintEditorModule.cpp
+++ b/Engine/Source/Editor/AnimationBlueprintEditor/Private/AnimationBlueprintEditorModule.cpp
@@ -6,6 +6,7 @@
 #include "AnimationBlueprintEditor.h"
 #include "AnimationBlueprintEditorSettings.h"
 #include "AnimationGraphFactory.h"
+#include "BlueprintEditor.h"
 #include "Delegates/Delegate.h"
 #include "Developer/MessageLog/Public/MessageLogModule.h"
 #include "EdGraph/EdGraph.h"
diff --git a/Engine/Source/Editor/AnimationBlueprintEditor/Private/AnimationGraphFactory.h b/Engine/Source/Editor/AnimationBlueprintEditor/Private/AnimationGraphFactory.h
--- a/Engine/Source/Editor/AnimationBlueprintEditor/Private/AnimationGraphFactory.h
+++ b/Engine/Source/Editor/AnimationBlueprintEditor/Private/AnimationGraphFactory.h
@@ -2,7 +2,6 @@
 
 #pragma once
 
-#include "ConnectionDrawingPolicy.h"
 #include "EdGraphUtilities.h"
 #include "HAL/Platform.h"
 #include "Templates/SharedPointer.h"
diff --git a/Engine/Source/Editor/AnimationBlueprintEditor/Private/AnimationNodes/SAnimationGraphNode.cpp b/Engine/Source/Editor/AnimationBlueprintEditor/Private/AnimationNodes/SAnimationGraphNode.cpp
--- a/Engine/Source/Editor/AnimationBlueprintEditor/Private/AnimationNodes/SAnimationGraphNode.cpp
+++ b/Engine/Source/Editor/AnimationBlueprintEditor/Private/AnimationNodes/SAnimationGraphNode.cpp
@@ -64,6 +64,7 @@
 #include "Widgets/SBoxPanel.h"
 #include "Widgets/SCompoundWidget.h"
 #include "Widgets/SNullWidget.h"
+#include "Widgets/SToolTip.h"
 #include "Widgets/SWidget.h"
 #include "Widgets/Text/SInlineEditableTextBlock.h"
 
diff --git a/Engine/Source/Editor/AnimationBlueprintEditor/Private/AnimationPins/SGraphPinPose.cpp b/Engine/Source/Editor/AnimationBlueprintEditor/Private/AnimationPins/SGraphPinPose.cpp
--- a/Engine/Source/Editor/AnimationBlueprintEditor/Private/AnimationPins/SGraphPinPose.cpp
+++ b/Engine/Source/Editor/AnimationBlueprintEditor/Private/AnimationPins/SGraphPinPose.cpp
@@ -29,7 +29,6 @@
 #include "Templates/Casts.h"
 #include "Templates/SharedPointer.h"
 #include "Templates/SubclassOf.h"
-#include "Templates/UnrealTemplate.h"
 #include "UObject/UObjectGlobals.h"
 #include "UObject/WeakObjectPtr.h"
 #include "Widgets/Images/SImage.h"
diff --git a/Engine/Source/Editor/AnimationBlueprintEditor/Private/AnimationStateNodes/SGraphNodeAnimTransition.cpp b/Engine/Source/Editor/AnimationBlueprintEditor/Private/AnimationStateNodes/SGraphNodeAnimTransition.cpp
--- a/Engine/Source/Editor/AnimationBlueprintEditor/Private/AnimationStateNodes/SGraphNodeAnimTransition.cpp
+++ b/Engine/Source/Editor/AnimationBlueprintEditor/Private/AnimationStateNodes/SGraphNodeAnimTransition.cpp
@@ -22,7 +22,6 @@
 #include "EdGraphSchema_K2.h"
 #include "Engine/Blueprint.h"
 #include "HAL/PlatformCrt.h"
-#include "HAL/PlatformMath.h"
 #include "IDocumentation.h"
 #include "Internationalization/Internationalization.h"
 #include "Kismet2/BlueprintEditorUtils.h"
diff --git a/Engine/Source/Editor/AudioEditor/Private/SoundClassGraph.cpp b/Engine/Source/Editor/AudioEditor/Private/SoundClassGraph.cpp
--- a/Engine/Source/Editor/AudioEditor/Private/SoundClassGraph.cpp
+++ b/Engine/Source/Editor/AudioEditor/Private/SoundClassGraph.cpp
@@ -9,7 +9,6 @@
 #include "EdGraph/EdGraphPin.h"
 #include "GraphEditor.h"
 #include "HAL/PlatformCrt.h"
-#include "HAL/PlatformMath.h"
 #include "Misc/AssertionMacros.h"
 #include "Sound/SoundClass.h"
 #include "SoundClassGraph/SoundClassGraphNode.h"
diff --git a/Engine/Source/Editor/AudioEditor/Private/SoundCueEditorUtilities.cpp b/Engine/Source/Editor/AudioEditor/Private/SoundCueEditorUtilities.cpp
--- a/Engine/Source/Editor/AudioEditor/Private/SoundCueEditorUtilities.cpp
+++ b/Engine/Source/Editor/AudioEditor/Private/SoundCueEditorUtilities.cpp
@@ -5,7 +5,6 @@
 #include "EdGraph/EdGraph.h"
 #include "EdGraph/EdGraphNode.h"
 #include "HAL/PlatformCrt.h"
-#include "HAL/PlatformMath.h"
 #include "ISoundCueEditor.h"
 #include "Misc/AssertionMacros.h"
 #include "Sound/DialogueTypes.h"
diff --git a/Engine/Source/Editor/AudioEditor/Private/SoundCueGraphConnectionDrawingPolicy.cpp b/Engine/Source/Editor/AudioEditor/Private/SoundCueGraphConnectionDrawingPolicy.cpp
--- a/Engine/Source/Editor/AudioEditor/Private/SoundCueGraphConnectionDrawingPolicy.cpp
+++ b/Engine/Source/Editor/AudioEditor/Private/SoundCueGraphConnectionDrawingPolicy.cpp
@@ -14,7 +14,6 @@
 #include "Editor/EditorEngine.h"
 #include "GraphEditorSettings.h"
 #include "HAL/PlatformCrt.h"
-#include "HAL/PlatformMath.h"
 #include "Math/Vector2D.h"
 #include "Misc/App.h"
 #include "Misc/AssertionMacros.h"
diff --git a/Engine/Source/Editor/AudioEditor/Private/SoundCueGraphNode_Base.cpp b/Engine/Source/Editor/AudioEditor/Private/SoundCueGraphNode_Base.cpp
--- a/Engine/Source/Editor/AudioEditor/Private/SoundCueGraphNode_Base.cpp
+++ b/Engine/Source/Editor/AudioEditor/Private/SoundCueGraphNode_Base.cpp
@@ -10,7 +10,6 @@
 #include "EdGraph/EdGraphPin.h"
 #include "EdGraph/EdGraphSchema.h"
 #include "HAL/PlatformCrt.h"
-#include "HAL/PlatformMath.h"
 #include "Misc/AssertionMacros.h"
 #include "SoundCueGraph/SoundCueGraphSchema.h"
 #include "Templates/Casts.h"
diff --git a/Engine/Source/Editor/BehaviorTreeEditor/Private/BehaviorTreeDecoratorGraph.cpp b/Engine/Source/Editor/BehaviorTreeEditor/Private/BehaviorTreeDecoratorGraph.cpp
--- a/Engine/Source/Editor/BehaviorTreeEditor/Private/BehaviorTreeDecoratorGraph.cpp
+++ b/Engine/Source/Editor/BehaviorTreeEditor/Private/BehaviorTreeDecoratorGraph.cpp
@@ -10,7 +10,6 @@
 #include "EdGraph/EdGraphPin.h"
 #include "EdGraphSchema_BehaviorTreeDecorator.h"
 #include "HAL/PlatformCrt.h"
-#include "HAL/PlatformMath.h"
 #include "Templates/Casts.h"
 #include "Templates/SubclassOf.h"
 #include "UObject/ObjectPtr.h"
diff --git a/Engine/Source/Editor/BehaviorTreeEditor/Private/BehaviorTreeGraphNode_SubtreeTask.cpp b/Engine/Source/Editor/BehaviorTreeEditor/Private/BehaviorTreeGraphNode_SubtreeTask.cpp
--- a/Engine/Source/Editor/BehaviorTreeEditor/Private/BehaviorTreeGraphNode_SubtreeTask.cpp
+++ b/Engine/Source/Editor/BehaviorTreeEditor/Private/BehaviorTreeGraphNode_SubtreeTask.cpp
@@ -14,7 +14,6 @@
 #include "EdGraph/EdGraph.h"
 #include "EdGraph/EdGraphPin.h"
 #include "HAL/PlatformCrt.h"
-#include "HAL/PlatformMath.h"
 #include "Templates/Casts.h"
 #include "UObject/Object.h"
 #include "UObject/ObjectPtr.h"
diff --git a/Engine/Source/Editor/BehaviorTreeEditor/Private/FindInBT.cpp b/Engine/Source/Editor/BehaviorTreeEditor/Private/FindInBT.cpp
--- a/Engine/Source/Editor/BehaviorTreeEditor/Private/FindInBT.cpp
+++ b/Engine/Source/Editor/BehaviorTreeEditor/Private/FindInBT.cpp
@@ -5,7 +5,6 @@
 #include "BehaviorTreeGraphNode.h"
 #include "BehaviorTreeGraphNode_Decorator.h"
 #include "BehaviorTreeGraphNode_Service.h"
-#include "Delegates/IDelegateInstance.h"
 #include "EdGraph/EdGraph.h"
 #include "EdGraph/EdGraphNode.h"
 #include "Framework/Application/SlateApplication.h"
diff --git a/Engine/Source/Editor/BlueprintGraph/Classes/K2Node.h b/Engine/Source/Editor/BlueprintGraph/Classes/K2Node.h
--- a/Engine/Source/Editor/BlueprintGraph/Classes/K2Node.h
+++ b/Engine/Source/Editor/BlueprintGraph/Classes/K2Node.h
@@ -2,6 +2,7 @@
 
 #pragma once
 
+#include "BlueprintActionFilter.h"
 #include "BlueprintNodeSignature.h"
 #include "Containers/Array.h"
 #include "Containers/Map.h"
diff --git a/Engine/Source/Editor/BlueprintGraph/Private/BlueprintEditorSettings.cpp b/Engine/Source/Editor/BlueprintGraph/Private/BlueprintEditorSettings.cpp
--- a/Engine/Source/Editor/BlueprintGraph/Private/BlueprintEditorSettings.cpp
+++ b/Engine/Source/Editor/BlueprintGraph/Private/BlueprintEditorSettings.cpp
@@ -14,7 +14,6 @@
 #include "EdGraphSchema_K2.h"
 #include "Editor/EditorPerProjectUserSettings.h"
 #include "Engine/Blueprint.h"
-#include "HAL/Platform.h"
 #include "Misc/AssertionMacros.h"
 #include "Misc/ConfigCacheIni.h"
 #include "Modules/ModuleManager.h"
diff --git a/Engine/Source/Editor/BlueprintGraph/Private/K2Node_BaseAsyncTask.cpp b/Engine/Source/Editor/BlueprintGraph/Private/K2Node_BaseAsyncTask.cpp
--- a/Engine/Source/Editor/BlueprintGraph/Private/K2Node_BaseAsyncTask.cpp
+++ b/Engine/Source/Editor/BlueprintGraph/Private/K2Node_BaseAsyncTask.cpp
@@ -15,7 +15,6 @@
 #include "Engine/MemberReference.h"
 #include "EngineLogs.h"
 #include "HAL/PlatformCrt.h"
-#include "HAL/PlatformMath.h"
 #include "Internationalization/Internationalization.h"
 #include "K2Node_AddDelegate.h"
 #include "K2Node_AssignmentStatement.h"
diff --git a/Engine/Source/Editor/BlueprintGraph/Private/K2Node_CallParentFunction.cpp b/Engine/Source/Editor/BlueprintGraph/Private/K2Node_CallParentFunction.cpp
--- a/Engine/Source/Editor/BlueprintGraph/Private/K2Node_CallParentFunction.cpp
+++ b/Engine/Source/Editor/BlueprintGraph/Private/K2Node_CallParentFunction.cpp
@@ -13,6 +13,7 @@
 #include "K2Node.h"
 #include "Misc/Guid.h"
 #include "Settings/EditorStyleSettings.h"
+#include "Templates/Casts.h"
 #include "UObject/Class.h"
 #include "UObject/NameTypes.h"
 #include "UObject/Script.h"
diff --git a/Engine/Source/Editor/BlueprintGraph/Private/K2Node_ConstructObjectFromClass.cpp b/Engine/Source/Editor/BlueprintGraph/Private/K2Node_ConstructObjectFromClass.cpp
--- a/Engine/Source/Editor/BlueprintGraph/Private/K2Node_ConstructObjectFromClass.cpp
+++ b/Engine/Source/Editor/BlueprintGraph/Private/K2Node_ConstructObjectFromClass.cpp
@@ -11,7 +11,6 @@
 #include "EditorCategoryUtils.h"
 #include "Engine/Blueprint.h"
 #include "FindInBlueprintManager.h"
-#include "HAL/Platform.h"
 #include "HAL/PlatformCrt.h"
 #include "Internationalization/Internationalization.h"
 #include "Kismet2/BlueprintEditorUtils.h"
diff --git a/Engine/Source/Editor/BlueprintGraph/Private/K2Node_CreateDelegate.cpp b/Engine/Source/Editor/BlueprintGraph/Private/K2Node_CreateDelegate.cpp
--- a/Engine/Source/Editor/BlueprintGraph/Private/K2Node_CreateDelegate.cpp
+++ b/Engine/Source/Editor/BlueprintGraph/Private/K2Node_CreateDelegate.cpp
@@ -13,7 +13,6 @@
 #include "Engine/MemberReference.h"
 #include "FindInBlueprintManager.h"
 #include "HAL/PlatformCrt.h"
-#include "HAL/PlatformMath.h"
 #include "Internationalization/Internationalization.h"
 #include "K2Node_BaseMCDelegate.h"
 #include "K2Node_Event.h"
diff --git a/Engine/Source/Editor/BlueprintGraph/Private/K2Node_CustomEvent.cpp b/Engine/Source/Editor/BlueprintGraph/Private/K2Node_CustomEvent.cpp
--- a/Engine/Source/Editor/BlueprintGraph/Private/K2Node_CustomEvent.cpp
+++ b/Engine/Source/Editor/BlueprintGraph/Private/K2Node_CustomEvent.cpp
@@ -16,7 +16,6 @@
 #include "Engine/MemberReference.h"
 #include "FindInBlueprintManager.h"
 #include "HAL/PlatformCrt.h"
-#include "HAL/PlatformMath.h"
 #include "Internationalization/Internationalization.h"
 #include "K2Node.h"
 #include "K2Node_BaseMCDelegate.h"
diff --git a/Engine/Source/Editor/BlueprintGraph/Private/K2Node_DelegateSet.cpp b/Engine/Source/Editor/BlueprintGraph/Private/K2Node_DelegateSet.cpp
--- a/Engine/Source/Editor/BlueprintGraph/Private/K2Node_DelegateSet.cpp
+++ b/Engine/Source/Editor/BlueprintGraph/Private/K2Node_DelegateSet.cpp
@@ -13,7 +13,6 @@
 #include "EdGraphUtilities.h"
 #include "Engine/MemberReference.h"
 #include "HAL/PlatformCrt.h"
-#include "HAL/PlatformMath.h"
 #include "Internationalization/Internationalization.h"
 #include "K2Node_Event.h"
 #include "Kismet2/CompilerResultsLog.h"
diff --git a/Engine/Source/Editor/BlueprintGraph/Private/K2Node_DynamicCast.cpp b/Engine/Source/Editor/BlueprintGraph/Private/K2Node_DynamicCast.cpp
--- a/Engine/Source/Editor/BlueprintGraph/Private/K2Node_DynamicCast.cpp
+++ b/Engine/Source/Editor/BlueprintGraph/Private/K2Node_DynamicCast.cpp
@@ -17,7 +17,6 @@
 #include "Engine/Blueprint.h"
 #include "Framework/Commands/UIAction.h"
 #include "HAL/PlatformCrt.h"
-#include "HAL/PlatformMath.h"
 #include "Internationalization/Internationalization.h"
 #include "Kismet2/CompilerResultsLog.h"
 #include "Kismet2/KismetEditorUtilities.h"
diff --git a/Engine/Source/Editor/BlueprintGraph/Private/K2Node_Event.cpp b/Engine/Source/Editor/BlueprintGraph/Private/K2Node_Event.cpp
--- a/Engine/Source/Editor/BlueprintGraph/Private/K2Node_Event.cpp
+++ b/Engine/Source/Editor/BlueprintGraph/Private/K2Node_Event.cpp
@@ -16,7 +16,6 @@
 #include "GameFramework/Actor.h"
 #include "GraphEditorSettings.h"
 #include "HAL/PlatformCrt.h"
-#include "HAL/PlatformMath.h"
 #include "Internationalization/Internationalization.h"
 #include "K2Node_CallFunction.h"
 #include "K2Node_CreateDelegate.h"
diff --git a/Engine/Source/Editor/BlueprintGraph/Private/K2Node_FormatText.cpp b/Engine/Source/Editor/BlueprintGraph/Private/K2Node_FormatText.cpp
--- a/Engine/Source/Editor/BlueprintGraph/Private/K2Node_FormatText.cpp
+++ b/Engine/Source/Editor/BlueprintGraph/Private/K2Node_FormatText.cpp
@@ -14,7 +14,6 @@
 #include "EditorCategoryUtils.h"
 #include "Engine/Blueprint.h"
 #include "HAL/PlatformCrt.h"
-#include "HAL/PlatformMath.h"
 #include "Internationalization/Internationalization.h"
 #include "K2Node_CallFunction.h"
 #include "K2Node_MakeArray.h"
diff --git a/Engine/Source/Editor/BlueprintGraph/Private/K2Node_FunctionEntry.cpp b/Engine/Source/Editor/BlueprintGraph/Private/K2Node_FunctionEntry.cpp
--- a/Engine/Source/Editor/BlueprintGraph/Private/K2Node_FunctionEntry.cpp
+++ b/Engine/Source/Editor/BlueprintGraph/Private/K2Node_FunctionEntry.cpp
@@ -19,7 +19,6 @@
 #include "Engine/Blueprint.h"
 #include "Engine/MemberReference.h"
 #include "EngineLogs.h"
-#include "HAL/Platform.h"
 #include "HAL/PlatformCrt.h"
 #include "Internationalization/Internationalization.h"
 #include "K2Node_CallFunction.h"
diff --git a/Engine/Source/Editor/BlueprintGraph/Private/K2Node_GetArrayItem.cpp b/Engine/Source/Editor/BlueprintGraph/Private/K2Node_GetArrayItem.cpp
--- a/Engine/Source/Editor/BlueprintGraph/Private/K2Node_GetArrayItem.cpp
+++ b/Engine/Source/Editor/BlueprintGraph/Private/K2Node_GetArrayItem.cpp
@@ -22,7 +22,6 @@
 #include "Framework/Commands/UIAction.h"
 #include "Framework/Notifications/NotificationManager.h"
 #include "HAL/PlatformCrt.h"
-#include "HAL/PlatformMath.h"
 #include "Internationalization/Internationalization.h"
 #include "K2Node_CallArrayFunction.h"
 #include "K2Node_CallFunction.h"
diff --git a/Engine/Source/Editor/BlueprintGraph/Private/K2Node_Literal.cpp b/Engine/Source/Editor/BlueprintGraph/Private/K2Node_Literal.cpp
--- a/Engine/Source/Editor/BlueprintGraph/Private/K2Node_Literal.cpp
+++ b/Engine/Source/Editor/BlueprintGraph/Private/K2Node_Literal.cpp
@@ -19,7 +19,6 @@
 #include "EdGraphSchema_K2.h"
 #include "GameFramework/Actor.h"
 #include "HAL/PlatformCrt.h"
-#include "HAL/PlatformMath.h"
 #include "Internationalization/Internationalization.h"
 #include "Kismet2/CompilerResultsLog.h"
 #include "KismetCompiledFunctionContext.h"
diff --git a/Engine/Source/Editor/BlueprintGraph/Private/K2Node_LoadAsset.cpp b/Engine/Source/Editor/BlueprintGraph/Private/K2Node_LoadAsset.cpp
--- a/Engine/Source/Editor/BlueprintGraph/Private/K2Node_LoadAsset.cpp
+++ b/Engine/Source/Editor/BlueprintGraph/Private/K2Node_LoadAsset.cpp
@@ -23,7 +23,6 @@
 #include "Kismet2/CompilerResultsLog.h"
 #include "KismetCompiler.h"
 #include "Misc/AssertionMacros.h"
-#include "Templates/Casts.h"
 #include "UObject/Class.h"
 #include "UObject/Object.h"
 #include "UObject/UnrealNames.h"
diff --git a/Engine/Source/Editor/BlueprintGraph/Private/K2Node_MacroInstance.cpp b/Engine/Source/Editor/BlueprintGraph/Private/K2Node_MacroInstance.cpp
--- a/Engine/Source/Editor/BlueprintGraph/Private/K2Node_MacroInstance.cpp
+++ b/Engine/Source/Editor/BlueprintGraph/Private/K2Node_MacroInstance.cpp
@@ -13,7 +13,6 @@
 #include "Engine/Blueprint.h"
 #include "Framework/Commands/UIAction.h"
 #include "HAL/PlatformCrt.h"
-#include "HAL/PlatformMath.h"
 #include "Internationalization/Internationalization.h"
 #include "K2Node_EditablePinBase.h"
 #include "Kismet2/BlueprintEditorUtils.h"
diff --git a/Engine/Source/Editor/BlueprintGraph/Private/K2Node_MultiGate.cpp b/Engine/Source/Editor/BlueprintGraph/Private/K2Node_MultiGate.cpp
--- a/Engine/Source/Editor/BlueprintGraph/Private/K2Node_MultiGate.cpp
+++ b/Engine/Source/Editor/BlueprintGraph/Private/K2Node_MultiGate.cpp
@@ -14,7 +14,6 @@
 #include "EdGraphSchema_K2.h"
 #include "EdGraphUtilities.h"
 #include "HAL/PlatformCrt.h"
-#include "HAL/PlatformMath.h"
 #include "Internationalization/Internationalization.h"
 #include "K2Node_AssignmentStatement.h"
 #include "K2Node_TemporaryVariable.h"
diff --git a/Engine/Source/Editor/BlueprintGraph/Private/K2Node_Switch.cpp b/Engine/Source/Editor/BlueprintGraph/Private/K2Node_Switch.cpp
--- a/Engine/Source/Editor/BlueprintGraph/Private/K2Node_Switch.cpp
+++ b/Engine/Source/Editor/BlueprintGraph/Private/K2Node_Switch.cpp
@@ -16,7 +16,6 @@
 #include "EditorCategoryUtils.h"
 #include "Engine/Blueprint.h"
 #include "HAL/PlatformCrt.h"
-#include "HAL/PlatformMath.h"
 #include "Internationalization/Internationalization.h"
 #include "Kismet2/CompilerResultsLog.h"
 #include "KismetCompiledFunctionContext.h"
diff --git a/Engine/Source/Editor/BlueprintGraph/Private/K2Node_SwitchName.cpp b/Engine/Source/Editor/BlueprintGraph/Private/K2Node_SwitchName.cpp
--- a/Engine/Source/Editor/BlueprintGraph/Private/K2Node_SwitchName.cpp
+++ b/Engine/Source/Editor/BlueprintGraph/Private/K2Node_SwitchName.cpp
@@ -8,7 +8,6 @@
 #include "EdGraph/EdGraph.h"
 #include "EdGraphSchema_K2.h"
 #include "HAL/PlatformCrt.h"
-#include "HAL/PlatformMath.h"
 #include "Internationalization/Internationalization.h"
 #include "Kismet/KismetMathLibrary.h"
 #include "Misc/AssertionMacros.h"
diff --git a/Engine/Source/Editor/BlueprintGraph/Private/K2Node_SwitchString.cpp b/Engine/Source/Editor/BlueprintGraph/Private/K2Node_SwitchString.cpp
--- a/Engine/Source/Editor/BlueprintGraph/Private/K2Node_SwitchString.cpp
+++ b/Engine/Source/Editor/BlueprintGraph/Private/K2Node_SwitchString.cpp
@@ -8,7 +8,6 @@
 #include "EdGraph/EdGraph.h"
 #include "EdGraphSchema_K2.h"
 #include "HAL/PlatformCrt.h"
-#include "HAL/PlatformMath.h"
 #include "Internationalization/Internationalization.h"
 #include "Kismet/KismetStringLibrary.h"
 #include "Kismet2/CompilerResultsLog.h"
diff --git a/Engine/Source/Editor/BlueprintGraph/Private/K2Node_Variable.cpp b/Engine/Source/Editor/BlueprintGraph/Private/K2Node_Variable.cpp
--- a/Engine/Source/Editor/BlueprintGraph/Private/K2Node_Variable.cpp
+++ b/Engine/Source/Editor/BlueprintGraph/Private/K2Node_Variable.cpp
@@ -18,7 +18,6 @@
 #include "GameFramework/Actor.h"
 #include "GameFramework/MovementComponent.h"
 #include "HAL/PlatformCrt.h"
-#include "HAL/PlatformMath.h"
 #include "Internationalization/Internationalization.h"
 #include "K2Node_FunctionEntry.h"
 #include "Kismet2/BlueprintEditorUtils.h"
diff --git a/Engine/Source/Editor/Blutility/Private/AssetTypeActions_EditorUtilityWidgetBlueprint.cpp b/Engine/Source/Editor/Blutility/Private/AssetTypeActions_EditorUtilityWidgetBlueprint.cpp
--- a/Engine/Source/Editor/Blutility/Private/AssetTypeActions_EditorUtilityWidgetBlueprint.cpp
+++ b/Engine/Source/Editor/Blutility/Private/AssetTypeActions_EditorUtilityWidgetBlueprint.cpp
@@ -2,6 +2,7 @@
 
 #include "AssetTypeActions_EditorUtilityWidgetBlueprint.h"
 
+#include "BlueprintEditor.h"
 #include "Containers/UnrealString.h"
 #include "Delegates/Delegate.h"
 #include "Editor.h"
diff --git a/Engine/Source/Editor/ClassViewer/Private/ClassViewerFilter.cpp b/Engine/Source/Editor/ClassViewer/Private/ClassViewerFilter.cpp
--- a/Engine/Source/Editor/ClassViewer/Private/ClassViewerFilter.cpp
+++ b/Engine/Source/Editor/ClassViewer/Private/ClassViewerFilter.cpp
@@ -14,7 +14,6 @@
 #include "Engine/BlueprintGeneratedClass.h"
 #include "Engine/Brush.h"
 #include "GameFramework/Actor.h"
-#include "HAL/Platform.h"
 #include "HAL/PlatformCrt.h"
 #include "Kismet2/KismetEditorUtilities.h"
 #include "Logging/LogCategory.h"
diff --git a/Engine/Source/Editor/ClothingSystemEditor/Private/ClothingAssetFactory.cpp b/Engine/Source/Editor/ClothingSystemEditor/Private/ClothingAssetFactory.cpp
--- a/Engine/Source/Editor/ClothingSystemEditor/Private/ClothingAssetFactory.cpp
+++ b/Engine/Source/Editor/ClothingSystemEditor/Private/ClothingAssetFactory.cpp
@@ -26,6 +26,7 @@
 #include "Misc/ConfigCacheIni.h"
 #include "Misc/Guid.h"
 #include "ObjectTools.h"
+#include "PhysicsEngine/PhysicsAsset.h"
 #include "PointWeightMap.h"
 #include "ReferenceSkeleton.h"
 #include "Rendering/SkeletalMeshLODModel.h"
diff --git a/Engine/Source/Editor/ContentBrowser/Private/DragDropHandler.cpp b/Engine/Source/Editor/ContentBrowser/Private/DragDropHandler.cpp
--- a/Engine/Source/Editor/ContentBrowser/Private/DragDropHandler.cpp
+++ b/Engine/Source/Editor/ContentBrowser/Private/DragDropHandler.cpp
@@ -26,9 +26,9 @@
 #include "Misc/AssertionMacros.h"
 #include "Misc/EnumClassFlags.h"
 #include "Styling/AppStyle.h"
+#include "Templates/Casts.h"
 #include "Templates/Invoke.h"
 #include "Templates/Tuple.h"
-#include "Templates/UnrealTemplate.h"
 #include "Textures/SlateIcon.h"
 #include "ToolMenu.h"
 #include "ToolMenuContext.h"
diff --git a/Engine/Source/Editor/ContentBrowser/Private/NewAssetOrClassContextMenu.cpp b/Engine/Source/Editor/ContentBrowser/Private/NewAssetOrClassContextMenu.cpp
--- a/Engine/Source/Editor/ContentBrowser/Private/NewAssetOrClassContextMenu.cpp
+++ b/Engine/Source/Editor/ContentBrowser/Private/NewAssetOrClassContextMenu.cpp
@@ -11,7 +11,6 @@
 #include "Misc/AssertionMacros.h"
 #include "Settings/ContentBrowserSettings.h"
 #include "Styling/AppStyle.h"
-#include "Templates/Casts.h"
 #include "Textures/SlateIcon.h"
 #include "ToolMenu.h"
 #include "ToolMenuSection.h"
diff --git a/Engine/Source/Editor/ContentBrowser/Private/SAssetPicker.cpp b/Engine/Source/Editor/ContentBrowser/Private/SAssetPicker.cpp
--- a/Engine/Source/Editor/ContentBrowser/Private/SAssetPicker.cpp
+++ b/Engine/Source/Editor/ContentBrowser/Private/SAssetPicker.cpp
@@ -34,6 +34,7 @@
 #include "SContentBrowser.h"
 #include "SFilterList.h"
 #include "SlotBase.h"
+#include "SourceControlOperations.h"
 #include "Styling/AppStyle.h"
 #include "Styling/ISlateStyle.h"
 #include "Styling/SlateColor.h"
diff --git a/Engine/Source/Editor/ContentBrowser/Private/SContentBrowser.cpp b/Engine/Source/Editor/ContentBrowser/Private/SContentBrowser.cpp
--- a/Engine/Source/Editor/ContentBrowser/Private/SContentBrowser.cpp
+++ b/Engine/Source/Editor/ContentBrowser/Private/SContentBrowser.cpp
@@ -29,7 +29,6 @@
 #include "ContentBrowserLog.h"
 #include "ContentBrowserMenuContexts.h"
 #include "ContentBrowserModule.h"
-#include "ContentBrowserPluginFilters.h"
 #include "ContentBrowserSingleton.h"
 #include "ContentBrowserUtils.h"
 #include "ContentBrowserVirtualPathTree.h"
diff --git a/Engine/Source/Editor/ContentBrowser/Private/SFilterList.cpp b/Engine/Source/Editor/ContentBrowser/Private/SFilterList.cpp
--- a/Engine/Source/Editor/ContentBrowser/Private/SFilterList.cpp
+++ b/Engine/Source/Editor/ContentBrowser/Private/SFilterList.cpp
@@ -13,10 +13,12 @@
 #include "ContentBrowserMenuContexts.h"
 #include "ContentBrowserUtils.h"
 #include "Filters/FilterBarConfig.h"
+#include "Filters/SAssetFilterBar.h"
 #include "Framework/Application/MenuStack.h"
 #include "Framework/Application/SlateApplication.h"
 #include "FrontendFilterBase.h"
 #include "FrontendFilters.h"
+#include "HAL/PlatformCrt.h"
 #include "IContentBrowserDataModule.h"
 #include "Input/Events.h"
 #include "Internationalization/Internationalization.h"
@@ -29,6 +31,7 @@
 #include "SlotBase.h"
 #include "Styling/SlateTypes.h"
 #include "Templates/Casts.h"
+#include "Templates/TypeHash.h"
 #include "ToolMenu.h"
 #include "ToolMenuContext.h"
 #include "ToolMenuDelegates.h"
diff --git a/Engine/Source/Editor/ContentBrowser/Private/SFilterList.h b/Engine/Source/Editor/ContentBrowser/Private/SFilterList.h
--- a/Engine/Source/Editor/ContentBrowser/Private/SFilterList.h
+++ b/Engine/Source/Editor/ContentBrowser/Private/SFilterList.h
@@ -14,14 +14,12 @@
 #include "Filters/SBasicFilterBar.h"
 #include "FrontendFilters.h"
 #include "HAL/Platform.h"
-#include "HAL/PlatformCrt.h"
 #include "IContentBrowserSingleton.h"
 #include "Input/Reply.h"
 #include "InputCoreTypes.h"
 #include "Internationalization/Text.h"
 #include "Math/Color.h"
 #include "Templates/SharedPointer.h"
-#include "Templates/TypeHash.h"
 #include "UObject/NameTypes.h"
 #include "Widgets/DeclarativeSyntaxSupport.h"
 
diff --git a/Engine/Source/Editor/ContentBrowser/Private/SSearchToggleButton.h b/Engine/Source/Editor/ContentBrowser/Private/SSearchToggleButton.h
--- a/Engine/Source/Editor/ContentBrowser/Private/SSearchToggleButton.h
+++ b/Engine/Source/Editor/ContentBrowser/Private/SSearchToggleButton.h
@@ -7,6 +7,7 @@
 #include "Styling/AppStyle.h"
 #include "Styling/ISlateStyle.h"
 #include "Styling/SlateTypes.h"
+#include "Styling/SlateWidgetStyleAsset.h"
 #include "Templates/SharedPointer.h"
 #include "Widgets/DeclarativeSyntaxSupport.h"
 #include "Widgets/SCompoundWidget.h"
diff --git a/Engine/Source/Editor/CurveEditor/Private/CurveEditorModule.cpp b/Engine/Source/Editor/CurveEditor/Private/CurveEditorModule.cpp
--- a/Engine/Source/Editor/CurveEditor/Private/CurveEditorModule.cpp
+++ b/Engine/Source/Editor/CurveEditor/Private/CurveEditorModule.cpp
@@ -11,7 +11,6 @@
 #include "ICurveEditorModule.h"
 #include "Misc/CoreDelegates.h"
 #include "Modules/ModuleManager.h"
-#include "Templates/UnrealTemplate.h"
 #include "ToolMenus.h"
 
 class FCurveEditorModule : public ICurveEditorModule
diff --git a/Engine/Source/Editor/DataLayerEditor/Private/DataLayer/DataLayerEditorSubsystem.cpp b/Engine/Source/Editor/DataLayerEditor/Private/DataLayer/DataLayerEditorSubsystem.cpp
--- a/Engine/Source/Editor/DataLayerEditor/Private/DataLayer/DataLayerEditorSubsystem.cpp
+++ b/Engine/Source/Editor/DataLayerEditor/Private/DataLayer/DataLayerEditorSubsystem.cpp
@@ -15,7 +15,6 @@
 #include "Engine/World.h"
 #include "EngineUtils.h"
 #include "GameFramework/Actor.h"
-#include "HAL/Platform.h"
 #include "HAL/PlatformCrt.h"
 #include "Internationalization/Internationalization.h"
 #include "Internationalization/Text.h"
diff --git a/Engine/Source/Editor/DataLayerEditor/Private/DataLayer/SDataLayerBrowser.cpp b/Engine/Source/Editor/DataLayerEditor/Private/DataLayer/SDataLayerBrowser.cpp
--- a/Engine/Source/Editor/DataLayerEditor/Private/DataLayer/SDataLayerBrowser.cpp
+++ b/Engine/Source/Editor/DataLayerEditor/Private/DataLayer/SDataLayerBrowser.cpp
@@ -52,6 +52,7 @@
 #include "WorldPartition/WorldPartitionActorDesc.h"
 #include "WorldPartition/WorldPartitionHandle.h"
 
+class ISceneOutliner;
 class UObject;
 
 #define LOCTEXT_NAMESPACE "DataLayer"
diff --git a/Engine/Source/Editor/DataTableEditor/Private/DataTableEditor.cpp b/Engine/Source/Editor/DataTableEditor/Private/DataTableEditor.cpp
--- a/Engine/Source/Editor/DataTableEditor/Private/DataTableEditor.cpp
+++ b/Engine/Source/Editor/DataTableEditor/Private/DataTableEditor.cpp
@@ -78,6 +78,7 @@
 #include "Widgets/Notifications/SNotificationList.h"
 #include "Widgets/SBoxPanel.h"
 #include "Widgets/SNullWidget.h"
+#include "Widgets/SToolTip.h"
 #include "Widgets/Text/STextBlock.h"
 #include "Widgets/Views/SListView.h"
 
diff --git a/Engine/Source/Editor/DetailCustomizations/Private/MovieSceneCaptureCustomization.h b/Engine/Source/Editor/DetailCustomizations/Private/MovieSceneCaptureCustomization.h
--- a/Engine/Source/Editor/DetailCustomizations/Private/MovieSceneCaptureCustomization.h
+++ b/Engine/Source/Editor/DetailCustomizations/Private/MovieSceneCaptureCustomization.h
@@ -7,7 +7,6 @@
 #include "Delegates/IDelegateInstance.h"
 #include "IDetailCustomization.h"
 #include "Templates/SharedPointer.h"
-#include "UObject/WeakObjectPtr.h"
 #include "UObject/WeakObjectPtrTemplates.h"
 
 class IDetailLayoutBuilder;
diff --git a/Engine/Source/Editor/DetailCustomizations/Private/PhysicsConstraintComponentDetails.cpp b/Engine/Source/Editor/DetailCustomizations/Private/PhysicsConstraintComponentDetails.cpp
--- a/Engine/Source/Editor/DetailCustomizations/Private/PhysicsConstraintComponentDetails.cpp
+++ b/Engine/Source/Editor/DetailCustomizations/Private/PhysicsConstraintComponentDetails.cpp
@@ -26,7 +26,6 @@
 #include "Math/Matrix.h"
 #include "Math/RotationMatrix.h"
 #include "Math/Rotator.h"
-#include "Math/UnrealMathUtility.h"
 #include "Misc/AssertionMacros.h"
 #include "Misc/Attribute.h"
 #include "Misc/EnumClassFlags.h"
diff --git a/Engine/Source/Editor/DeviceProfileEditor/Private/DetailsPanel/SDeviceProfileDetailsPanel.h b/Engine/Source/Editor/DeviceProfileEditor/Private/DetailsPanel/SDeviceProfileDetailsPanel.h
--- a/Engine/Source/Editor/DeviceProfileEditor/Private/DetailsPanel/SDeviceProfileDetailsPanel.h
+++ b/Engine/Source/Editor/DeviceProfileEditor/Private/DetailsPanel/SDeviceProfileDetailsPanel.h
@@ -2,6 +2,7 @@
 
 #pragma once
 
+#include "SDeviceProfileDetailsPanel.h"
 #include "Templates/SharedPointer.h"
 #include "UObject/WeakObjectPtr.h"
 #include "UObject/WeakObjectPtrTemplates.h"
diff --git a/Engine/Source/Editor/Documentation/Private/SDocumentationAnchor.cpp b/Engine/Source/Editor/Documentation/Private/SDocumentationAnchor.cpp
--- a/Engine/Source/Editor/Documentation/Private/SDocumentationAnchor.cpp
+++ b/Engine/Source/Editor/Documentation/Private/SDocumentationAnchor.cpp
@@ -14,6 +14,7 @@
 #include "SSimpleButton.h"
 #include "Styling/AppStyle.h"
 #include "Styling/ISlateStyle.h"
+#include "Widgets/SToolTip.h"
 
 void SDocumentationAnchor::Construct(const FArguments& InArgs )
 {
diff --git a/Engine/Source/Editor/EditorWidgets/Private/SObjectNameEditableTextBox.cpp b/Engine/Source/Editor/EditorWidgets/Private/SObjectNameEditableTextBox.cpp
--- a/Engine/Source/Editor/EditorWidgets/Private/SObjectNameEditableTextBox.cpp
+++ b/Engine/Source/Editor/EditorWidgets/Private/SObjectNameEditableTextBox.cpp
@@ -20,6 +20,7 @@
 #include "Styling/WidgetStyle.h"
 #include "Types/WidgetActiveTimerDelegate.h"
 #include "UObject/Object.h"
+#include "UObject/WeakObjectPtr.h"
 #include "Widgets/SCompoundWidget.h"
 #include "Widgets/SWindow.h"
 #include "Widgets/Text/SInlineEditableTextBlock.h"
diff --git a/Engine/Source/Editor/EditorWidgets/Private/SObjectNameEditableTextBox.h b/Engine/Source/Editor/EditorWidgets/Private/SObjectNameEditableTextBox.h
--- a/Engine/Source/Editor/EditorWidgets/Private/SObjectNameEditableTextBox.h
+++ b/Engine/Source/Editor/EditorWidgets/Private/SObjectNameEditableTextBox.h
@@ -11,7 +11,6 @@
 #include "Layout/Visibility.h"
 #include "Templates/SharedPointer.h"
 #include "Types/SlateEnums.h"
-#include "UObject/WeakObjectPtr.h"
 #include "UObject/WeakObjectPtrTemplates.h"
 #include "Widgets/DeclarativeSyntaxSupport.h"
 
diff --git a/Engine/Source/Editor/FontEditor/Private/FontFaceDetailsCustomization.cpp b/Engine/Source/Editor/FontEditor/Private/FontFaceDetailsCustomization.cpp
--- a/Engine/Source/Editor/FontEditor/Private/FontFaceDetailsCustomization.cpp
+++ b/Engine/Source/Editor/FontEditor/Private/FontFaceDetailsCustomization.cpp
@@ -32,6 +32,7 @@
 #include "Templates/UnrealTemplate.h"
 #include "Types/SlateEnums.h"
 #include "UObject/Object.h"
+#include "UObject/WeakObjectPtr.h"
 #include "Widgets/DeclarativeSyntaxSupport.h"
 #include "Widgets/Input/SButton.h"
 #include "Widgets/SBoxPanel.h"
diff --git a/Engine/Source/Editor/FontEditor/Private/FontFaceDetailsCustomization.h b/Engine/Source/Editor/FontEditor/Private/FontFaceDetailsCustomization.h
--- a/Engine/Source/Editor/FontEditor/Private/FontFaceDetailsCustomization.h
+++ b/Engine/Source/Editor/FontEditor/Private/FontFaceDetailsCustomization.h
@@ -8,7 +8,6 @@
 #include "Input/Reply.h"
 #include "Internationalization/Text.h"
 #include "Templates/SharedPointer.h"
-#include "UObject/WeakObjectPtr.h"
 #include "UObject/WeakObjectPtrTemplates.h"
 
 class IDetailLayoutBuilder;
diff --git a/Engine/Source/Editor/GraphEditor/Private/MaterialNodes/SGraphNodeMaterialResult.cpp b/Engine/Source/Editor/GraphEditor/Private/MaterialNodes/SGraphNodeMaterialResult.cpp
--- a/Engine/Source/Editor/GraphEditor/Private/MaterialNodes/SGraphNodeMaterialResult.cpp
+++ b/Engine/Source/Editor/GraphEditor/Private/MaterialNodes/SGraphNodeMaterialResult.cpp
@@ -27,6 +27,7 @@
 #include "TutorialMetaData.h"
 #include "UObject/NameTypes.h"
 #include "UObject/ObjectPtr.h"
+#include "Widgets/SToolTip.h"
 
 class SToolTip;
 
diff --git a/Engine/Source/Editor/GraphEditor/Private/SGraphEditorImpl.cpp b/Engine/Source/Editor/GraphEditor/Private/SGraphEditorImpl.cpp
--- a/Engine/Source/Editor/GraphEditor/Private/SGraphEditorImpl.cpp
+++ b/Engine/Source/Editor/GraphEditor/Private/SGraphEditorImpl.cpp
@@ -19,7 +19,6 @@
 #include "GraphEditorActions.h"
 #include "GraphEditorModule.h"
 #include "GraphEditorSettings.h"
-#include "HAL/PlatformMath.h"
 #include "Input/Events.h"
 #include "InputCoreTypes.h"
 #include "Internationalization/Internationalization.h"
diff --git a/Engine/Source/Editor/HardwareTargeting/Private/HardwareTargetingModule.cpp b/Engine/Source/Editor/HardwareTargeting/Private/HardwareTargetingModule.cpp
--- a/Engine/Source/Editor/HardwareTargeting/Private/HardwareTargetingModule.cpp
+++ b/Engine/Source/Editor/HardwareTargeting/Private/HardwareTargetingModule.cpp
@@ -31,6 +31,7 @@
 #include "UObject/UObjectGlobals.h"
 #include "UObject/UnrealType.h"
 #include "Widgets/DeclarativeSyntaxSupport.h"
+#include "Widgets/SToolTip.h"
 
 #define LOCTEXT_NAMESPACE "HardwareTargeting"
 
diff --git a/Engine/Source/Editor/HierarchicalLODOutliner/Private/HLODOutliner.cpp b/Engine/Source/Editor/HierarchicalLODOutliner/Private/HLODOutliner.cpp
--- a/Engine/Source/Editor/HierarchicalLODOutliner/Private/HLODOutliner.cpp
+++ b/Engine/Source/Editor/HierarchicalLODOutliner/Private/HLODOutliner.cpp
@@ -21,6 +21,7 @@
 #include "Framework/Notifications/NotificationManager.h"
 #include "GameFramework/Actor.h"
 #include "GameFramework/WorldSettings.h"
+#include "HLODOutlinerDragDrop.h"
 #include "HLODTreeWidgetItem.h"
 #include "HierarchicalLOD.h"
 #include "HierarchicalLODUtilitiesModule.h"
diff --git a/Engine/Source/Editor/InputBindingEditor/Private/Widgets/SChordEditor.h b/Engine/Source/Editor/InputBindingEditor/Private/Widgets/SChordEditor.h
--- a/Engine/Source/Editor/InputBindingEditor/Private/Widgets/SChordEditor.h
+++ b/Engine/Source/Editor/InputBindingEditor/Private/Widgets/SChordEditor.h
@@ -7,6 +7,7 @@
 #include "Framework/Commands/UICommandInfo.h"
 #include "Input/Reply.h"
 #include "Internationalization/Text.h"
+#include "SChordEditor.h"
 #include "Templates/SharedPointer.h"
 #include "Widgets/DeclarativeSyntaxSupport.h"
 #include "Widgets/Input/SEditableText.h"
diff --git a/Engine/Source/Editor/Kismet/Private/BPDelegateDragDropAction.cpp b/Engine/Source/Editor/Kismet/Private/BPDelegateDragDropAction.cpp
--- a/Engine/Source/Editor/Kismet/Private/BPDelegateDragDropAction.cpp
+++ b/Engine/Source/Editor/Kismet/Private/BPDelegateDragDropAction.cpp
@@ -2,6 +2,7 @@
 
 #include "BPDelegateDragDropAction.h"
 
+#include "BlueprintEditor.h"
 #include "Delegates/Delegate.h"
 #include "EdGraph/EdGraph.h"
 #include "EdGraph/EdGraphSchema.h"
diff --git a/Engine/Source/Editor/Kismet/Private/BPFunctionDragDropAction.cpp b/Engine/Source/Editor/Kismet/Private/BPFunctionDragDropAction.cpp
--- a/Engine/Source/Editor/Kismet/Private/BPFunctionDragDropAction.cpp
+++ b/Engine/Source/Editor/Kismet/Private/BPFunctionDragDropAction.cpp
@@ -2,6 +2,7 @@
 
 #include "BPFunctionDragDropAction.h"
 
+#include "BlueprintEditor.h"
 #include "BlueprintFunctionNodeSpawner.h"
 #include "BlueprintNodeBinder.h"
 #include "EdGraph/EdGraph.h"
diff --git a/Engine/Source/Editor/Kismet/Private/BPVariableDragDropAction.cpp b/Engine/Source/Editor/Kismet/Private/BPVariableDragDropAction.cpp
--- a/Engine/Source/Editor/Kismet/Private/BPVariableDragDropAction.cpp
+++ b/Engine/Source/Editor/Kismet/Private/BPVariableDragDropAction.cpp
@@ -2,6 +2,7 @@
 
 #include "BPVariableDragDropAction.h"
 
+#include "BlueprintEditor.h"
 #include "Containers/EnumAsByte.h"
 #include "Delegates/Delegate.h"
 #include "EdGraph/EdGraph.h"
diff --git a/Engine/Source/Editor/Kismet/Private/BlueprintActionMenuUtils.cpp b/Engine/Source/Editor/Kismet/Private/BlueprintActionMenuUtils.cpp
--- a/Engine/Source/Editor/Kismet/Private/BlueprintActionMenuUtils.cpp
+++ b/Engine/Source/Editor/Kismet/Private/BlueprintActionMenuUtils.cpp
@@ -7,6 +7,7 @@
 #include "BlueprintActionMenuBuilder.h"
 #include "BlueprintActionMenuItem.h"
 #include "BlueprintDragDropMenuItem.h"
+#include "BlueprintEditor.h"
 #include "BlueprintEditorSettings.h"
 #include "BlueprintNodeBinder.h"
 #include "BlueprintNodeSpawner.h"
diff --git a/Engine/Source/Editor/Kismet/Private/BlueprintDetailsCustomization.cpp b/Engine/Source/Editor/Kismet/Private/BlueprintDetailsCustomization.cpp
--- a/Engine/Source/Editor/Kismet/Private/BlueprintDetailsCustomization.cpp
+++ b/Engine/Source/Editor/Kismet/Private/BlueprintDetailsCustomization.cpp
@@ -2,6 +2,7 @@
 
 #include "BlueprintDetailsCustomization.h"
 
+#include "BlueprintEditor.h"
 #include "BlueprintEditorModule.h"
 #include "BlueprintEditorSettings.h"
 #include "BlueprintNamespaceRegistry.h"
diff --git a/Engine/Source/Editor/Kismet/Private/BlueprintEditorModule.cpp b/Engine/Source/Editor/Kismet/Private/BlueprintEditorModule.cpp
--- a/Engine/Source/Editor/Kismet/Private/BlueprintEditorModule.cpp
+++ b/Engine/Source/Editor/Kismet/Private/BlueprintEditorModule.cpp
@@ -21,7 +21,6 @@
 #include "Editor/Transactor.h"
 #include "EditorUndoClient.h"
 #include "Framework/Commands/UICommandList.h"
-#include "HAL/Platform.h"
 #include "HAL/PlatformCrt.h"
 #include "IMessageLogListing.h"
 #include "ISettingsModule.h"
diff --git a/Engine/Source/Editor/Kismet/Private/BlueprintEditorTabFactories.h b/Engine/Source/Editor/Kismet/Private/BlueprintEditorTabFactories.h
--- a/Engine/Source/Editor/Kismet/Private/BlueprintEditorTabFactories.h
+++ b/Engine/Source/Editor/Kismet/Private/BlueprintEditorTabFactories.h
@@ -2,6 +2,7 @@
 
 #pragma once
 
+#include "BlueprintEditor.h"
 #include "Delegates/Delegate.h"
 #include "EdGraph/EdGraph.h"
 #include "EdGraph/EdGraphSchema.h"
diff --git a/Engine/Source/Editor/Kismet/Private/BlueprintMemberReferenceCustomization.cpp b/Engine/Source/Editor/Kismet/Private/BlueprintMemberReferenceCustomization.cpp
--- a/Engine/Source/Editor/Kismet/Private/BlueprintMemberReferenceCustomization.cpp
+++ b/Engine/Source/Editor/Kismet/Private/BlueprintMemberReferenceCustomization.cpp
@@ -2,6 +2,7 @@
 
 #include "BlueprintMemberReferenceCustomization.h"
 
+#include "BlueprintEditor.h"
 #include "BlueprintEditorModule.h"
 #include "Containers/Array.h"
 #include "Containers/UnrealString.h"
diff --git a/Engine/Source/Editor/Kismet/Private/FindInBlueprints.cpp b/Engine/Source/Editor/Kismet/Private/FindInBlueprints.cpp
--- a/Engine/Source/Editor/Kismet/Private/FindInBlueprints.cpp
+++ b/Engine/Source/Editor/Kismet/Private/FindInBlueprints.cpp
@@ -84,6 +84,7 @@
 #include "Widgets/Layout/SScrollBox.h"
 #include "Widgets/Notifications/SProgressBar.h"
 #include "Widgets/SBoxPanel.h"
+#include "Widgets/SToolTip.h"
 #include "Widgets/Text/SMultiLineEditableText.h"
 #include "Widgets/Text/STextBlock.h"
 #include "Widgets/Views/STableRow.h"
diff --git a/Engine/Source/Editor/Kismet/Private/FormatTextDetails.cpp b/Engine/Source/Editor/Kismet/Private/FormatTextDetails.cpp
--- a/Engine/Source/Editor/Kismet/Private/FormatTextDetails.cpp
+++ b/Engine/Source/Editor/Kismet/Private/FormatTextDetails.cpp
@@ -6,7 +6,6 @@
 #include "DetailCategoryBuilder.h"
 #include "DetailLayoutBuilder.h"
 #include "DetailWidgetRow.h"
-#include "HAL/Platform.h"
 #include "HAL/PlatformCrt.h"
 #include "IDetailChildrenBuilder.h"
 #include "Internationalization/Internationalization.h"
diff --git a/Engine/Source/Editor/Kismet/Private/MyBlueprintItemDragDropAction.cpp b/Engine/Source/Editor/Kismet/Private/MyBlueprintItemDragDropAction.cpp
--- a/Engine/Source/Editor/Kismet/Private/MyBlueprintItemDragDropAction.cpp
+++ b/Engine/Source/Editor/Kismet/Private/MyBlueprintItemDragDropAction.cpp
@@ -2,6 +2,7 @@
 
 #include "MyBlueprintItemDragDropAction.h"
 
+#include "BlueprintEditor.h"
 #include "CoreTypes.h"
 #include "EdGraph/EdGraphSchema.h"
 #include "Internationalization/Internationalization.h"
diff --git a/Engine/Source/Editor/Kismet/Private/SBlueprintActionMenu.cpp b/Engine/Source/Editor/Kismet/Private/SBlueprintActionMenu.cpp
--- a/Engine/Source/Editor/Kismet/Private/SBlueprintActionMenu.cpp
+++ b/Engine/Source/Editor/Kismet/Private/SBlueprintActionMenu.cpp
@@ -61,6 +61,7 @@
 #include "Widgets/SCompoundWidget.h"
 #include "Widgets/SNullWidget.h"
 #include "Widgets/SPanel.h"
+#include "Widgets/SToolTip.h"
 #include "Widgets/Text/STextBlock.h"
 
 class SWidget;
diff --git a/Engine/Source/Editor/Kismet/Private/SBlueprintDiff.cpp b/Engine/Source/Editor/Kismet/Private/SBlueprintDiff.cpp
--- a/Engine/Source/Editor/Kismet/Private/SBlueprintDiff.cpp
+++ b/Engine/Source/Editor/Kismet/Private/SBlueprintDiff.cpp
@@ -1,29 +1,63 @@
 // Copyright Epic Games, Inc. All Rights Reserved.
 
 #include "SBlueprintDiff.h"
-#include "Widgets/Layout/SSplitter.h"
-#include "Widgets/SOverlay.h"
-#include "SlateOptMacros.h"
-#include "Widgets/Layout/SSpacer.h"
-#include "Framework/MultiBox/MultiBoxDefs.h"
+
+#include "BlueprintEditor.h"
+#include "Containers/Set.h"
+#include "DiffControl.h"
+#include "DiffUtils.h"
+#include "EdGraph/EdGraph.h"
+#include "EdGraph/EdGraphNode.h"
+#include "EdGraph/EdGraphPin.h"
+#include "EdGraphUtilities.h"
+#include "Editor.h"
+#include "Editor/EditorEngine.h"
+#include "Engine/Blueprint.h"
+#include "Framework/Application/SlateApplication.h"
+#include "Framework/Commands/GenericCommands.h"
+#include "Framework/Commands/UIAction.h"
+#include "Framework/Commands/UICommandList.h"
 #include "Framework/MultiBox/MultiBoxBuilder.h"
-#include "Styling/AppStyle.h"
-#include "Animation/AnimBlueprint.h"
+#include "Framework/MultiBox/MultiBoxDefs.h"
+#include "GameFramework/Actor.h"
+#include "GraphDiffControl.h"
+#include "HAL/PlatformApplicationMisc.h"
+#include "HAL/PlatformCrt.h"
+#include "IAssetTypeActions.h"
+#include "Internationalization/Text.h"
 #include "K2Node_MathExpression.h"
-#include "Kismet2/BlueprintEditorUtils.h"
 #include "Kismet2/KismetEditorUtilities.h"
-#include "DetailsDiff.h"
-#include "EdGraphUtilities.h"
-#include "GraphDiffControl.h"
+#include "Layout/Children.h"
+#include "Layout/Margin.h"
+#include "Layout/Visibility.h"
+#include "Math/Color.h"
+#include "Math/Vector2D.h"
+#include "Misc/AssertionMacros.h"
+#include "Misc/CString.h"
+#include "PropertyEditorDelegates.h"
+#include "SKismetInspector.h"
 #include "SMyBlueprint.h"
-#include "WorkflowOrientedApp/SModeWidget.h"
-#include "Framework/Commands/GenericCommands.h"
-#include "WidgetBlueprint.h"
-#include "HAL/PlatformApplicationMisc.h"
-#include "Framework/Application/SlateApplication.h"
-#include "DiffControl.h"
-
+#include "SlateOptMacros.h"
+#include "SlotBase.h"
+#include "Styling/AppStyle.h"
+#include "Styling/SlateColor.h"
 #include "Subsystems/AssetEditorSubsystem.h"
+#include "Templates/Casts.h"
+#include "Templates/SubclassOf.h"
+#include "UObject/Class.h"
+#include "UObject/Object.h"
+#include "UObject/UnrealNames.h"
+#include "Widgets/Layout/SBorder.h"
+#include "Widgets/Layout/SBox.h"
+#include "Widgets/Layout/SSpacer.h"
+#include "Widgets/Layout/SSplitter.h"
+#include "Widgets/SBoxPanel.h"
+#include "Widgets/SOverlay.h"
+#include "Widgets/SWindow.h"
+#include "Widgets/Text/STextBlock.h"
+#include "Widgets/Views/SListView.h"
+
+class FProperty;
 
 #define LOCTEXT_NAMESPACE "SBlueprintDif"
 
diff --git a/Engine/Source/Editor/Kismet/Private/SBlueprintEditorSelectedDebugObjectWidget.cpp b/Engine/Source/Editor/Kismet/Private/SBlueprintEditorSelectedDebugObjectWidget.cpp
--- a/Engine/Source/Editor/Kismet/Private/SBlueprintEditorSelectedDebugObjectWidget.cpp
+++ b/Engine/Source/Editor/Kismet/Private/SBlueprintEditorSelectedDebugObjectWidget.cpp
@@ -41,6 +41,7 @@
 #include "UnrealEdGlobals.h"
 #include "Widgets/Input/STextComboBox.h"
 #include "Widgets/SBoxPanel.h"
+#include "Widgets/SToolTip.h"
 #include "Widgets/SWidget.h"
 #include "Widgets/Text/STextBlock.h"
 
diff --git a/Engine/Source/Editor/Kismet/Private/SCSDiff.cpp b/Engine/Source/Editor/Kismet/Private/SCSDiff.cpp
--- a/Engine/Source/Editor/Kismet/Private/SCSDiff.cpp
+++ b/Engine/Source/Editor/Kismet/Private/SCSDiff.cpp
@@ -22,6 +22,7 @@
 #include "SlotBase.h"
 #include "SubobjectData.h"
 #include "SubobjectDataHandle.h"
+#include "Templates/Casts.h"
 #include "Templates/ChooseClass.h"
 #include "Templates/SubclassOf.h"
 #include "Types/SlateEnums.h"
diff --git a/Engine/Source/Editor/Kismet/Private/SCSEditorViewportClient.cpp b/Engine/Source/Editor/Kismet/Private/SCSEditorViewportClient.cpp
--- a/Engine/Source/Editor/Kismet/Private/SCSEditorViewportClient.cpp
+++ b/Engine/Source/Editor/Kismet/Private/SCSEditorViewportClient.cpp
@@ -34,7 +34,6 @@
 #include "GameFramework/Actor.h"
 #include "HitProxies.h"
 #include "ISCSEditorCustomization.h"
-#include "Internationalization/Internationalization.h"
 #include "Internationalization/Text.h"
 #include "Kismet2/ComponentEditorUtils.h"
 #include "Logging/LogMacros.h"
diff --git a/Engine/Source/Editor/Kismet/Private/SMyBlueprint.cpp b/Engine/Source/Editor/Kismet/Private/SMyBlueprint.cpp
--- a/Engine/Source/Editor/Kismet/Private/SMyBlueprint.cpp
+++ b/Engine/Source/Editor/Kismet/Private/SMyBlueprint.cpp
@@ -47,7 +47,6 @@
 #include "GraphEditorDragDropAction.h"
 #include "HAL/PlatformApplicationMisc.h"
 #include "HAL/PlatformCrt.h"
-#include "HAL/PlatformMath.h"
 #include "Input/DragAndDrop.h"
 #include "Input/Events.h"
 #include "InputCoreTypes.h"
diff --git a/Engine/Source/Editor/Kismet/Private/SSCSEditor.cpp b/Engine/Source/Editor/Kismet/Private/SSCSEditor.cpp
--- a/Engine/Source/Editor/Kismet/Private/SSCSEditor.cpp
+++ b/Engine/Source/Editor/Kismet/Private/SSCSEditor.cpp
@@ -9,6 +9,7 @@
 #include "AssetRegistry/AssetData.h"
 #include "AssetSelection.h"
 #include "BPVariableDragDropAction.h"
+#include "BlueprintEditor.h"
 #include "BlueprintEditorModule.h"
 #include "BlueprintEditorSettings.h"
 #include "ClassIconFinder.h"
diff --git a/Engine/Source/Editor/Kismet/Private/SSCSEditorViewport.cpp b/Engine/Source/Editor/Kismet/Private/SSCSEditorViewport.cpp
--- a/Engine/Source/Editor/Kismet/Private/SSCSEditorViewport.cpp
+++ b/Engine/Source/Editor/Kismet/Private/SSCSEditorViewport.cpp
@@ -22,7 +22,6 @@
 #include "Framework/MultiBox/MultiBoxBuilder.h"
 #include "GenericPlatform/ICursor.h"
 #include "HAL/PlatformCrt.h"
-#include "Internationalization/Internationalization.h"
 #include "Internationalization/Text.h"
 #include "Layout/Children.h"
 #include "Layout/Margin.h"
diff --git a/Engine/Source/Editor/Kismet/Public/BlueprintEditor.h b/Engine/Source/Editor/Kismet/Public/BlueprintEditor.h
--- a/Engine/Source/Editor/Kismet/Public/BlueprintEditor.h
+++ b/Engine/Source/Editor/Kismet/Public/BlueprintEditor.h
@@ -31,6 +31,7 @@
 #include "Misc/Optional.h"
 #include "PreviewScene.h"
 #include "SKismetInspector.h"
+#include "SourceControlOperations.h"
 #include "Stats/Stats.h"
 #include "Stats/Stats2.h"
 #include "Styling/SlateBrush.h"
diff --git a/Engine/Source/Editor/KismetCompiler/Public/KismetCompiler.h b/Engine/Source/Editor/KismetCompiler/Public/KismetCompiler.h
--- a/Engine/Source/Editor/KismetCompiler/Public/KismetCompiler.h
+++ b/Engine/Source/Editor/KismetCompiler/Public/KismetCompiler.h
@@ -28,6 +28,7 @@
 #include "Templates/Function.h"
 #include "Templates/SharedPointer.h"
 #include "Templates/SubclassOf.h"
+#include "Templates/TypeHash.h"
 #include "UObject/NameTypes.h"
 #include "UObject/Object.h"
 #include "UObject/ObjectMacros.h"
diff --git a/Engine/Source/Editor/KismetWidgets/Private/CreateBlueprintFromActorDialog.cpp b/Engine/Source/Editor/KismetWidgets/Private/CreateBlueprintFromActorDialog.cpp
--- a/Engine/Source/Editor/KismetWidgets/Private/CreateBlueprintFromActorDialog.cpp
+++ b/Engine/Source/Editor/KismetWidgets/Private/CreateBlueprintFromActorDialog.cpp
@@ -22,7 +22,6 @@
 #include "Fonts/SlateFontInfo.h"
 #include "Framework/Notifications/NotificationManager.h"
 #include "GameFramework/Actor.h"
-#include "HAL/Platform.h"
 #include "HAL/PlatformCrt.h"
 #include "HAL/PlatformMisc.h"
 #include "IAssetTools.h"
