Fixed non-unity mac compile issue.

diff --git a/Engine/Source/Editor/Kismet/Private/SMyBlueprint.cpp b/Engine/Source/Editor/Kismet/Private/SMyBlueprint.cpp
--- a/Engine/Source/Editor/Kismet/Private/SMyBlueprint.cpp
+++ b/Engine/Source/Editor/Kismet/Private/SMyBlueprint.cpp
@@ -5,6 +5,7 @@
 
 #include "Animation/AnimClassInterface.h"
 #include "AnimationGraph.h"
+#include "AnimationStateMachineGraph.h"
 #include "AnimGraphNode_StateMachineBase.h"
 #include "BPDelegateDragDropAction.h"
 #include "BPFunctionDragDropAction.h"
