Block dragging functions into the Anim BP from My Blueprnit, as these are exec call nodes, not usable on the Anim Graph

diff --git a/Engine/Source/Editor/Kismet/Private/SMyBlueprint.cpp b/Engine/Source/Editor/Kismet/Private/SMyBlueprint.cpp
--- a/Engine/Source/Editor/Kismet/Private/SMyBlueprint.cpp
+++ b/Engine/Source/Editor/Kismet/Private/SMyBlueprint.cpp
@@ -1919,13 +1919,24 @@ FReply SMyBlueprint::OnActionDragged( const TArray< TSharedPtr<FEdGraphSchemaAct
 			if (FuncAction->GraphType == EEdGraphSchemaAction_K2Graph::Function ||FuncAction->GraphType == EEdGraphSchemaAction_K2Graph::Interface)
 			{
 				// Callback function to report that the user cannot drop this function in the graph
-				auto CanDragDropAction = [](TSharedPtr<FEdGraphSchemaAction> /*DropAction*/, UEdGraph* /*HoveredGraphIn*/, FText& ImpededReasonOut, bool bIsBlueprintCallableFunction)->bool
+				auto CanDragDropAction = [](TSharedPtr<FEdGraphSchemaAction> /*DropAction*/, UEdGraph* HoveredGraphIn, FText& ImpededReasonOut, bool bIsBlueprintCallableFunction)->bool
 				{
 					if (!bIsBlueprintCallableFunction)
 					{
 						ImpededReasonOut = LOCTEXT("NonBlueprintCallable", "This function was not marked as Blueprint Callable and cannot be placed in a graph!");
+						return false;
+					}
+					
+					if (const UEdGraphSchema_K2* K2Schema = Cast<UEdGraphSchema_K2>(HoveredGraphIn->GetSchema()))
+					{
+						if (K2Schema->DoesGraphSupportImpureFunctions(HoveredGraphIn) == false)
+						{
+							ImpededReasonOut = LOCTEXT("GraphNoSupportImpureF", "The target graph does not support Impure Functions.");
+							return false;
+						}
 					}
-					return bIsBlueprintCallableFunction;
+
+					return true;
 				};
 
 				bool bIsBlueprintCallableFunction = false;
