Enable external linkage for SFindInBlueprints (and related) editor APIs.

diff --git a/Engine/Source/Editor/Kismet/Public/FindInBlueprintManager.h b/Engine/Source/Editor/Kismet/Public/FindInBlueprintManager.h
--- a/Engine/Source/Editor/Kismet/Public/FindInBlueprintManager.h
+++ b/Engine/Source/Editor/Kismet/Public/FindInBlueprintManager.h
@@ -303,7 +303,7 @@ enum class EFiBCacheOpFlags
 ENUM_CLASS_FLAGS(EFiBCacheOpFlags);
 
 /** Options to configure the bulk caching task */
-struct FFindInBlueprintCachingOptions
+struct KISMET_API FFindInBlueprintCachingOptions
 {
 	/** Type of caching operation */
 	EFiBCacheOpType OpType = EFiBCacheOpType::CachePendingAssets;
@@ -396,7 +396,7 @@ typedef TSharedPtr<FFindInBlueprintsResult> FSearchResult;
 ////////////////////////////////////
 // FStreamSearch
 
-struct FStreamSearchOptions
+struct KISMET_API FStreamSearchOptions
 {
 	/** Filter to limit the FilteredImaginaryResults to */
 	enum ESearchQueryFilter ImaginaryDataFilter;
diff --git a/Engine/Source/Editor/Kismet/Public/FindInBlueprints.h b/Engine/Source/Editor/Kismet/Public/FindInBlueprints.h
--- a/Engine/Source/Editor/Kismet/Public/FindInBlueprints.h
+++ b/Engine/Source/Editor/Kismet/Public/FindInBlueprints.h
@@ -256,7 +256,7 @@ enum class EFiBSearchBarWidget
 };
 
 /*Widget for searching for (functions/events) across all blueprints or just a single blueprint */
-class SFindInBlueprints: public SCompoundWidget
+class KISMET_API SFindInBlueprints: public SCompoundWidget
 {
 public:
 	SLATE_BEGIN_ARGS( SFindInBlueprints )
