Centered the "Graph is not editable." message in the blueprint editor details view.

#rnx

diff --git a/Engine/Source/Editor/Kismet/Private/BlueprintDetailsCustomization.cpp b/Engine/Source/Editor/Kismet/Private/BlueprintDetailsCustomization.cpp
--- a/Engine/Source/Editor/Kismet/Private/BlueprintDetailsCustomization.cpp
+++ b/Engine/Source/Editor/Kismet/Private/BlueprintDetailsCustomization.cpp
@@ -4158,13 +4158,15 @@ void FBlueprintGraphActionDetails::CustomizeDetails( IDetailLayoutBuilder& Detai
 		{
 			IDetailCategoryBuilder& Category = DetailLayout.EditCategory("Graph", LOCTEXT("FunctionDetailsGraph", "Graph"));
 			Category.AddCustomRow( FText::GetEmpty() )
+			.WholeRowContent()
+			.VAlign(VAlign_Center)
 			[
 				SNew(STextBlock)
 				.Text( LOCTEXT("GraphPresentButNotEditable", "Graph is not editable.") )
 			];
 		}
 	}
-
+	
 	if (MyBlueprint.IsValid())
 	{
 		TWeakPtr<FBlueprintEditor> BlueprintEditor = MyBlueprint.Pin()->GetBlueprintEditor();
