[BugFix] The class settings in the blueprint diff tool are missing the implemented interfaces
Because the interfaces settings in blueprints use IDetailCustomization to create custom widgets in the details panel, This bugfix required implementing the neccesary infrustructure to attribute properties to the custom widgets so they can be diffed and highlighted.

diff --git a/Engine/Source/Editor/Kismet/Private/BlueprintDetailsCustomization.cpp b/Engine/Source/Editor/Kismet/Private/BlueprintDetailsCustomization.cpp
--- a/Engine/Source/Editor/Kismet/Private/BlueprintDetailsCustomization.cpp
+++ b/Engine/Source/Editor/Kismet/Private/BlueprintDetailsCustomization.cpp
@@ -5795,20 +5795,25 @@ void FBlueprintImportsLayout::OnGetNamespacesToExclude(TSet<FString>& OutNamespa
 }
 
 
-FBlueprintInterfaceLayout::FBlueprintInterfaceLayout(TWeakPtr<class FBlueprintGlobalOptionsDetails> InGlobalOptionsDetails, bool bInShowsInheritedInterfaces)
+FBlueprintInterfaceLayout::FBlueprintInterfaceLayout(TWeakPtr<class FBlueprintGlobalOptionsDetails> InGlobalOptionsDetails, TSharedPtr<IPropertyHandle> InInterfacesProperty)
 	: FBlueprintGlobalOptionsManagedListDetails(InGlobalOptionsDetails)
-	, bShowsInheritedInterfaces(bInShowsInheritedInterfaces)
+	, InterfacesProperty(InInterfacesProperty)
 {
-	DisplayOptions.TitleText = bShowsInheritedInterfaces ?
+	DisplayOptions.TitleText = InterfacesProperty ?
 		LOCTEXT("BlueprintInheritedInterfaceTitle", "Inherited Interfaces") :
 		LOCTEXT("BlueprintImplementedInterfaceTitle", "Implemented Interfaces");
 	DisplayOptions.NoItemsLabelText = LOCTEXT("NoBlueprintInterface", "No Interfaces");
 	DisplayOptions.BrowseButtonToolTipText = LOCTEXT("BlueprintInterfaceBrowseTooltip", "Opens this interface");
 }
 
+TSharedPtr<IPropertyHandle> FBlueprintInterfaceLayout::GetPropertyHandle() const
+{
+	return InterfacesProperty;
+}
+
 TSharedPtr<SWidget> FBlueprintInterfaceLayout::MakeAddItemWidget()
 {
-	if (!bShowsInheritedInterfaces)
+	if (InterfacesProperty)
 	{
 		return SAssignNew(AddInterfaceComboButton, SComboButton)
 		.ButtonContent()
@@ -5826,17 +5831,27 @@ void FBlueprintInterfaceLayout::GetManagedListItems(TArray<FManagedListItem>& Ou
 {
 	const UBlueprint* Blueprint = GetBlueprintObjectChecked();
 
-	if (!bShowsInheritedInterfaces)
+	if (InterfacesProperty)
 	{
+		checkf(
+			&Blueprint->ImplementedInterfaces == (TArray<FBPInterfaceDescription>*)InterfacesProperty->GetValueBaseAddress((uint8*)Blueprint),
+			TEXT("Different Property provided than ImplementedInterfaces")
+		);
+		
 		// Generate a list of interfaces already implemented
-		for (const FBPInterfaceDescription& ImplementedInterface : Blueprint->ImplementedInterfaces)
+		for (int32 InterfaceIndex = 0; InterfaceIndex < Blueprint->ImplementedInterfaces.Num(); ++InterfaceIndex)
 		{
-			if (const TSubclassOf<UInterface> Interface = ImplementedInterface.Interface)
+			if (const TSubclassOf<UInterface> Interface = Blueprint->ImplementedInterfaces[InterfaceIndex].Interface)
 			{
 				FManagedListItem ItemDesc;
 				ItemDesc.ItemName = Interface->GetPathName();
 				ItemDesc.DisplayName = Interface->GetDisplayNameText();
 				ItemDesc.bIsRemovable = true;
+				const TSharedPtr<IPropertyHandle> Property = GetPropertyHandle()->GetChildHandle(InterfaceIndex);
+				if (Property && Property->IsValidHandle())
+				{
+					ItemDesc.PropertyHandles.Add(Property);
+				}
 
 				// Allow browsing to Blueprint interface class assets.
 				if (UBlueprintGeneratedClass* Class = Cast<UBlueprintGeneratedClass>(*Interface))
@@ -5950,6 +5965,11 @@ float FBlueprintGlobalOptionsDetails::NamespacePropertyValueCustomization_MinDes
 
 UBlueprint* FBlueprintGlobalOptionsDetails::GetBlueprintObj() const
 {
+	if(BlueprintObjOverride)
+	{
+		return BlueprintObjOverride;
+	}
+	
 	if(BlueprintEditorPtr.IsValid())
 	{
 		return BlueprintEditorPtr.Pin()->GetBlueprintObj();
@@ -6183,6 +6203,10 @@ void FBlueprintGlobalOptionsDetails::CustomizeDetails(IDetailLayoutBuilder& Deta
 
 		// Display the parent class and set up the menu for reparenting
 		IDetailCategoryBuilder& Category = DetailLayout.EditCategory("ClassOptions", LOCTEXT("ClassOptions", "Class Options"));
+
+		// ParentClass is a hidden property so we have to add it to the property map manually to use it
+		const TSharedPtr<IPropertyHandle> ParentClassProperty = DetailLayout.AddObjectPropertyData({const_cast<UBlueprint*>(Blueprint)}, TEXT("ParentClass"));
+
 		Category.AddCustomRow( LOCTEXT("ClassOptions", "Class Options") )
 		.NameContent()
 		[
@@ -6206,7 +6230,8 @@ void FBlueprintGlobalOptionsDetails::CustomizeDetails(IDetailLayoutBuilder& Deta
 					.Font(IDetailLayoutBuilder::GetDetailFont())
 				]
 			]
-		];
+		]
+		.PropertyHandleList({ParentClassProperty});
 		
 		const bool bIsInterfaceBP = FBlueprintEditorUtils::IsInterfaceBlueprint(Blueprint);
 		const bool bIsMacroLibrary = Blueprint->BlueprintType == BPTYPE_MacroLibrary;
@@ -6231,11 +6256,19 @@ void FBlueprintGlobalOptionsDetails::CustomizeDetails(IDetailLayoutBuilder& Deta
 		{
 			// Interface details customization
 			IDetailCategoryBuilder& InterfacesCategory = DetailLayout.EditCategory("Interfaces", LOCTEXT("BlueprintInterfacesDetailsCategory", "Interfaces"));
-		
-			TSharedRef<FBlueprintInterfaceLayout> InheritedInterfacesLayout = MakeShareable(new FBlueprintInterfaceLayout(SharedThis(this), /*bShowInheritedInterfaces = */true));
+
+			// ImplementedInterfaces is a hidden property so we have to add it to the property map manually to use it
+			const TSharedPtr<IPropertyHandle> InterfacesProperty = DetailLayout.AddObjectPropertyData({const_cast<UBlueprint*>(Blueprint)}, TEXT("ImplementedInterfaces"));
+			
+			TSharedRef<FBlueprintInterfaceLayout> InheritedInterfacesLayout = MakeShareable(new FBlueprintInterfaceLayout(
+				SharedThis(this)
+				));
 			InterfacesCategory.AddCustomBuilder(InheritedInterfacesLayout);
 
-			TSharedRef<FBlueprintInterfaceLayout> LocalInterfacesLayout = MakeShareable(new FBlueprintInterfaceLayout(SharedThis(this), /*bShowInheritedInterfaces = */false));
+			TSharedRef<FBlueprintInterfaceLayout> LocalInterfacesLayout = MakeShareable(new FBlueprintInterfaceLayout(
+				SharedThis(this),
+				InterfacesProperty.ToSharedRef()
+				));
 			InterfacesCategory.AddCustomBuilder(LocalInterfacesLayout);
 		}
 
@@ -6275,7 +6308,8 @@ void FBlueprintGlobalOptionsDetails::CustomizeDetails(IDetailLayoutBuilder& Deta
 					.IsChecked( this, &FBlueprintGlobalOptionsDetails::IsDeprecatedBlueprint )
 					.OnCheckStateChanged( this, &FBlueprintGlobalOptionsDetails::OnDeprecateBlueprint )
 					.ToolTipText( this, &FBlueprintGlobalOptionsDetails::GetDeprecatedTooltip )
-				];
+				]
+				.PropertyHandleList({DetailLayout.GetProperty(DeprecatePropName)});
 		}
 
 		static FName BlueprintNamespacePropertyName = GET_MEMBER_NAME_CHECKED(UBlueprint, BlueprintNamespace);
diff --git a/Engine/Source/Editor/Kismet/Private/BlueprintDetailsCustomization.h b/Engine/Source/Editor/Kismet/Private/BlueprintDetailsCustomization.h
--- a/Engine/Source/Editor/Kismet/Private/BlueprintDetailsCustomization.h
+++ b/Engine/Source/Editor/Kismet/Private/BlueprintDetailsCustomization.h
@@ -762,8 +762,11 @@ private:
 class FBlueprintInterfaceLayout : public FBlueprintGlobalOptionsManagedListDetails, public TSharedFromThis<FBlueprintInterfaceLayout>
 {
 public:
-	FBlueprintInterfaceLayout(TWeakPtr<class FBlueprintGlobalOptionsDetails> InGlobalOptionsDetails, bool bInShowsInheritedInterfaces);
+	FBlueprintInterfaceLayout(TWeakPtr<class FBlueprintGlobalOptionsDetails> InGlobalOptionsDetails, TSharedPtr<IPropertyHandle> InInterfacesProperty = nullptr);
 
+	/** IDetailCustomNodeBuilder interface */
+	virtual TSharedPtr<IPropertyHandle> GetPropertyHandle() const override;
+	/** END IDetailCustomNodeBuilder interface */
 protected:
 	/** FBlueprintManagedListDetails interface */
 	virtual TSharedPtr<SWidget> MakeAddItemWidget() override;
@@ -782,8 +785,8 @@ private:
 	/** The add interface combo button */
 	TSharedPtr<SComboButton> AddInterfaceComboButton;
 
-	/** Whether we show inherited interfaces versus implemented interfaces */
-	bool bShowsInheritedInterfaces;
+	/** Property that stores interfaces */
+	TSharedPtr<IPropertyHandle> InterfacesProperty;
 };
 
 /** Details customization for Blueprint settings */
@@ -793,6 +796,13 @@ public:
 	/** Constructor */
 	FBlueprintGlobalOptionsDetails(TWeakPtr<FBlueprintEditor> InBlueprintEditorPtr)
 		:BlueprintEditorPtr(InBlueprintEditorPtr)
+		,BlueprintObjOverride(nullptr)
+	{
+
+	}
+	
+	FBlueprintGlobalOptionsDetails(UBlueprint* InBlueprintPtr)
+		:BlueprintObjOverride(InBlueprintPtr)
 	{
 
 	}
@@ -803,6 +813,13 @@ public:
 		return MakeShareable(new FBlueprintGlobalOptionsDetails(InBlueprintEditorPtr));
 	}
 
+	/** Diff functionality doesn't need access to the FBlueprintEditor so this allows creation without editor access
+	 *  but with limited functionality */
+	static TSharedRef<IDetailCustomization> MakeInstanceForDiff(UBlueprint* InBlueprintPtr)
+	{
+		return MakeShareable(new FBlueprintGlobalOptionsDetails(InBlueprintPtr));
+	}
+
 	/** IDetailCustomization interface */
 	virtual void CustomizeDetails( IDetailLayoutBuilder& DetailLayout ) override;
 	
@@ -852,6 +869,9 @@ protected:
 private:
 	/** Weak reference to the Blueprint editor */
 	TWeakPtr<FBlueprintEditor> BlueprintEditorPtr;
+	
+	/** Weak reference to the Blueprint editor */
+	TObjectPtr<UBlueprint> BlueprintObjOverride;
 
 	/** Combo button used to choose a parent class */
 	TSharedPtr<SComboButton> ParentClassComboButton;
diff --git a/Engine/Source/Editor/Kismet/Private/BlueprintManagedListDetails.cpp b/Engine/Source/Editor/Kismet/Private/BlueprintManagedListDetails.cpp
--- a/Engine/Source/Editor/Kismet/Private/BlueprintManagedListDetails.cpp
+++ b/Engine/Source/Editor/Kismet/Private/BlueprintManagedListDetails.cpp
@@ -91,7 +91,8 @@ void FBlueprintManagedListDetails::GenerateChildContent(IDetailChildrenBuilder&
 			.EditCondition(
 				DisplayOptions.EditCondition,
 				FOnBooleanValueChanged()
-			);
+			)
+			.PropertyHandleList(Item.PropertyHandles);
 
 			if (Item.AssetPtr.IsValid())
 			{
diff --git a/Engine/Source/Editor/Kismet/Private/DetailsDiff.cpp b/Engine/Source/Editor/Kismet/Private/DetailsDiff.cpp
--- a/Engine/Source/Editor/Kismet/Private/DetailsDiff.cpp
+++ b/Engine/Source/Editor/Kismet/Private/DetailsDiff.cpp
@@ -2,6 +2,7 @@
 
 #include "DetailsDiff.h"
 
+#include "BlueprintDetailsCustomization.h"
 #include "Algo/Copy.h"
 #include "Algo/Transform.h"
 #include "Containers/Set.h"
@@ -14,6 +15,7 @@
 #include "PropertyEditorModule.h"
 #include "PropertyPath.h"
 #include "UObject/WeakObjectPtrTemplates.h"
+#include "Engine/MemberReference.h"
 
 class UObject;
 
@@ -29,6 +31,15 @@ FDetailsDiff::FDetailsDiff(const UObject* InObject, FOnDisplayedPropertiesChange
 	FPropertyEditorModule& EditModule = FModuleManager::Get().GetModuleChecked<FPropertyEditorModule>("PropertyEditor");
 	DetailsView = EditModule.CreateDetailView(DetailsViewArgs);
 	DetailsView->SetIsPropertyEditingEnabledDelegate(FIsPropertyEditingEnabled::CreateStatic([]{return false; }));
+	if (InObject->IsA<UBlueprint>())
+	{
+		// create a custom property layout so that sections like interfaces will be included in the diff view
+		const FOnGetDetailCustomizationInstance LayoutOptionDetails = FOnGetDetailCustomizationInstance::CreateStatic(
+			&FBlueprintGlobalOptionsDetails::MakeInstanceForDiff,
+			const_cast<UBlueprint *>(Cast<UBlueprint>(InObject))
+		);
+		DetailsView->RegisterInstancedCustomPropertyLayout(UBlueprint::StaticClass(), LayoutOptionDetails);
+	}
 	// Forcing all advanced properties to be displayed for now, the logic to show changes made to advance properties
 	// conditionally is fragile and low priority for now:
 	DetailsView->ShowAllAdvancedProperties();
diff --git a/Engine/Source/Editor/Kismet/Private/DiffControl.cpp b/Engine/Source/Editor/Kismet/Private/DiffControl.cpp
--- a/Engine/Source/Editor/Kismet/Private/DiffControl.cpp
+++ b/Engine/Source/Editor/Kismet/Private/DiffControl.cpp
@@ -226,56 +226,6 @@ void FClassSettingsDiffControl::GenerateTreeEntries(TArray<TSharedPtr<FBlueprint
 {
 	TDetailsDiffControl::GenerateTreeEntries(OutTreeEntries, OutRealDifferences);
 
-	// Check for parent class and interfaces here
-	const UBlueprint* OldBlueprint = Cast<UBlueprint>(OldDetails.GetDisplayedObject());
-	const UBlueprint* NewBlueprint = Cast<UBlueprint>(NewDetails.GetDisplayedObject());
-
-	if (OldBlueprint && NewBlueprint)
-	{
-		if (OldBlueprint->ParentClass != NewBlueprint->ParentClass)
-		{
-			FText DiffText = FText::Format(LOCTEXT("ParentChanged", "Parent Class changed from {0} to {1}"), FText::FromString(GetNameSafe(OldBlueprint->ParentClass)), FText::FromString(GetNameSafe(NewBlueprint->ParentClass)));
-
-			TSharedPtr<FBlueprintDifferenceTreeEntry> Entry = MakeShared<FBlueprintDifferenceTreeEntry>(
-				SelectionCallback,
-				FGenerateDiffEntryWidget::CreateStatic(&GenerateSimpleDiffWidget, DiffText));
-
-			Children.Push(Entry);
-			OutRealDifferences.Push(Entry);
-		}
-
-		FString OldInterfaces, NewInterfaces;
-		for (const FBPInterfaceDescription& Desc : OldBlueprint->ImplementedInterfaces)
-		{
-			if (!OldInterfaces.IsEmpty())
-			{
-				OldInterfaces += TEXT(", ");
-			}
-			OldInterfaces += GetNameSafe(*Desc.Interface);
-		}
-
-		for (const FBPInterfaceDescription& Desc : NewBlueprint->ImplementedInterfaces)
-		{
-			if (!NewInterfaces.IsEmpty())
-			{
-				NewInterfaces += TEXT(", ");
-			}
-			NewInterfaces += GetNameSafe(*Desc.Interface);
-		}
-			
-		if (OldInterfaces != NewInterfaces)
-		{
-			FText DiffText = FText::Format(LOCTEXT("InterfacesChanged", "Interfaces changed from '{0}' to '{1}'"), FText::FromString(OldInterfaces), FText::FromString(NewInterfaces));
-
-			TSharedPtr<FBlueprintDifferenceTreeEntry> Entry = MakeShared<FBlueprintDifferenceTreeEntry>(
-				SelectionCallback,
-				FGenerateDiffEntryWidget::CreateStatic(&GenerateSimpleDiffWidget, DiffText));
-
-			Children.Push(Entry);
-			OutRealDifferences.Push(Entry);
-		}
-	}
-
 	const bool bHasDifferences = Children.Num() != 0;
 	if (!bHasDifferences)
 	{
diff --git a/Engine/Source/Editor/Kismet/Private/DiffUtils.cpp b/Engine/Source/Editor/Kismet/Private/DiffUtils.cpp
--- a/Engine/Source/Editor/Kismet/Private/DiffUtils.cpp
+++ b/Engine/Source/Editor/Kismet/Private/DiffUtils.cpp
@@ -81,6 +81,22 @@ FResolvedProperty FPropertySoftPath::Resolve(const UObject* Object) const
 	return ResolvedProperty;
 }
 
+
+
+int32 FPropertySoftPath::TryReadIndex(const TArray<FChainElement>& LocalPropertyChain, int32& OutIndex)
+{
+	if(OutIndex + 1 < LocalPropertyChain.Num())
+	{
+		FString AsString = LocalPropertyChain[OutIndex + 1].DisplayString;
+		if(AsString.IsNumeric())
+		{
+			++OutIndex;
+			return FCString::Atoi(*AsString);
+		}
+	}
+	return INDEX_NONE;
+};
+
 FResolvedProperty FPropertySoftPath::Resolve(const UStruct* Struct, const void* StructData) const
 {
 	// dig into the object, finding nested objects, etc:
@@ -93,6 +109,34 @@ FResolvedProperty FPropertySoftPath::Resolve(const UStruct* Struct, const void*
 	{
 		CurrentBlock = NextBlock;
 		const FProperty* NextProperty = UEDiffUtils_Private::Resolve(NextClass, PropertyChain[i].PropertyName);
+
+		// if an index was provided, resolve it
+		const int32 PropertyIndex = TryReadIndex(PropertyChain, i);
+		if (NextProperty && PropertyIndex != INDEX_NONE)
+		{
+			Property = NextProperty;
+			if (const FArrayProperty* ArrayProperty = CastField<FArrayProperty>(Property))
+			{
+				FScriptArrayHelper ArrayHelper(ArrayProperty, Property->ContainerPtrToValuePtr<UObject*>(CurrentBlock));
+				if (ArrayHelper.IsValidIndex(PropertyIndex))
+				{
+					NextProperty = ArrayProperty->Inner;
+					NextBlock = ArrayHelper.GetRawPtr(PropertyIndex);
+				}
+			}
+			else if( const FSetProperty* SetProperty = CastField<FSetProperty>(Property) )
+			{
+				checkf(false, TEXT("Set Indexing not supported yet"));
+				// TODO: @jordan.hoffmann: handle indexing for sets
+			}
+			else if( const FMapProperty* MapProperty = CastField<FMapProperty>(Property) )
+			{
+				checkf(false, TEXT("Map Indexing not supported yet"));
+				// TODO: @jordan.hoffmann: handle indexing for maps
+			}
+		}
+		
+		CurrentBlock = NextBlock;
 		if (NextProperty)
 		{
 			Property = NextProperty;
@@ -144,20 +188,6 @@ FPropertyPath FPropertySoftPath::ResolvePath(const UObject* Object) const
 		}
 	};
 
-	auto TryReadIndex = [](const TArray<FChainElement>& LocalPropertyChain, int32& OutIndex) -> int32
-	{
-		if(OutIndex + 1 < LocalPropertyChain.Num())
-		{
-			FString AsString = LocalPropertyChain[OutIndex + 1].DisplayString;
-			if(AsString.IsNumeric())
-			{
-				++OutIndex;
-				return FCString::Atoi(*AsString);
-			}
-		}
-		return INDEX_NONE;
-	};
-
 	const void* ContainerAddress = Object;
 	const UStruct* ContainerStruct = (Object ? Object->GetClass() : nullptr);
 
@@ -176,10 +206,9 @@ FPropertyPath FPropertySoftPath::ResolvePath(const UObject* Object) const
 		// calculate offset so we can continue resolving object properties/structproperties:
 		if( const FArrayProperty* ArrayProperty = CastField<FArrayProperty>(ResolvedProperty) )
 		{
-			if(PropertyIndex != INDEX_NONE)
+			FScriptArrayHelper ArrayHelper(ArrayProperty, ArrayProperty->ContainerPtrToValuePtr<const void*>( ContainerAddress ));
+			if (ArrayHelper.IsValidIndex(PropertyIndex))
 			{
-				FScriptArrayHelper ArrayHelper(ArrayProperty, ArrayProperty->ContainerPtrToValuePtr<const void*>( ContainerAddress ));
-
 				UpdateContainerAddress( ArrayProperty->Inner, ArrayHelper.GetRawPtr(PropertyIndex), ContainerAddress, ContainerStruct );
 
 				FPropertyInfo ArrayInfo(ArrayProperty->Inner, PropertyIndex);
diff --git a/Engine/Source/Editor/Kismet/Public/BlueprintManagedListDetails.h b/Engine/Source/Editor/Kismet/Public/BlueprintManagedListDetails.h
--- a/Engine/Source/Editor/Kismet/Public/BlueprintManagedListDetails.h
+++ b/Engine/Source/Editor/Kismet/Public/BlueprintManagedListDetails.h
@@ -31,6 +31,7 @@ protected:
 		FText DisplayName;
 		TWeakObjectPtr<> AssetPtr;
 		uint8 bIsRemovable : 1;
+		TArray<TSharedPtr<IPropertyHandle>> PropertyHandles;
 
 		FManagedListItem()
 			:bIsRemovable(false)
diff --git a/Engine/Source/Editor/Kismet/Public/DiffUtils.h b/Engine/Source/Editor/Kismet/Public/DiffUtils.h
--- a/Engine/Source/Editor/Kismet/Public/DiffUtils.h
+++ b/Engine/Source/Editor/Kismet/Public/DiffUtils.h
@@ -75,9 +75,17 @@ struct FPropertySoftPath
 
 	FPropertySoftPath( FPropertyPath InPropertyPath )
 	{
-		for (int32 i = 0, end = InPropertyPath.GetNumProperties(); i != end; ++i)
+		for (int32 PropertyIndex = 0, end = InPropertyPath.GetNumProperties(); PropertyIndex != end; ++PropertyIndex)
 		{
-			PropertyChain.Push(FChainElement(InPropertyPath.GetPropertyInfo(i).Property.Get()));
+			const FPropertyInfo& Info = InPropertyPath.GetPropertyInfo(PropertyIndex);
+			if (Info.ArrayIndex != INDEX_NONE)
+			{
+				PropertyChain.Push(FName(*FString::FromInt(Info.ArrayIndex)));
+			}
+			else
+			{
+				PropertyChain.Push(FChainElement(Info.Property.Get()));
+			}
 		}
 	}
 
@@ -163,6 +171,7 @@ private:
 			return !(*this == RHS);
 		}
 	};
+	static int32 TryReadIndex(const TArray<FChainElement>& LocalPropertyChain, int32& OutIndex);
 
 	friend uint32 GetTypeHash( FPropertySoftPath const& Path );
 	TArray<FChainElement> PropertyChain;
diff --git a/Engine/Source/Editor/PropertyEditor/Private/DetailItemNode.cpp b/Engine/Source/Editor/PropertyEditor/Private/DetailItemNode.cpp
--- a/Engine/Source/Editor/PropertyEditor/Private/DetailItemNode.cpp
+++ b/Engine/Source/Editor/PropertyEditor/Private/DetailItemNode.cpp
@@ -800,6 +800,26 @@ FPropertyPath FDetailItemNode::GetPropertyPath() const
 	{
 		Ret = *FPropertyNode::CreatePropertyPath( PropertyNode.ToSharedRef() );
 	}
+
+	// add properties used by custom widgets
+	if (Customization.WidgetDecl)
+	{
+		for (const TSharedPtr<IPropertyHandle> &ItemPropHandle : Customization.WidgetDecl->PropertyHandles)
+		{
+			if (ItemPropHandle)
+			{
+				if (ItemPropHandle->GetIndexInArray() != INDEX_NONE)
+				{
+					Ret.AddProperty(FPropertyInfo(ItemPropHandle->GetParentHandle()->GetProperty(), INDEX_NONE));
+					Ret.AddProperty(FPropertyInfo(ItemPropHandle->GetProperty(), ItemPropHandle->GetIndexInArray()));
+				}
+				else
+				{
+					Ret.AddProperty(FPropertyInfo(ItemPropHandle->GetProperty(), INDEX_NONE));
+				}
+			}
+		}
+	}
 	return Ret;
 }
 
diff --git a/Engine/Source/Editor/PropertyEditor/Private/ObjectPropertyNode.cpp b/Engine/Source/Editor/PropertyEditor/Private/ObjectPropertyNode.cpp
--- a/Engine/Source/Editor/PropertyEditor/Private/ObjectPropertyNode.cpp
+++ b/Engine/Source/Editor/PropertyEditor/Private/ObjectPropertyNode.cpp
@@ -699,18 +699,17 @@ void FObjectPropertyNode::InternalInitChildNodes( FName SinglePropertyName )
 		// Iterate over all fields, creating items.
 		for( TFieldIterator<FProperty> It(BaseClass.Get()); It; ++It )
 		{
-			if (PropertyEditorHelpers::ShouldBeVisible(*this, *It))
+			FProperty* CurProp = *It;
+
+			// if a SinglePropertyName was provided, bypass the 'property visibility check' and always include it
+			if ((SinglePropertyName == NAME_None && PropertyEditorHelpers::ShouldBeVisible(*this, CurProp)) || CurProp->GetFName() == SinglePropertyName)
 			{
-				FProperty* CurProp = *It;
-				if( SinglePropertyName == NAME_None || CurProp->GetFName() == SinglePropertyName )
-				{
-					SortedProperties.Add(CurProp);
+				SortedProperties.Add(CurProp);
 
-					if( SinglePropertyName != NAME_None )
-					{
-						// Generate no other children
-						break;
-					}
+				if( SinglePropertyName != NAME_None )
+				{
+					// Generate no other children
+					break;
 				}
 			}
 		}
diff --git a/Engine/Source/Editor/PropertyEditor/Private/PropertyHandleImpl.cpp b/Engine/Source/Editor/PropertyEditor/Private/PropertyHandleImpl.cpp
--- a/Engine/Source/Editor/PropertyEditor/Private/PropertyHandleImpl.cpp
+++ b/Engine/Source/Editor/PropertyEditor/Private/PropertyHandleImpl.cpp
@@ -865,7 +865,7 @@ TSharedPtr<FPropertyNode> FPropertyValueImpl::GetChildNode( FName ChildName, boo
 TSharedPtr<FPropertyNode> FPropertyValueImpl::GetChildNode( int32 ChildIndex ) const
 {
 	TSharedPtr<FPropertyNode> PropertyNodePin = PropertyNode.Pin();
-	if( PropertyNodePin.IsValid() )
+	if( PropertyNodePin.IsValid() && ChildIndex < PropertyNodePin->GetNumChildNodes() )
 	{
 		return PropertyNodePin->GetChildNode( ChildIndex );
 	}
diff --git a/Engine/Source/Editor/PropertyEditor/Private/SDetailsViewBase.cpp b/Engine/Source/Editor/PropertyEditor/Private/SDetailsViewBase.cpp
--- a/Engine/Source/Editor/PropertyEditor/Private/SDetailsViewBase.cpp
+++ b/Engine/Source/Editor/PropertyEditor/Private/SDetailsViewBase.cpp
@@ -222,6 +222,39 @@ static void FindTreeNodeFromPropertyRecursive( const TArray< TSharedRef<FDetailT
 	}
 }
 
+static void GetPropertyPathToTreeNodes(const TArray< TSharedRef<FDetailTreeNode> >& RootTreeNodes, TMap<FString, TArray<TSharedRef<FDetailTreeNode>>> &OutMap)
+{
+	for (const TSharedRef<FDetailTreeNode>& TreeNode : RootTreeNodes)
+	{
+		if (TreeNode->IsLeaf())
+		{
+			TArray<TSharedRef<FDetailTreeNode>>& Nodes = OutMap.FindOrAdd(TreeNode->GetPropertyPath().ToString());
+			Nodes.Add(TreeNode);
+		}
+
+		// Need to check children even if we're a leaf, because all DetailItemNodes are leaves, even if they may have sub-children
+		TArray< TSharedRef<FDetailTreeNode> > Children;
+		TreeNode->GetChildren(Children);
+		GetPropertyPathToTreeNodes(Children, OutMap);
+	}
+}
+
+static TArray< TSharedRef< FDetailTreeNode > > FindBestFitTreeNodeFromProperty(const TArray< TSharedRef<FDetailTreeNode> >& RootTreeNodes, FPropertyPath Property)
+{
+	TMap<FString, TArray<TSharedRef<FDetailTreeNode>>> PropertyPathToTreeNodes;
+	GetPropertyPathToTreeNodes(RootTreeNodes, PropertyPathToTreeNodes);
+	
+	TArray< TSharedRef< FDetailTreeNode > >* TreeNodeChain = PropertyPathToTreeNodes.Find(Property.ToString());
+
+	// if we couldn't find the exact property, trim the path to get the next best option
+	while (!TreeNodeChain && Property.GetNumProperties() > 1)
+	{
+		Property = *Property.TrimPath(1);
+		TreeNodeChain = PropertyPathToTreeNodes.Find(Property.ToString());
+	}
+	return TreeNodeChain? MoveTemp(*TreeNodeChain) : TArray< TSharedRef< FDetailTreeNode > >{};
+}
+
 void SDetailsViewBase::HighlightProperty(const FPropertyPath& Property)
 {
 	TSharedPtr<FDetailTreeNode> PrevHighlightedNodePtr = CurrentlyHighlightedNode.Pin();
@@ -230,9 +263,9 @@ void SDetailsViewBase::HighlightProperty(const FPropertyPath& Property)
 		PrevHighlightedNodePtr->SetIsHighlighted(false);
 	}
 
+	const TArray< TSharedRef< FDetailTreeNode > > TreeNodeChain = FindBestFitTreeNodeFromProperty(RootTreeNodes, Property);
+	
 	TSharedPtr< FDetailTreeNode > FinalNodePtr = nullptr;
-	TArray< TSharedPtr< FDetailTreeNode > > TreeNodeChain;
-	FindTreeNodeFromPropertyRecursive(RootTreeNodes, Property, TreeNodeChain);
 	if (TreeNodeChain.Num() > 0)
 	{
 		FinalNodePtr = TreeNodeChain[0];
diff --git a/Engine/Source/Editor/PropertyEditor/Public/PropertyPath.h b/Engine/Source/Editor/PropertyEditor/Public/PropertyPath.h
--- a/Engine/Source/Editor/PropertyEditor/Public/PropertyPath.h
+++ b/Engine/Source/Editor/PropertyEditor/Public/PropertyPath.h
@@ -76,7 +76,7 @@ public:
 		return Properties[ 0 ];
 	}
 
-	TSharedRef< FPropertyPath > ExtendPath( const FPropertyInfo& NewLeaf )
+	TSharedRef< FPropertyPath > ExtendPath( const FPropertyInfo& NewLeaf ) const
 	{
 		TSharedRef< FPropertyPath > NewPath = MakeShareable( new FPropertyPath() );
 
@@ -86,7 +86,7 @@ public:
 		return NewPath;
 	}
 
-	TSharedRef< FPropertyPath > ExtendPath( const TSharedRef< FPropertyPath >& Extension )
+	TSharedRef< FPropertyPath > ExtendPath( const TSharedRef< FPropertyPath >& Extension ) const
 	{
 		TSharedRef< FPropertyPath > NewPath = MakeShareable( new FPropertyPath() );
 
@@ -100,7 +100,7 @@ public:
 		return NewPath;
 	}
 
-	TSharedRef< FPropertyPath > TrimPath( const int32 AmountToTrim )
+	TSharedRef< FPropertyPath > TrimPath( const int32 AmountToTrim ) const
 	{
 		TSharedRef< FPropertyPath > NewPath = MakeShareable( new FPropertyPath() );
 
