[BugFix] Blueprint Diff tool shows incorrectly scaled/colored icons

diff --git a/Engine/Source/Editor/Kismet/Private/DiffControl.cpp b/Engine/Source/Editor/Kismet/Private/DiffControl.cpp
--- a/Engine/Source/Editor/Kismet/Private/DiffControl.cpp
+++ b/Engine/Source/Editor/Kismet/Private/DiffControl.cpp
@@ -502,7 +502,7 @@ TSharedRef<SWidget> FGraphToDiff::GenerateCategoryWidget()
 	const UEdGraph* Graph = GraphOld ? GraphOld : GraphNew;
 	check(Graph);
 	
-	FLinearColor Color = (GraphOld && GraphNew) ? FLinearColor::White : FLinearColor(0.3f,0.3f,1.f);
+	FLinearColor Color = (GraphOld && GraphNew) ? DiffViewUtils::Identical() : FLinearColor(0.3f,0.3f,1.f);
 
 	const bool bHasDiffs = DiffListSource.Num() > 0;
 
diff --git a/Engine/Source/Editor/Kismet/Private/DiffUtils.cpp b/Engine/Source/Editor/Kismet/Private/DiffUtils.cpp
--- a/Engine/Source/Editor/Kismet/Private/DiffUtils.cpp
+++ b/Engine/Source/Editor/Kismet/Private/DiffUtils.cpp
@@ -975,7 +975,8 @@ FLinearColor DiffViewUtils::Differs()
 
 FLinearColor DiffViewUtils::Identical()
 {
-	return FLinearColor::White;
+	const static FLinearColor ForegroundColor = FAppStyle::GetColor("Graph.ForegroundColor");
+	return ForegroundColor;
 }
 
 FLinearColor DiffViewUtils::Missing()
@@ -1097,7 +1098,7 @@ SHorizontalBox::FSlot::FSlotArguments DiffViewUtils::Box(bool bIsPresent, FLinea
 		.AutoWidth()
 		.HAlign(HAlign_Right)
 		.VAlign(VAlign_Center)
-		.MaxWidth(8.0f)
+		.Padding(0.5f, 0.f)
 		[
 			SNew(SImage)
 			.ColorAndOpacity(Color)
