[BugFix] Crash in UEQA_SC_P4V when selecting a Component within the Blueprint Diff window
introduced by [REDACTED_PERSON] in UE-142036 CL: 21503527

diff --git a/Engine/Source/Editor/Kismet/Private/SKismetInspector.cpp b/Engine/Source/Editor/Kismet/Private/SKismetInspector.cpp
--- a/Engine/Source/Editor/Kismet/Private/SKismetInspector.cpp
+++ b/Engine/Source/Editor/Kismet/Private/SKismetInspector.cpp
@@ -857,8 +857,8 @@ bool SKismetInspector::IsPropertyVisible( const FPropertyAndParent& PropertyAndP
 		{
 			if (const UClass* OwningClass = Property->GetOwnerClass())
 			{
-				return OwningClass == Blueprint->ParentClass ||
-					OwningClass->ClassGeneratedBy == Blueprint;
+				return Blueprint && (OwningClass == Blueprint->ParentClass ||
+					OwningClass->ClassGeneratedBy == Blueprint);
 			}
 			return false;
 		};
