Fix thread safe function checkbox not being updated on undo/redo

diff --git a/Engine/Source/Editor/Kismet/Private/BlueprintDetailsCustomization.cpp b/Engine/Source/Editor/Kismet/Private/BlueprintDetailsCustomization.cpp
--- a/Engine/Source/Editor/Kismet/Private/BlueprintDetailsCustomization.cpp
+++ b/Engine/Source/Editor/Kismet/Private/BlueprintDetailsCustomization.cpp
@@ -5487,6 +5487,7 @@ void FBlueprintGraphActionDetails::OnIsThreadSafeFunctionModified(const ECheckBo
 		if( UK2Node_FunctionEntry* EntryPoint = Cast<UK2Node_FunctionEntry>(FunctionEntryNodePtr.Get()) )
 		{
 			const FScopedTransaction Transaction( TransactionType );
+			EntryPoint->Modify();
 			EntryPoint->MetaData.bThreadSafe = bThreadSafe;
 			FBlueprintEditorUtils::MarkBlueprintAsModified( EntryPoint->GetBlueprint() );
 		}
