Make sure classes are recompiled on load if their ParentClass has been changed

diff --git a/Engine/Source/Editor/Kismet/Private/BlueprintCompilationManager.cpp b/Engine/Source/Editor/Kismet/Private/BlueprintCompilationManager.cpp
--- a/Engine/Source/Editor/Kismet/Private/BlueprintCompilationManager.cpp
+++ b/Engine/Source/Editor/Kismet/Private/BlueprintCompilationManager.cpp
@@ -652,7 +652,6 @@ void FBlueprintCompilationManagerImpl::FlushCompilationQueueImpl(bool bSuppressB
 		// STAGE II: Filter out data only and interface blueprints:
 		for(int32 I = 0; I < QueuedRequests.Num(); ++I)
 		{
-			bool bSkipCompile = false;
 			FBPCompileRequestInternal& QueuedJob = QueuedRequests[I];
 			UBlueprint* QueuedBP = QueuedJob.UserData.BPToCompile;
 
@@ -661,8 +660,9 @@ void FBlueprintCompilationManagerImpl::FlushCompilationQueueImpl(bool bSuppressB
 				!(QueuedBP->GeneratedClass->ClassDefaultObject->HasAnyFlags(RF_NeedLoad)));
 			bool bDefaultComponentMustBeAdded = false;
 			bool bHasPendingUberGraphFrame = false;
+			UBlueprintGeneratedClass* BPGC = Cast<UBlueprintGeneratedClass>(QueuedBP->GeneratedClass);
 
-			if(UBlueprintGeneratedClass* BPGC = Cast<UBlueprintGeneratedClass>(QueuedBP->GeneratedClass))
+			if(BPGC)
 			{
 				if( BPGC->SimpleConstructionScript &&
 					BPGC->SimpleConstructionScript->GetSceneRootComponentTemplate(true) == nullptr)
@@ -674,10 +674,15 @@ void FBlueprintCompilationManagerImpl::FlushCompilationQueueImpl(bool bSuppressB
 				bHasPendingUberGraphFrame = BPGC->UberGraphFramePointerProperty || BPGC->UberGraphFunction;
 #endif//USE_UBER_GRAPH_PERSISTENT_FRAME
 			}
-			
-			if( FBlueprintEditorUtils::IsDataOnlyBlueprint(QueuedBP) && !QueuedBP->bHasBeenRegenerated && QueuedBP->GetLinker() && !bDefaultComponentMustBeAdded && !bHasPendingUberGraphFrame)
+
+			bool bSkipCompile = false;
+			const UClass* ParentClass = QueuedBP->ParentClass;
+			const bool bHasClassAndMatchesParent = BPGC && BPGC->GetSuperClass() == ParentClass;
+			if( bHasClassAndMatchesParent &&
+				FBlueprintEditorUtils::IsDataOnlyBlueprint(QueuedBP) && !QueuedBP->bHasBeenRegenerated && 
+				QueuedBP->GetLinker() && !bDefaultComponentMustBeAdded && !bHasPendingUberGraphFrame )
 			{
-				const UClass* ParentClass = QueuedBP->ParentClass;
+				// consider skipping the compile operation for this DOB:
 				if (ParentClass && ParentClass->HasAllClassFlags(CLASS_Native))
 				{
 					bSkipCompile = true;
