Fix for MyBlueprint Functions Drag & Drop, so it supports dragging pure functions on anim graph but blocks impure.
Aslo blocked all functions on StateMachine schema

diff --git a/Engine/Source/Editor/Kismet/Private/SMyBlueprint.cpp b/Engine/Source/Editor/Kismet/Private/SMyBlueprint.cpp
--- a/Engine/Source/Editor/Kismet/Private/SMyBlueprint.cpp
+++ b/Engine/Source/Editor/Kismet/Private/SMyBlueprint.cpp
@@ -6,6 +6,7 @@
 #include "Animation/AnimClassInterface.h"
 #include "AnimationGraph.h"
 #include "AnimationStateMachineGraph.h"
+#include "AnimationStateMachineSchema.h"
 #include "AnimGraphNode_StateMachineBase.h"
 #include "BPDelegateDragDropAction.h"
 #include "BPFunctionDragDropAction.h"
@@ -1919,9 +1920,9 @@ FReply SMyBlueprint::OnActionDragged( const TArray< TSharedPtr<FEdGraphSchemaAct
 			if (FuncAction->GraphType == EEdGraphSchemaAction_K2Graph::Function ||FuncAction->GraphType == EEdGraphSchemaAction_K2Graph::Interface)
 			{
 				// Callback function to report that the user cannot drop this function in the graph
-				auto CanDragDropAction = [](TSharedPtr<FEdGraphSchemaAction> /*DropAction*/, UEdGraph* HoveredGraphIn, FText& ImpededReasonOut, bool bIsBlueprintCallableFunction)->bool
+				auto CanDragDropAction = [](TSharedPtr<FEdGraphSchemaAction> /*DropAction*/, UEdGraph* HoveredGraphIn, FText& ImpededReasonOut, int32 FunctionFlags)->bool //bool bIsBlueprintCallableFunction, bool bIsPureFunction)->bool
 				{
-					if (!bIsBlueprintCallableFunction)
+					if ((FunctionFlags & (FUNC_BlueprintCallable | FUNC_BlueprintPure)) == 0)
 					{
 						ImpededReasonOut = LOCTEXT("NonBlueprintCallable", "This function was not marked as Blueprint Callable and cannot be placed in a graph!");
 						return false;
@@ -1929,9 +1930,21 @@ FReply SMyBlueprint::OnActionDragged( const TArray< TSharedPtr<FEdGraphSchemaAct
 					
 					if (const UEdGraphSchema_K2* K2Schema = Cast<UEdGraphSchema_K2>(HoveredGraphIn->GetSchema()))
 					{
-						if (K2Schema->DoesGraphSupportImpureFunctions(HoveredGraphIn) == false)
+						if ((FunctionFlags & (FUNC_BlueprintPure)) == 0)
 						{
-							ImpededReasonOut = LOCTEXT("GraphNoSupportImpureF", "The target graph does not support Impure Functions.");
+							if (K2Schema->DoesGraphSupportImpureFunctions(HoveredGraphIn) == false)
+							{
+								ImpededReasonOut = LOCTEXT("GraphNoSupportImpureF", "The target graph does not support Impure Functions.");
+								return false;
+							}
+						}
+					}
+					else
+					{
+						// TODO : Check if this else has to block everything or just explicit schemas
+						if (const UAnimationStateMachineSchema* AnimationStateMachineSchema = Cast<UAnimationStateMachineSchema>(HoveredGraphIn->GetSchema()))
+						{
+							ImpededReasonOut = LOCTEXT("GraphNoSupportImpureF", "The target graph does not support Functions.");
 							return false;
 						}
 					}
@@ -1939,23 +1952,26 @@ FReply SMyBlueprint::OnActionDragged( const TArray< TSharedPtr<FEdGraphSchemaAct
 					return true;
 				};
 
-				bool bIsBlueprintCallableFunction = false;
+				int32 FunctionFlags = 0;
+
 				if (FuncAction->EdGraph)
 				{
 					for (UEdGraphNode* GraphNode : FuncAction->EdGraph->Nodes)
 					{
 						if (UK2Node_FunctionEntry* Node = Cast<UK2Node_FunctionEntry>(GraphNode))
 						{
-							// See whether this node is a blueprint callable function
-							if (Node->GetFunctionFlags() & (FUNC_BlueprintCallable|FUNC_BlueprintPure))
-							{
-								bIsBlueprintCallableFunction = true;
-							}
+							FunctionFlags |= Node->GetFunctionFlags();
 						}
 					}
 				}
 
-				return FReply::Handled().BeginDragDrop(FKismetFunctionDragDropAction::New(InAction, FuncAction->FuncName, GetBlueprintObj()->SkeletonGeneratedClass, FMemberReference(), AnalyticsDelegate, FKismetDragDropAction::FCanBeDroppedDelegate::CreateLambda(CanDragDropAction, bIsBlueprintCallableFunction)));
+				return FReply::Handled().BeginDragDrop(FKismetFunctionDragDropAction::New(
+					InAction, 
+					FuncAction->FuncName, 
+					GetBlueprintObj()->SkeletonGeneratedClass, 
+					FMemberReference(), 
+					AnalyticsDelegate, 
+					FKismetDragDropAction::FCanBeDroppedDelegate::CreateLambda(CanDragDropAction, FunctionFlags)));
 			}
 			else if (FuncAction->GraphType == EEdGraphSchemaAction_K2Graph::Macro)
 			{
