sMyBlueprint : Fixed localization bug

diff --git a/Engine/Source/Editor/Kismet/Private/SMyBlueprint.cpp b/Engine/Source/Editor/Kismet/Private/SMyBlueprint.cpp
--- a/Engine/Source/Editor/Kismet/Private/SMyBlueprint.cpp
+++ b/Engine/Source/Editor/Kismet/Private/SMyBlueprint.cpp
@@ -1944,7 +1944,7 @@ FReply SMyBlueprint::OnActionDragged( const TArray< TSharedPtr<FEdGraphSchemaAct
 						// TODO : Check if this else has to block everything or just explicit schemas
 						if (const UAnimationStateMachineSchema* AnimationStateMachineSchema = Cast<UAnimationStateMachineSchema>(HoveredGraphIn->GetSchema()))
 						{
-							ImpededReasonOut = LOCTEXT("GraphNoSupportImpureF", "The target graph does not support Functions.");
+							ImpededReasonOut = LOCTEXT("GraphNoSupportFunctions", "The target graph does not support Functions.");
 							return false;
 						}
 					}
