[BugFix] There's no way to see local function variables in blueprint diffs, preventing me from reviewing their changes

diff --git a/Engine/Source/Editor/BlueprintGraph/Private/K2Node_FunctionEntry.cpp b/Engine/Source/Editor/BlueprintGraph/Private/K2Node_FunctionEntry.cpp
--- a/Engine/Source/Editor/BlueprintGraph/Private/K2Node_FunctionEntry.cpp
+++ b/Engine/Source/Editor/BlueprintGraph/Private/K2Node_FunctionEntry.cpp
@@ -714,53 +714,143 @@ void UK2Node_FunctionEntry::FindDiffs(UEdGraphNode* OtherNode, struct FDiffResul
 			Results.Add(Diff);
 		}
 
-		bool bLocalVarsDiffer = (LocalVariables.Num() != OtherFunction->LocalVariables.Num());
-
-		for (int32 i = 0; i < LocalVariables.Num() && !bLocalVarsDiffer; i++)
+		// constructs a map from variable guid to index in the provided array
+		auto MapVarGuidToArrayIndex = [](const TArray<FBPVariableDescription>& Variables)
 		{
-			const FBPVariableDescription& ThisVar = LocalVariables[i];
-			const FBPVariableDescription& OtherVar = OtherFunction->LocalVariables[i];
-			
-			// Can't do a raw compare, for local variable defaults we need to compare the struct
-			if (ThisVar.VarName != OtherVar.VarName
-				|| ThisVar.VarType != OtherVar.VarType
-				|| ThisVar.FriendlyName != OtherVar.FriendlyName
-				|| !ThisVar.Category.EqualTo(OtherVar.Category)
-				|| ThisVar.PropertyFlags != OtherVar.PropertyFlags
-				|| ThisVar.RepNotifyFunc != OtherVar.RepNotifyFunc
-				|| ThisVar.ReplicationCondition != OtherVar.ReplicationCondition)
+			TMap<FGuid, int32> Map;
+			for (int32 i = 0; i < Variables.Num(); i++)
 			{
-				bLocalVarsDiffer = true;
+				Map.Add(Variables[i].VarGuid, i);
 			}
-		}
+			return MoveTemp(Map);
+		};
 
-		if (bLocalVarsDiffer)
+		auto AddDiff = [this, OtherNode, &Results](const FText& DisplayString)
 		{
 			FDiffSingleResult Diff;
 			Diff.Diff = EDiffType::NODE_PROPERTY;
 			Diff.Node1 = this;
 			Diff.Node2 = OtherNode;
-			Diff.DisplayString = LOCTEXT("DIF_FunctionLocalVariables", "Function local variables have changed in structure");
-			Diff.Category = EDiffType::MODIFICATION;
-
+			Diff.DisplayString = DisplayString;
 			Results.Add(Diff);
-		}
-		else
+		};
+
+		// using maps so that order doesn't matter and it's easy to diff
+		const TMap<FGuid, int32> VarGuidToIndex = MapVarGuidToArrayIndex(LocalVariables);
+		const TMap<FGuid, int32> OtherVarGuidToIndex = MapVarGuidToArrayIndex(OtherFunction->LocalVariables);
+
+		FDiffSingleResult Diff;
+		Diff.Diff = EDiffType::NODE_PROPERTY;
+		Diff.Node1 = this;
+		Diff.Node2 = OtherNode;
+		const FText NodeName = FText::FromName(FunctionReference.GetMemberName());
+		for (const auto& [VarGuid,Index] : VarGuidToIndex)
 		{
-			TSharedPtr<FStructOnScope> MyLocals = GetFunctionVariableCache();
-			TSharedPtr<FStructOnScope> OtherLocals = OtherFunction->GetFunctionVariableCache();
-
-			if (MyLocals.IsValid() && MyLocals->IsValid() && OtherLocals.IsValid() && OtherLocals->IsValid())
+			const FBPVariableDescription& ThisVar = LocalVariables[Index];
+			const FText OldVarName = FText::FromName(ThisVar.VarName);
+			
+			if (const int32* OtherIndex = OtherVarGuidToIndex.Find(VarGuid))
 			{
-				// Check for local var diffs
-				FDiffSingleResult Diff;
-				Diff.Diff = EDiffType::NODE_PROPERTY;
-				Diff.Node1 = this;
-				Diff.Node2 = OtherNode;
-				Diff.ToolTip = LOCTEXT("DIF_FunctionLocalVariableDefaults", "Function local variable default values have changed");
 				Diff.Category = EDiffType::MODIFICATION;
-
-				DiffProperties(const_cast<UStruct*>(MyLocals->GetStruct()), const_cast<UStruct*>(OtherLocals->GetStruct()), MyLocals->GetStructMemory(), OtherLocals->GetStructMemory(), Results, Diff);
+				const FBPVariableDescription& OtherVar = OtherFunction->LocalVariables[*OtherIndex];
+				const FText NewVarName = FText::FromName(OtherVar.VarName);
+				if (ThisVar.VarName != OtherVar.VarName)
+				{
+					// variable name changed
+					AddDiff(FText::Format(
+						LOCTEXT("DIF_FunctionLocalVariableNameChange", "Local variable name changed from {0}::{1} to {0}::{2}"),
+						NodeName,
+						OldVarName,
+						NewVarName
+					));
+				}
+				if (ThisVar.VarType != OtherVar.VarType)
+				{
+					// type changed
+					AddDiff(FText::Format(
+						LOCTEXT("DIF_FunctionLocalVariableNameChange", "Local variable {0}::{1} was a {2} but is now a {3}"),
+						NodeName,
+						NewVarName,
+						UEdGraphSchema_K2::TypeToText(ThisVar.VarType),
+						UEdGraphSchema_K2::TypeToText(OtherVar.VarType)
+					));
+				}
+				if (!ThisVar.Category.EqualTo(OtherVar.Category))
+				{
+					// category changed
+					AddDiff(FText::Format(
+						LOCTEXT("DIF_FunctionLocalVariableNameChange", "Local variable {0}::{1} changed category from {2} to {3}"),
+						NodeName,
+						NewVarName,
+						ThisVar.Category,
+						OtherVar.Category
+					));
+				}
+				if (ThisVar.PropertyFlags != OtherVar.PropertyFlags)
+				{
+					// property flags changed
+					AddDiff(FText::Format(
+						LOCTEXT("DIF_FunctionLocalVariableNameChange", "Local variable {0}::{1} changed property flags"),
+						NodeName,
+						NewVarName
+					));
+				}
+				if (ThisVar.RepNotifyFunc != OtherVar.RepNotifyFunc)
+				{
+					// replication notify function changed
+					AddDiff(FText::Format(
+						LOCTEXT("DIF_FunctionLocalVariableNameChange", "Local variable {0}::{1} changed RepNotifyFunc"),
+						NodeName,
+						NewVarName
+					));
+				}
+				if (ThisVar.ReplicationCondition != OtherVar.ReplicationCondition)
+				{
+					// replication condition changed
+					AddDiff(FText::Format(
+						LOCTEXT("DIF_FunctionLocalVariableNameChange", "Local variable {0}::{1} changed ReplicationCondition"),
+						NodeName,
+						NewVarName
+					));
+				}
+				if (ThisVar.DefaultValue != OtherVar.DefaultValue)
+				{
+					// replication condition changed
+					AddDiff(FText::Format(
+						LOCTEXT("DIF_FunctionLocalVariableNameChange", "Local variable {0}::{1} changed default value from '{2}' to '{3}'"),
+						NodeName,
+						NewVarName,
+						FText::FromString(ThisVar.DefaultValue),
+						FText::FromString(OtherVar.DefaultValue)
+					));
+				}
+			}
+			else
+			{
+				// local variable is in this node but not in other node
+				Diff.Category = EDiffType::SUBTRACTION;
+				AddDiff(FText::Format(
+				LOCTEXT("DIF_FunctionLocalVariableNameChange", "Local variable {0}::{1} removed"),
+					NodeName,
+					OldVarName
+				));
+			}
+		}
+		
+		for (const auto& [VarGuid,Index] : OtherVarGuidToIndex)
+		{
+			const FBPVariableDescription& OtherVar = OtherFunction->LocalVariables[Index];
+			if (!VarGuidToIndex.Contains(VarGuid))
+			{
+				const FText NewVarName = FText::FromName(OtherVar.VarName);
+				
+				// local variable is in other node but not in this node
+				Diff.Category = EDiffType::ADDITION;
+				AddDiff(FText::Format(
+				LOCTEXT("DIF_FunctionLocalVariableNameChange", "Local variable {0}::{1} added"),
+					NodeName,
+					NewVarName
+				));
 			}
 		}
 	}
diff --git a/Engine/Source/Editor/Kismet/Private/BlueprintDetailsCustomization.cpp b/Engine/Source/Editor/Kismet/Private/BlueprintDetailsCustomization.cpp
--- a/Engine/Source/Editor/Kismet/Private/BlueprintDetailsCustomization.cpp
+++ b/Engine/Source/Editor/Kismet/Private/BlueprintDetailsCustomization.cpp
@@ -3567,6 +3567,234 @@ void FBlueprintGraphArgumentLayout::OnPrePinInfoChange(const FEdGraphPinType& Pi
 	}
 }
 
+void FBlueprintGraphLocalVariableGroupLayout::GenerateChildContent(IDetailChildrenBuilder& ChildrenBuilder)
+{
+	bool WasContentAdded = false;
+	if(TargetGraph.IsValid())
+	{
+		if(const UEdGraph* TopLevelGraph = FBlueprintEditorUtils::GetTopLevelGraph(TargetGraph.Get()))
+		{
+			bool bSchemaImplementsGetLocalVariables = false;
+		
+			// grab the parent graph's name
+			if (UEdGraphSchema const* Schema = TopLevelGraph->GetSchema())
+			{
+				FGraphDisplayInfo EdGraphDisplayInfo;
+				Schema->GetGraphDisplayInformation(*TopLevelGraph, EdGraphDisplayInfo);
+
+				// Try to get the local variables from the schema
+				TArray<FBPVariableDescription> LocalVariables;
+				bSchemaImplementsGetLocalVariables = Schema->GetLocalVariables(TargetGraph.Get(), LocalVariables);
+				for (const FBPVariableDescription& LocalVariable : LocalVariables)
+				{
+					TSharedRef<class FBlueprintGraphLocalVariableLayout> BlueprintLocalVariableLayout = MakeShareable(new FBlueprintGraphLocalVariableLayout(OwningFunction, LocalVariable));
+					ChildrenBuilder.AddCustomBuilder(BlueprintLocalVariableLayout);
+					WasContentAdded = true;
+				}
+			}
+			// If the schema did not return any local variables, try to get them from the function entry
+			if (!bSchemaImplementsGetLocalVariables)
+			{
+				TArray<UK2Node_FunctionEntry*> FunctionEntryNodes;
+				TopLevelGraph->GetNodesOfClass<UK2Node_FunctionEntry>(FunctionEntryNodes);
+				if (!FunctionEntryNodes.IsEmpty())
+				{
+					TArray<FBPVariableDescription>& LocalVariables = FunctionEntryNodes[0]->LocalVariables;
+					
+					// Search in all FunctionEntry nodes for their local variables
+					FText ActionCategory;
+					for (int I = 0; I < LocalVariables.Num(); ++I)
+					{
+						TSharedPtr<class FBlueprintGraphLocalVariableLayout> BlueprintLocalVariableLayout = nullptr;
+						if (PropertyHandle)
+						{
+							BlueprintLocalVariableLayout = MakeShareable(new FBlueprintGraphLocalVariableLayout(OwningFunction, PropertyHandle->GetChildHandle(I)));
+						}
+						else
+						{
+							BlueprintLocalVariableLayout = MakeShareable(new FBlueprintGraphLocalVariableLayout(OwningFunction, LocalVariables[I]));
+						}
+						ChildrenBuilder.AddCustomBuilder(BlueprintLocalVariableLayout.ToSharedRef());
+						WasContentAdded = true;		
+					}
+				}
+			}
+		}
+	}
+	if (!WasContentAdded)
+	{
+		// Add a text widget to let the user know to hit the + icon to add parameters.
+		ChildrenBuilder.AddCustomRow(FText::GetEmpty()).WholeRowContent()
+			.MaxDesiredWidth(980.f)
+			[
+				SNew(SHorizontalBox)
+				+ SHorizontalBox::Slot()
+					.VAlign(VAlign_Center)
+					.Padding(0.0f, 0.0f, 4.0f, 0.0f)
+					.AutoWidth()
+					[
+						SNew(STextBlock)
+						.Text(LOCTEXT("NoLocalVariablesAddedForBlueprint", "No Local Variables"))
+						.Font(IDetailLayoutBuilder::GetDetailFont())
+					]
+			];
+	}
+}
+
+TSharedPtr<IPropertyHandle> FBlueprintGraphLocalVariableLayout::GetPropertyHandle() const
+{
+	return Data.IsType<TSharedPtr<IPropertyHandle>>() ? Data.Get<TSharedPtr<IPropertyHandle>>() : nullptr;
+}
+
+const FBPVariableDescription& FBlueprintGraphLocalVariableLayout::GetVariable() const
+{
+	// if stored by property handle, extract value from there
+	if (const TSharedPtr<IPropertyHandle> PropertyHandle = GetPropertyHandle())
+	{
+		void* Address;
+		PropertyHandle->GetValueData(Address);
+		return *(FBPVariableDescription*)Address;
+	}
+	// if stored as variable description, return it directly
+	return Data.Get<FBPVariableDescription>();
+}
+
+void FBlueprintGraphLocalVariableLayout::SetVariable(const FBPVariableDescription& NewValue)
+{
+	const TSharedPtr<IPropertyHandle> PropertyHandle = GetPropertyHandle();
+	checkf(PropertyHandle != nullptr, TEXT("Tried to call a setter on a read-only local variable layout"));
+	
+	void* Address;
+	PropertyHandle->GetValueData(Address);
+	*(FBPVariableDescription*)Address = NewValue;
+}
+
+void FBlueprintGraphLocalVariableLayout::GenerateHeaderRowContent(FDetailWidgetRow& NodeRow)
+{
+	const UEdGraphSchema_K2* K2Schema = GetDefault<UEdGraphSchema_K2>();
+	TSharedPtr<IPropertyHandle> PropHandle = GetPropertyHandle();
+	
+	NodeRow
+	.NameContent()
+	[
+		SNew(SHorizontalBox)
+		+SHorizontalBox::Slot()
+		.FillWidth(1)
+		.VAlign(VAlign_Center)
+		[
+			SNew(SBox)
+			.MinDesiredWidth(125.f)
+			[
+				SAssignNew(VariableNameWidget, SEditableTextBox)
+				.Text( this, &FBlueprintGraphLocalVariableLayout::OnGetVarNameText )
+#if UE_BP_LOCAL_VAR_LAYOUT_SETTERS_IMPLEMENTED
+				.OnTextChanged(this, &FBlueprintGraphLocalVariableLayout::OnVarNameChange)
+				.OnTextCommitted(this, &FBlueprintGraphLocalVariableLayout::OnVarNameTextCommitted)
+#endif
+				.ToolTipText(this, &FBlueprintGraphLocalVariableLayout::OnGetVarToolTipText)
+				.Font( IDetailLayoutBuilder::GetDetailFont() )
+				.IsEnabled(!ShouldVarBeReadOnly())
+			]
+		]
+	]
+	.ValueContent()
+	.MaxDesiredWidth(980.f)
+	[
+		SNew(SHorizontalBox)
+		+SHorizontalBox::Slot()
+		.VAlign(VAlign_Center)
+		.Padding(0.0f)
+		.FillWidth(1.0f)
+		[
+			SNew(SPinTypeSelector, FGetPinTypeTree::CreateUObject(K2Schema, &UEdGraphSchema_K2::GetVariableTypeTree))
+				.TargetPinType(this, &FBlueprintGraphLocalVariableLayout::OnGetVariableType)
+#if UE_BP_LOCAL_VAR_LAYOUT_SETTERS_IMPLEMENTED
+				.OnPinTypePreChanged(this, &FBlueprintGraphLocalVariableLayout::OnPreVariableTypeChange)
+				.OnPinTypeChanged(this, &FBlueprintGraphLocalVariableLayout::VariableTypeChanged)
+#endif
+				.Schema(K2Schema)
+				.bAllowArrays(!ShouldVarBeReadOnly())
+				.IsEnabled(!ShouldVarBeReadOnly(true))
+				.Font( IDetailLayoutBuilder::GetDetailFont() )
+		]
+		+ SHorizontalBox::Slot()
+		.HAlign(HAlign_Right)
+		.VAlign(VAlign_Center)
+		.Padding(10, 0, 0, 0)
+		.AutoWidth()
+#if UE_BP_LOCAL_VAR_LAYOUT_SETTERS_IMPLEMENTED
+		[
+			PropertyCustomizationHelpers::MakeClearButton(FSimpleDelegate::CreateSP(this, &FBlueprintGraphLocalVariableLayout::OnRemoveClicked), LOCTEXT("LocalVariableDetailsClearTooltip", "Remove this variable."), !IsVariableEditingReadOnly())
+		]
+#endif
+
+	]
+	.PropertyHandleList({GetPropertyHandle()})
+#if UE_BP_LOCAL_VAR_LAYOUT_SETTERS_IMPLEMENTED
+	.DragDropHandler(/* implement drag drop handler and add it here */);
+	static_assert(false)
+#endif
+	;
+}
+
+void FBlueprintGraphLocalVariableLayout::GenerateChildContent(IDetailChildrenBuilder& ChildrenBuilder)
+{
+	if (const UFunction* Function = OwningFunction.Get())
+	{
+		const TSharedPtr<FStructOnScope> StructData = MakeShareable(new FStructOnScope(Function));
+
+		// ensure that the default value is up to date inside property
+		for (TFieldIterator<FProperty> PropertyIterator(Function); PropertyIterator; ++PropertyIterator)
+		{
+			const FProperty *VariableProperty = *PropertyIterator;
+			if (VariableProperty->GetFName() == GetName())
+			{
+				FBlueprintEditorUtils::PropertyValueFromString(VariableProperty, GetVariable().DefaultValue, StructData->GetStructMemory());
+				break;
+			}
+		}
+		if (IDetailPropertyRow* Row = ChildrenBuilder.AddExternalStructureProperty(StructData.ToSharedRef(), GetName()))
+		{
+			Row->DisplayName(LOCTEXT("LocalVariableDefaultValue", "Default Value"));
+		}
+	}
+}
+
+bool FBlueprintGraphLocalVariableLayout::ShouldVarBeReadOnly(bool bIsEditingPinType) const
+{
+#if UE_BP_LOCAL_VAR_LAYOUT_SETTERS_IMPLEMENTED
+	// if this is ever made non-const, implement this method
+	static_assert(false);
+#endif
+	return true;
+}
+
+bool FBlueprintGraphLocalVariableLayout::IsVariableEditingReadOnly(bool bIsEditingPinType) const
+{
+#if UE_BP_LOCAL_VAR_LAYOUT_SETTERS_IMPLEMENTED
+	// if this is ever made non-const, implement this method
+	static_assert(false)
+#endif
+	return true;
+}
+
+FText FBlueprintGraphLocalVariableLayout::OnGetVarNameText() const
+{
+	return FText::FromName(GetName());
+}
+
+FText FBlueprintGraphLocalVariableLayout::OnGetVarToolTipText() const
+{
+	const FBPVariableDescription& Variable = GetVariable();
+	const FText PinTypeText = UEdGraphSchema_K2::TypeToText(Variable.VarType);
+	return FText::Format(LOCTEXT("BlueprintArgToolTipText", "Name: {0}\nType: {1}"), FText::FromName(Variable.VarName), PinTypeText);
+}
+
+FEdGraphPinType FBlueprintGraphLocalVariableLayout::OnGetVariableType() const
+{
+	return GetVariable().VarType;
+}
+
 BEGIN_SLATE_FUNCTION_BUILD_OPTIMIZATION
 void FBlueprintGraphActionDetails::CustomizeDetails( IDetailLayoutBuilder& DetailLayout )
 {
@@ -4127,6 +4355,25 @@ void FBlueprintGraphActionDetails::CustomizeDetails( IDetailLayoutBuilder& Detai
 			OutputsCategory.HeaderContent(OutputsHeaderContentWidget);
 		}
 
+		if (bShowLocalVariables)
+		{
+			const UEdGraph* TopLevelGraph = FBlueprintEditorUtils::GetTopLevelGraph(GetGraph());
+			TSharedPtr<IPropertyHandle> LocalVariablesProperty = nullptr;
+			if (TopLevelGraph)
+			{
+				LocalVariablesProperty = DetailLayout.AddObjectPropertyData({FunctionEntryNode}, TEXT("LocalVariables"));
+			}
+		
+			IDetailCategoryBuilder& LocalVarsCategory = DetailLayout.EditCategory("Local Variables", LOCTEXT("FunctionDetailsLocalVariables", "Local Variables"));
+			TSharedRef<FBlueprintGraphLocalVariableGroupLayout> LocalVarsArgumentGroup =
+					MakeShareable(new FBlueprintGraphLocalVariableGroupLayout(SharedThis(this), GetGraph(), GetBlueprintObj(), FindFunction(), LocalVariablesProperty));
+			LocalVarsCategory.AddCustomBuilder(LocalVarsArgumentGroup);
+		
+			TSharedRef<SHorizontalBox> LocalVarsHeaderContentWidget = SNew(SHorizontalBox);
+
+			LocalVarsCategory.HeaderContent(LocalVarsHeaderContentWidget);
+		}
+		
 		// See if anything else wants to customize our details
 		TWeakPtr<FBlueprintEditor> BlueprintEditor = MyBlueprint.Pin()->GetBlueprintEditor();
 		FBlueprintEditorModule& BlueprintEditorModule = FModuleManager::GetModuleChecked<FBlueprintEditorModule>("Kismet");
diff --git a/Engine/Source/Editor/Kismet/Private/BlueprintDetailsCustomization.h b/Engine/Source/Editor/Kismet/Private/BlueprintDetailsCustomization.h
--- a/Engine/Source/Editor/Kismet/Private/BlueprintDetailsCustomization.h
+++ b/Engine/Source/Editor/Kismet/Private/BlueprintDetailsCustomization.h
@@ -562,18 +562,131 @@ private:
 	TSharedPtr<SGraphPin> DefaultValuePinWidget;
 };
 
+// local variables in the details view only support read-only mode at the moment.
+// the incomplete read/write functionality can be turned on by setting this to true but you'll have to write some more code to make it compile/work
+#define UE_BP_LOCAL_VAR_LAYOUT_SETTERS_IMPLEMENTED false
+
+/** Custom struct for each group of local variables in the function editing details */
+class FBlueprintGraphLocalVariableGroupLayout : public IDetailCustomNodeBuilder, public TSharedFromThis<FBlueprintGraphLocalVariableGroupLayout>
+{
+public:
+	FBlueprintGraphLocalVariableGroupLayout(TWeakPtr<class FBaseBlueprintGraphActionDetails> InGraphActionDetails, TWeakObjectPtr<UEdGraph> InTargetGraph,
+		TWeakObjectPtr<UBlueprint> InBlueprintObj, TWeakObjectPtr<UFunction> InOwningFunction, TSharedPtr<IPropertyHandle> Property = nullptr)
+		: GraphActionDetailsPtr(InGraphActionDetails)
+		, TargetGraph(InTargetGraph)
+		, BlueprintObj(InBlueprintObj)
+		, PropertyHandle(Property)
+		, OwningFunction(InOwningFunction) {}
+	
+	virtual TSharedPtr<IPropertyHandle> GetPropertyHandle() const override {return PropertyHandle;}
+
+private:
+	/** IDetailCustomNodeBuilder Interface*/
+	virtual void GenerateHeaderRowContent( FDetailWidgetRow& NodeRow ) override {}
+	virtual void GenerateChildContent( IDetailChildrenBuilder& ChildrenBuilder ) override;
+	virtual void Tick( float DeltaTime ) override {}
+	virtual bool RequiresTick() const override { return false; }
+	virtual FName GetName() const override { return NAME_None; }
+	virtual bool InitiallyCollapsed() const override { return false; }
+	
+private:
+	/** The parent graph action details customization */
+	TWeakPtr<class FBaseBlueprintGraphActionDetails> GraphActionDetailsPtr;
+
+	/** The target graph that the local variable's are on */
+	TWeakObjectPtr<UEdGraph> TargetGraph;
+
+	/** The blueprint object that the target graph is in */
+	TWeakObjectPtr<UBlueprint> BlueprintObj;
+
+	/** handle to property for array of FBPVariableDescription's */
+	TSharedPtr<IPropertyHandle> PropertyHandle;
+
+	/** function that owns the local variable's being displayed */
+	TWeakObjectPtr<UFunction> OwningFunction;
+};
+
+/** Custom struct for each local variable in the function editing details */
+class FBlueprintGraphLocalVariableLayout : public IDetailCustomNodeBuilder, public TSharedFromThis<FBlueprintGraphLocalVariableLayout>
+{
+public:
+	FBlueprintGraphLocalVariableLayout(TWeakObjectPtr<UFunction> InOwningFunction, const FBPVariableDescription& VariableDescription)
+		: OwningFunction(InOwningFunction)
+	{
+		Data.Set<FBPVariableDescription>(VariableDescription);
+	}
+	
+	FBlueprintGraphLocalVariableLayout(TWeakObjectPtr<UFunction> InOwningFunction, TSharedPtr<IPropertyHandle> Property = nullptr)
+		: OwningFunction(InOwningFunction)
+	{
+		Data.Set<TSharedPtr<IPropertyHandle>>(Property);
+	}
+	
+	virtual TSharedPtr<IPropertyHandle> GetPropertyHandle() const override;
+
+	const FBPVariableDescription& GetVariable() const;
+	void SetVariable(const FBPVariableDescription&);
+
+private:
+	/** IDetailCustomNodeBuilder Interface*/
+	virtual void SetOnRebuildChildren( FSimpleDelegate InOnRegenerateChildren ) override {}
+	virtual void GenerateHeaderRowContent( FDetailWidgetRow& NodeRow ) override;
+	virtual void GenerateChildContent( IDetailChildrenBuilder& ChildrenBuilder ) override;
+	virtual void Tick( float DeltaTime ) override {}
+	virtual bool RequiresTick() const override { return false; }
+	virtual FName GetName() const override { return GetVariable().VarName; }
+	virtual bool InitiallyCollapsed() const override { return true; }
+private:
+	
+	/** Determines if this pin should not be editable */
+	bool ShouldVarBeReadOnly(bool bIsEditingPinType = false) const;
+	
+	/** Determines if editing the pins on the node should be read only */
+	bool IsVariableEditingReadOnly(bool bIsEditingPinType = false) const;
+	
+
+	/** Callbacks for all the functionality for modifying local variables */
+	FText OnGetVarNameText() const;
+	FText OnGetVarToolTipText() const;
+	FEdGraphPinType OnGetVariableType() const;
+
+// currently this layout is for read only data. If you need to use it for read/write, you can implement the following methods.
+#if UE_BP_LOCAL_VAR_LAYOUT_SETTERS_IMPLEMENTED
+	void OnRemoveClicked();
+	void OnVarNameChange(const FText& InNewText);
+	void OnVarNameTextCommitted(const FText& NewText, ETextCommit::Type InTextCommit);
+	void VariableTypeChanged(const FEdGraphPinType& PinType);
+	void OnPreVariableTypeChange(const FEdGraphPinType& PinType);
+#endif
+
+private:
+
+	/** Holds a weak pointer to the argument name widget, used for error notifications */
+	TWeakPtr<SEditableTextBox> VariableNameWidget;
+
+	/** The SGraphPin widget created to show/edit default value */
+	TSharedPtr<SGraphPin> DefaultValueWidget;
+
+	/** either the read only copy of FBPVariableDescription or the property handle to the original variable description */
+	TVariant<TSharedPtr<IPropertyHandle>, FBPVariableDescription> Data;
+
+	/** function that owns this local variable */
+	TWeakObjectPtr<UFunction> OwningFunction;
+};
+
 /** Details customization for functions and graphs selected in the MyBlueprint panel */
 class FBlueprintGraphActionDetails : public FBaseBlueprintGraphActionDetails
 {
 public:
 	/** Makes a new instance of this detail layout class for a specific detail view requesting it */
-	static TSharedRef<class IDetailCustomization> MakeInstance(TWeakPtr<SMyBlueprint> InMyBlueprint)
+	static TSharedRef<class IDetailCustomization> MakeInstance(TWeakPtr<SMyBlueprint> InMyBlueprint, bool bInShowLocalVariables = false)
 	{
-		return MakeShareable(new FBlueprintGraphActionDetails(InMyBlueprint));
+		return MakeShareable(new FBlueprintGraphActionDetails(InMyBlueprint, bInShowLocalVariables));
 	}
 
-	FBlueprintGraphActionDetails(TWeakPtr<SMyBlueprint> InMyBlueprint)
+	FBlueprintGraphActionDetails(TWeakPtr<SMyBlueprint> InMyBlueprint, bool bInShowLocalVariables = false)
 		: FBaseBlueprintGraphActionDetails(InMyBlueprint)
+		, bShowLocalVariables(bInShowLocalVariables)
 	{
 	}
 	
@@ -715,6 +828,8 @@ private:
 
 	/** External detail customizations */
 	TArray<TSharedPtr<IDetailCustomization>> ExternalDetailCustomizations;
+
+	bool bShowLocalVariables = false;
 };
 
 /** Blueprint global options - managed list details customization base */
diff --git a/Engine/Source/Editor/Kismet/Private/SBlueprintDiff.cpp b/Engine/Source/Editor/Kismet/Private/SBlueprintDiff.cpp
--- a/Engine/Source/Editor/Kismet/Private/SBlueprintDiff.cpp
+++ b/Engine/Source/Editor/Kismet/Private/SBlueprintDiff.cpp
@@ -97,7 +97,8 @@ void FDiffPanel::InitializeDiffPanel()
 		.HideNameArea(true)
 		.ViewIdentifier(FName("BlueprintInspector"))
 		.MyBlueprintWidget(MyBlueprint)
-		.IsPropertyEditingEnabledDelegate(FIsPropertyEditingEnabled::CreateStatic([] { return false; }));
+		.IsPropertyEditingEnabledDelegate(FIsPropertyEditingEnabled::CreateStatic([] { return false; }))
+		.ShowLocalVariables(true);
 	DetailsView = Inspector;
 	MyBlueprint->SetInspector(DetailsView);
 }
@@ -632,7 +633,6 @@ void FDiffPanel::GeneratePanel(UEdGraph* Graph, TSharedPtr<TArray<FDiffSingleRes
 		}
 
 		MyBlueprint->SetFocusedGraph(Graph);
-		MyBlueprint->Refresh();
 
 		TSharedRef<SGraphEditor> Editor = SNew(SGraphEditor)
 			.AdditionalCommands(GraphEditorCommands)
@@ -784,7 +784,8 @@ void SBlueprintDiff::GenerateDifferencesList()
 			.HideNameArea(true)
 			.ViewIdentifier(FName("BlueprintInspector"))
 			.MyBlueprintWidget(InMyBlueprint)
-			.IsPropertyEditingEnabledDelegate(FIsPropertyEditingEnabled::CreateStatic([] { return false; }));
+			.IsPropertyEditingEnabledDelegate(FIsPropertyEditingEnabled::CreateStatic([] { return false; }))
+			.ShowLocalVariables(true);
 	};
 
 	PanelOld.GenerateMyBlueprintWidget();
diff --git a/Engine/Source/Editor/Kismet/Private/SKismetInspector.cpp b/Engine/Source/Editor/Kismet/Private/SKismetInspector.cpp
--- a/Engine/Source/Editor/Kismet/Private/SKismetInspector.cpp
+++ b/Engine/Source/Editor/Kismet/Private/SKismetInspector.cpp
@@ -418,7 +418,7 @@ void SKismetInspector::Construct(const FArguments& InArgs)
 		PropertyView->RegisterInstancedCustomPropertyLayout(UMulticastDelegatePropertyWrapper::StaticClass(), LayoutDelegateDetails);
 		
 		// Register function and variable details customization
-		FOnGetDetailCustomizationInstance LayoutGraphDetails = FOnGetDetailCustomizationInstance::CreateStatic(&FBlueprintGraphActionDetails::MakeInstance, MyBlueprint);
+		FOnGetDetailCustomizationInstance LayoutGraphDetails = FOnGetDetailCustomizationInstance::CreateStatic(&FBlueprintGraphActionDetails::MakeInstance, MyBlueprint, InArgs._ShowLocalVariables);
 		PropertyView->RegisterInstancedCustomPropertyLayout(UEdGraph::StaticClass(), LayoutGraphDetails);
 		PropertyView->RegisterInstancedCustomPropertyLayout(UK2Node_EditablePinBase::StaticClass(), LayoutGraphDetails);
 		PropertyView->RegisterInstancedCustomPropertyLayout(UK2Node_CallFunction::StaticClass(), LayoutGraphDetails);
diff --git a/Engine/Source/Editor/Kismet/Public/SKismetInspector.h b/Engine/Source/Editor/Kismet/Public/SKismetInspector.h
--- a/Engine/Source/Editor/Kismet/Public/SKismetInspector.h
+++ b/Engine/Source/Editor/Kismet/Public/SKismetInspector.h
@@ -49,6 +49,7 @@ public:
 		, _HideNameArea(true)
 		, _SetNotifyHook(true)
 		, _ShowTitleArea(false)
+		, _ShowLocalVariables(false)
 		{}
 
 		SLATE_ARGUMENT(TWeakPtr<FBlueprintEditor>, Kismet2)
@@ -60,6 +61,7 @@ public:
 		SLATE_ARGUMENT( FName, ViewIdentifier)
 		SLATE_ARGUMENT( bool, SetNotifyHook)
 		SLATE_ARGUMENT( bool, ShowTitleArea)
+		SLATE_ARGUMENT( bool, ShowLocalVariables)
 	SLATE_END_ARGS()
 
 	void Construct(const FArguments& InArgs);
