Fix for inspector panel not updating node properties when not explicitly selected

diff --git a/Engine/Source/Editor/Kismet/Private/BlueprintEditor.cpp b/Engine/Source/Editor/Kismet/Private/BlueprintEditor.cpp
--- a/Engine/Source/Editor/Kismet/Private/BlueprintEditor.cpp
+++ b/Engine/Source/Editor/Kismet/Private/BlueprintEditor.cpp
@@ -3777,7 +3777,7 @@ void FBlueprintEditor::OnGraphEditorFocused(const TSharedRef<SGraphEditor>& InGr
 
 	// Update the inspector as well, to show selection from the focused graph editor
 	FGraphPanelSelectionSet SelectedNodes = GetSelectedNodes();
-	FocusInspectorOnGraphSelection(SelectedNodes);
+	FocusInspectorOnGraphSelection(SelectedNodes, /*bForceRefresh=*/ true);
 
 	// During undo, garbage graphs can be temporarily brought into focus, ensure that before a refresh of the MyBlueprint window that the graph is owned by a Blueprint
 	if ( FocusedGraphEdPtr.IsValid() && MyBlueprintWidget.IsValid() )
