Always send BP compilation results to UI through the game-thread

diff --git a/Engine/Source/Editor/Kismet/Private/BlueprintCompilationManager.cpp b/Engine/Source/Editor/Kismet/Private/BlueprintCompilationManager.cpp
--- a/Engine/Source/Editor/Kismet/Private/BlueprintCompilationManager.cpp
+++ b/Engine/Source/Editor/Kismet/Private/BlueprintCompilationManager.cpp
@@ -2,6 +2,7 @@
 
 #include "BlueprintCompilationManager.h"
 
+#include "Async/Async.h"
 #include "Kismet2/BlueprintEditorUtils.h"
 #include "BlueprintCompilerExtension.h"
 #include "BlueprintEditorSettings.h"
@@ -1629,10 +1630,28 @@ void FBlueprintCompilationManagerImpl::FlushCompilationQueueImpl(bool bSuppressB
 
 			if(CompilerData.ShouldRegisterCompilerResults())
 			{
-				// This helper structure registers the results log messages with the UI control that displays them:
-				FScopedBlueprintMessageLog MessageLog(BP);
-				MessageLog.Log->ClearMessages();
-				MessageLog.Log->AddMessages(CompilerData.ActiveResultsLog->Messages, false);
+				if (IsInGameThread())
+				{
+					// This helper structure registers the results log messages with the UI control that displays them:
+					FScopedBlueprintMessageLog MessageLog(BP);
+					MessageLog.Log->ClearMessages();
+					MessageLog.Log->AddMessages(CompilerData.ActiveResultsLog->Messages, false);
+				}
+				else
+				{
+					Async(EAsyncExecution::TaskGraphMainThread,
+						[Messages = CompilerData.ActiveResultsLog->Messages, WeakBP = TWeakObjectPtr<UBlueprint>(BP)]()
+						{
+							if (WeakBP.IsValid())
+							{
+								FScopedBlueprintMessageLog MessageLog(WeakBP.Get());
+								MessageLog.Log->ClearMessages();
+								MessageLog.Log->AddMessages(Messages, false);
+							}
+						}
+					);
+				}
+				
 			}
 
 			if(CompilerData.ShouldSetTemporaryBlueprintFlags())
