Fix null dereference in the Kismet inspector widget.

#rnx

diff --git a/Engine/Source/Editor/Kismet/Private/SKismetInspector.cpp b/Engine/Source/Editor/Kismet/Private/SKismetInspector.cpp
--- a/Engine/Source/Editor/Kismet/Private/SKismetInspector.cpp
+++ b/Engine/Source/Editor/Kismet/Private/SKismetInspector.cpp
@@ -838,7 +838,7 @@ bool SKismetInspector::IsPropertyVisible( const FPropertyAndParent& PropertyAndP
 	for (const TWeakObjectPtr<UObject>& SelectedObject : SelectedObjects)
 	{
 		UObject* Object = SelectedObject.Get();
-		if (!Object->IsTemplate())
+		if (!Object || !Object->IsTemplate())
 		{
 			bIsEditingTemplate = false;
 			break;
