[BugFix] SkeletonOnly Blueprints were reconstructing nodes in BP diff tool and causing crash

diff --git a/Engine/Source/Editor/Kismet/Private/BlueprintCompilationManager.cpp b/Engine/Source/Editor/Kismet/Private/BlueprintCompilationManager.cpp
--- a/Engine/Source/Editor/Kismet/Private/BlueprintCompilationManager.cpp
+++ b/Engine/Source/Editor/Kismet/Private/BlueprintCompilationManager.cpp
@@ -444,7 +444,7 @@ struct FCompilerData
 	bool ShouldValidate() const { return JobType == ECompilationManagerJobType::Normal; }
 	bool ShouldRegenerateSkeleton() const { return JobType != ECompilationManagerJobType::RelinkOnly; }
 	bool ShouldMarkUpToDateAfterSkeletonStage() const { return IsSkeletonOnly(); }
-	bool ShouldReconstructNodes() const { return JobType == ECompilationManagerJobType::Normal || BP->bIsRegeneratingOnLoad; }
+	bool ShouldReconstructNodes() const { return JobType == ECompilationManagerJobType::Normal || (!IsSkeletonOnly() && BP->bIsRegeneratingOnLoad); }
 	bool ShouldSkipReinstancerCreation() const { return (IsSkeletonOnly() && (!BP->ParentClass || BP->ParentClass->IsNative())); }
 	bool ShouldInitiateReinstancing() const { return JobType == ECompilationManagerJobType::Normal || BP->bIsRegeneratingOnLoad; }
 	bool ShouldCompileClassLayout() const { return JobType == ECompilationManagerJobType::Normal; }
