Add a proper Modify call and scoped transaction for when you add a call to the parent BP node. This makes undo/redo actually track the change and work properly.

#rnx

diff --git a/Engine/Source/Editor/Kismet/Private/BlueprintEditor.cpp b/Engine/Source/Editor/Kismet/Private/BlueprintEditor.cpp
--- a/Engine/Source/Editor/Kismet/Private/BlueprintEditor.cpp
+++ b/Engine/Source/Editor/Kismet/Private/BlueprintEditor.cpp
@@ -5357,6 +5357,9 @@ void FBlueprintEditor::OnAddParentNode()
 			UEdGraph* TargetGraph = FunctionFromNode.Node->GetGraph();
 			if (ValidParent && TargetGraph)
 			{
+				const FScopedTransaction Transaction(LOCTEXT("AddParentNode", "Add Parent Node"));
+				TargetGraph->Modify();
+
 				FGraphNodeCreator<UK2Node_CallParentFunction> FunctionNodeCreator(*TargetGraph);
 				UK2Node_CallParentFunction* ParentFunctionNode = FunctionNodeCreator.CreateNode();
 				ParentFunctionNode->SetFromFunction(ValidParent);
