Fixed typo: SeperatorTitle -> SeparatorTitle

diff --git a/Engine/Source/Editor/Kismet/Private/SMyBlueprint.cpp b/Engine/Source/Editor/Kismet/Private/SMyBlueprint.cpp
--- a/Engine/Source/Editor/Kismet/Private/SMyBlueprint.cpp
+++ b/Engine/Source/Editor/Kismet/Private/SMyBlueprint.cpp
@@ -764,70 +764,70 @@ void SMyBlueprint::OnCategoryNameCommitted(const FText& InNewText, ETextCommit::
 
 FText SMyBlueprint::OnGetSectionTitle( int32 InSectionID )
 {
-	FText SeperatorTitle;
+	FText SeparatorTitle;
 	/* Setup an appropriate name for the section for this node */
 	switch( InSectionID )
 	{
 	case NodeSectionID::VARIABLE:
-		SeperatorTitle = NSLOCTEXT("GraphActionNode", "Variables", "Variables");
+		SeparatorTitle = NSLOCTEXT("GraphActionNode", "Variables", "Variables");
 		break;
 	case NodeSectionID::COMPONENT:
-		SeperatorTitle = NSLOCTEXT("GraphActionNode", "Components", "Components");
+		SeparatorTitle = NSLOCTEXT("GraphActionNode", "Components", "Components");
 		break;
 	case NodeSectionID::FUNCTION:
 		if ( OverridableFunctionActions.Num() > 0 )
 		{
-			SeperatorTitle = FText::Format(NSLOCTEXT("GraphActionNode", "FunctionsOverridableFormat", "Functions <TinyText.Subdued>({0} Overridable)</>"), FText::AsNumber(OverridableFunctionActions.Num()));
+			SeparatorTitle = FText::Format(NSLOCTEXT("GraphActionNode", "FunctionsOverridableFormat", "Functions <TinyText.Subdued>({0} Overridable)</>"), FText::AsNumber(OverridableFunctionActions.Num()));
 		}
 		else
 		{
-			SeperatorTitle = NSLOCTEXT("GraphActionNode", "Functions", "Functions");
+			SeparatorTitle = NSLOCTEXT("GraphActionNode", "Functions", "Functions");
 		}
 
 		break;
 	case NodeSectionID::FUNCTION_OVERRIDABLE:
-		SeperatorTitle = NSLOCTEXT("GraphActionNode", "OverridableFunctions", "Overridable Functions");
+		SeparatorTitle = NSLOCTEXT("GraphActionNode", "OverridableFunctions", "Overridable Functions");
 		break;
 	case NodeSectionID::MACRO:
-		SeperatorTitle = NSLOCTEXT("GraphActionNode", "Macros", "Macros");
+		SeparatorTitle = NSLOCTEXT("GraphActionNode", "Macros", "Macros");
 		break;
 	case NodeSectionID::INTERFACE:
-		SeperatorTitle = NSLOCTEXT("GraphActionNode", "Interfaces", "Interfaces");
+		SeparatorTitle = NSLOCTEXT("GraphActionNode", "Interfaces", "Interfaces");
 		break;
 	case NodeSectionID::DELEGATE:
-		SeperatorTitle = NSLOCTEXT("GraphActionNode", "EventDispatchers", "Event Dispatchers");
+		SeparatorTitle = NSLOCTEXT("GraphActionNode", "EventDispatchers", "Event Dispatchers");
 		break;	
 	case NodeSectionID::GRAPH:
-		SeperatorTitle = NSLOCTEXT("GraphActionNode", "Graphs", "Graphs");
+		SeparatorTitle = NSLOCTEXT("GraphActionNode", "Graphs", "Graphs");
 		break;
 	case NodeSectionID::ANIMGRAPH:
-		SeperatorTitle = NSLOCTEXT("GraphActionNode", "AnimationGraphs", "Animation Graphs");
+		SeparatorTitle = NSLOCTEXT("GraphActionNode", "AnimationGraphs", "Animation Graphs");
 		break;
 	case NodeSectionID::ANIMLAYER:
-		SeperatorTitle = NSLOCTEXT("GraphActionNode", "AnimationLayers", "Animation Layers");
+		SeparatorTitle = NSLOCTEXT("GraphActionNode", "AnimationLayers", "Animation Layers");
 		break;
 	case NodeSectionID::USER_ENUM:
-		SeperatorTitle = NSLOCTEXT("GraphActionNode", "Userenums", "User Enums");
+		SeparatorTitle = NSLOCTEXT("GraphActionNode", "Userenums", "User Enums");
 		break;	
 	case NodeSectionID::LOCAL_VARIABLE:
 		if ( GetFocusedGraph() )
 		{
-			SeperatorTitle = FText::Format(NSLOCTEXT("GraphActionNode", "LocalVariables_Focused", "Local Variables <TinyText.Subdued>({0})</>"), FText::FromName(GetFocusedGraph()->GetFName()));
+			SeparatorTitle = FText::Format(NSLOCTEXT("GraphActionNode", "LocalVariables_Focused", "Local Variables <TinyText.Subdued>({0})</>"), FText::FromName(GetFocusedGraph()->GetFName()));
 		}
 		else
 		{
-			SeperatorTitle = NSLOCTEXT("GraphActionNode", "LocalVariables", "Local Variables");
+			SeparatorTitle = NSLOCTEXT("GraphActionNode", "LocalVariables", "Local Variables");
 		}
 		break;
 	case NodeSectionID::USER_STRUCT:
-		SeperatorTitle = NSLOCTEXT("GraphActionNode", "Userstructs", "User Structs");
+		SeparatorTitle = NSLOCTEXT("GraphActionNode", "Userstructs", "User Structs");
 		break;	
 	default:
 	case NodeSectionID::NONE:
-		SeperatorTitle = FText::GetEmpty();
+		SeparatorTitle = FText::GetEmpty();
 		break;
 	}
-	return SeperatorTitle;
+	return SeparatorTitle;
 }
 
 TSharedRef<SWidget> SMyBlueprint::OnGetSectionWidget(TSharedRef<SWidget> RowWidget, int32 InSectionID)
