During re-instancing, we were previously checking wrongly against the immediate outer `UObject` to see if it was in
   a transient state. This caught some unexpected classes in the mix, which was removing the archetypes during re-instancing
   and thus causing one of our types, which was indirectly from a `UObject` archetype that had
   been already removed, to raise the warning.

diff --git a/Engine/Source/Editor/Kismet/Private/BlueprintCompilationManager.cpp b/Engine/Source/Editor/Kismet/Private/BlueprintCompilationManager.cpp
--- a/Engine/Source/Editor/Kismet/Private/BlueprintCompilationManager.cpp
+++ b/Engine/Source/Editor/Kismet/Private/BlueprintCompilationManager.cpp
@@ -2354,7 +2354,7 @@ void FBlueprintCompilationManagerImpl::ReinstanceBatch(TArray<FReinstancingJob>&
 					// that things that are not directly outered to the transient package will be 
 					// 'reinst'd', this is specifically to handle components, which need to be up to date
 					// on the REINST_ actor class:
-					return !bIsArchetype || Obj->GetOuter() == GetTransientPackage() || Obj->HasAnyFlags(RF_NewerVersionExists); 
+					return !bIsArchetype || Obj->GetOutermost() == GetTransientPackage() || Obj->HasAnyFlags(RF_NewerVersionExists); 
 				}
 			);
 
