PR #9659: Adding KISMET_API allows creating visual Diff for custom node graph editors! (Contributed by [REDACTED_PERSON])

diff --git a/Engine/Source/Editor/Kismet/Private/DiffControl.h b/Engine/Source/Editor/Kismet/Private/DiffControl.h
--- a/Engine/Source/Editor/Kismet/Private/DiffControl.h
+++ b/Engine/Source/Editor/Kismet/Private/DiffControl.h
@@ -14,7 +14,7 @@ class UEdGraph;
 struct FEdGraphEditAction;
 
 /** Interface responsible for generating FBlueprintDifferenceTreeEntry's for visual diff tools */
-class IDiffControl
+class KISMET_API IDiffControl
 {
 public:
 	virtual ~IDiffControl() {}
@@ -66,7 +66,7 @@ private:
 
 /** Generic wrapper around a details view template parameter determines whether TreeEntries are populated */
 template<bool bPopulateOutTreeEntries>
-class TDetailsDiffControl : public TSharedFromThis<TDetailsDiffControl<bPopulateOutTreeEntries>>, public IDiffControl
+class KISMET_API TDetailsDiffControl : public TSharedFromThis<TDetailsDiffControl<bPopulateOutTreeEntries>>, public IDiffControl
 {
 public:
 	TDetailsDiffControl(const UObject* InOldObject, const UObject* InNewObject, FOnDiffEntryFocused InSelectionCallback)
diff --git a/Engine/Source/Editor/Kismet/Public/DiffUtils.h b/Engine/Source/Editor/Kismet/Public/DiffUtils.h
--- a/Engine/Source/Editor/Kismet/Public/DiffUtils.h
+++ b/Engine/Source/Editor/Kismet/Public/DiffUtils.h
@@ -291,7 +291,7 @@ namespace DiffUtils
 DECLARE_DELEGATE(FOnDiffEntryFocused);
 DECLARE_DELEGATE_RetVal(TSharedRef<SWidget>, FGenerateDiffEntryWidget);
 
-class FBlueprintDifferenceTreeEntry
+class KISMET_API FBlueprintDifferenceTreeEntry
 {
 public:
 	FBlueprintDifferenceTreeEntry(FOnDiffEntryFocused InOnFocus, FGenerateDiffEntryWidget InGenerateWidget, TArray< TSharedPtr<FBlueprintDifferenceTreeEntry> > InChildren = TArray< TSharedPtr<FBlueprintDifferenceTreeEntry> >())
@@ -303,16 +303,16 @@ public:
 	}
 
 	/** Displays message to user saying there are no differences */
-	KISMET_API static TSharedPtr<FBlueprintDifferenceTreeEntry> NoDifferencesEntry();
+	static TSharedPtr<FBlueprintDifferenceTreeEntry> NoDifferencesEntry();
 
 	/** Displays message to user warning that there may be undetected differences */
-	KISMET_API static TSharedPtr<FBlueprintDifferenceTreeEntry> UnknownDifferencesEntry();
+	static TSharedPtr<FBlueprintDifferenceTreeEntry> UnknownDifferencesEntry();
 
 	/** Create category message for the diff UI */
-	KISMET_API static TSharedPtr<FBlueprintDifferenceTreeEntry> CreateCategoryEntry(const FText& LabelText, const FText& ToolTipText, FOnDiffEntryFocused FocusCallback, const TArray< TSharedPtr<FBlueprintDifferenceTreeEntry> >& Children, bool bHasDifferences);
+	static TSharedPtr<FBlueprintDifferenceTreeEntry> CreateCategoryEntry(const FText& LabelText, const FText& ToolTipText, FOnDiffEntryFocused FocusCallback, const TArray< TSharedPtr<FBlueprintDifferenceTreeEntry> >& Children, bool bHasDifferences);
 
 	/** Create category message for the merge UI */
-	KISMET_API static TSharedPtr<FBlueprintDifferenceTreeEntry> CreateCategoryEntryForMerge(const FText& LabelText, const FText& ToolTipText, FOnDiffEntryFocused FocusCallback, const TArray< TSharedPtr<FBlueprintDifferenceTreeEntry> >& Children, bool bHasRemoteDifferences, bool bHasLocalDifferences, bool bHasConflicts);
+	static TSharedPtr<FBlueprintDifferenceTreeEntry> CreateCategoryEntryForMerge(const FText& LabelText, const FText& ToolTipText, FOnDiffEntryFocused FocusCallback, const TArray< TSharedPtr<FBlueprintDifferenceTreeEntry> >& Children, bool bHasRemoteDifferences, bool bHasLocalDifferences, bool bHasConflicts);
 	
 	FOnDiffEntryFocused OnFocus;
 	FGenerateDiffEntryWidget GenerateWidget;
