Cleanup FBlueprintEditorToolbar so that AddOptionsDropdown function can be removed

diff --git a/Engine/Source/Editor/Kismet/Private/SBlueprintEditorToolbar.cpp b/Engine/Source/Editor/Kismet/Private/SBlueprintEditorToolbar.cpp
--- a/Engine/Source/Editor/Kismet/Private/SBlueprintEditorToolbar.cpp
+++ b/Engine/Source/Editor/Kismet/Private/SBlueprintEditorToolbar.cpp
@@ -284,46 +284,48 @@ void FFullBlueprintEditorCommands::RegisterCommands()
 
 namespace BlueprintEditorToolbarImpl
 {
-	static TSharedRef<SWidget> GenerateCompileOptionsWidget(TSharedRef<FUICommandList> CommandList);
-	static void MakeSaveOnCompileSubMenu(FMenuBuilder& InMenuBuilder);
-	static void MakeCompileDeveloperSubMenu(FMenuBuilder& InMenuBuilder);
+	static void GenerateCompileOptionsMenu(UToolMenu* InMenu);
+	static void MakeSaveOnCompileSubMenu(UToolMenu* InMenu);
+	static void MakeCompileDeveloperSubMenu(UToolMenu* InMenu);
 };
 
-static TSharedRef<SWidget> BlueprintEditorToolbarImpl::GenerateCompileOptionsWidget(TSharedRef<FUICommandList> CommandList)
+static void BlueprintEditorToolbarImpl::GenerateCompileOptionsMenu(UToolMenu* InMenu)
 {
-	FMenuBuilder MenuBuilder(/*bShouldCloseWindowAfterMenuSelection =*/true, CommandList);
-
+	FToolMenuSection& Section = InMenu->AddSection("Section");
 	const FFullBlueprintEditorCommands& Commands = FFullBlueprintEditorCommands::Get();
 
 	// @TODO: disable the menu and change up the tooltip when all sub items are disabled
-	MenuBuilder.AddSubMenu(
+	Section.AddSubMenu(
+		"SaveOnCompile",
 		LOCTEXT("SaveOnCompileSubMenu", "Save on Compile"),
 		LOCTEXT("SaveOnCompileSubMenu_ToolTip", "Determines how the Blueprint is saved whenever you compile it."),
-		FNewMenuDelegate::CreateStatic(&BlueprintEditorToolbarImpl::MakeSaveOnCompileSubMenu));
+		FNewToolMenuDelegate::CreateStatic(&BlueprintEditorToolbarImpl::MakeSaveOnCompileSubMenu));
 
-	MenuBuilder.AddMenuEntry(Commands.JumpToErrorNode);
+	Section.AddMenuEntry(Commands.JumpToErrorNode);
 
-// 	MenuBuilder.AddSubMenu(
+// 	Section.AddSubMenu(
+// 		"DevCompile",
 // 		LOCTEXT("DevCompileSubMenu", "Developer"),
 // 		LOCTEXT("DevCompileSubMenu_ToolTip", "Advanced settings that aid in devlopment/debugging of the Blueprint system as a whole."),
 // 		FNewMenuDelegate::CreateStatic(&BlueprintEditorToolbarImpl::MakeCompileDeveloperSubMenu));
 
-	return MenuBuilder.MakeWidget();
 }
 
-static void BlueprintEditorToolbarImpl::MakeSaveOnCompileSubMenu(FMenuBuilder& InMenuBuilder)
+static void BlueprintEditorToolbarImpl::MakeSaveOnCompileSubMenu(UToolMenu* InMenu)
 {
+	FToolMenuSection& Section = InMenu->AddSection("Section");
 	const FFullBlueprintEditorCommands& Commands = FFullBlueprintEditorCommands::Get();
-	InMenuBuilder.AddMenuEntry(Commands.SaveOnCompile_Never);
-	InMenuBuilder.AddMenuEntry(Commands.SaveOnCompile_SuccessOnly);
-	InMenuBuilder.AddMenuEntry(Commands.SaveOnCompile_Always);
+	Section.AddMenuEntry(Commands.SaveOnCompile_Never);
+	Section.AddMenuEntry(Commands.SaveOnCompile_SuccessOnly);
+	Section.AddMenuEntry(Commands.SaveOnCompile_Always);
 }
 
-static void BlueprintEditorToolbarImpl::MakeCompileDeveloperSubMenu(FMenuBuilder& InMenuBuilder)
+static void BlueprintEditorToolbarImpl::MakeCompileDeveloperSubMenu(UToolMenu* InMenu)
 {
+	FToolMenuSection& Section = InMenu->AddSection("Section");
 	const FBlueprintEditorCommands& EditorCommands = FBlueprintEditorCommands::Get();
-	InMenuBuilder.AddMenuEntry(EditorCommands.SaveIntermediateBuildProducts);
-	InMenuBuilder.AddMenuEntry(EditorCommands.ShowActionMenuItemSignatures);
+	Section.AddMenuEntry(EditorCommands.SaveIntermediateBuildProducts);
+	Section.AddMenuEntry(EditorCommands.ShowActionMenuItemSignatures);
 }
 
 
@@ -367,21 +369,22 @@ void FBlueprintEditorToolbar::AddCompileToolbar(UToolMenu* InMenu)
 			{
 				const FFullBlueprintEditorCommands& Commands = FFullBlueprintEditorCommands::Get();
 
-				FToolMenuEntry CompileButton = FToolMenuEntry::InitToolBarButton(
+				FToolMenuEntry& CompileButton = InSection.AddEntry(FToolMenuEntry::InitToolBarButton(
 					Commands.Compile,
 					TAttribute<FText>(),
 					TAttribute<FText>(BlueprintEditorToolbar.ToSharedRef(), &FBlueprintEditorToolbar::GetStatusTooltip),
 					TAttribute<FSlateIcon>(BlueprintEditorToolbar.ToSharedRef(), &FBlueprintEditorToolbar::GetStatusImage),
-					"CompileBlueprint");
-
+					"CompileBlueprint"));
 				CompileButton.StyleNameOverride = "CalloutToolbar";
 
-				CompileButton.AddOptionsDropdown(
+				FToolMenuEntry& CompileOptions = InSection.AddEntry(FToolMenuEntry::InitComboButton(
+					"CompileComboButton",
 					FUIAction(),
-					FOnGetContent::CreateStatic(&BlueprintEditorToolbarImpl::GenerateCompileOptionsWidget, Context->BlueprintEditor.Pin()->GetToolkitCommands()),
-					LOCTEXT("BlupeintCompileOptions_ToolbarTooltip", "Options to customize how Blueprints compile"));
-
-				InSection.AddEntry(CompileButton);
+					FNewToolMenuDelegate::CreateStatic(&BlueprintEditorToolbarImpl::GenerateCompileOptionsMenu),
+					LOCTEXT("BlupeintCompileOptions_ToolbarTooltip", "Options to customize how Blueprints compile")
+				));
+				CompileOptions.StyleNameOverride = "CalloutToolbar";
+				CompileOptions.ToolBarData.bSimpleComboBox = true;
 			}
 		}
 	}));
@@ -398,17 +401,15 @@ void FBlueprintEditorToolbar::AddCompileToolbar(UToolMenu* InMenu)
 				TSharedPtr<class FBlueprintEditorToolbar> BlueprintEditorToolbar = Context->BlueprintEditor.Pin()->GetToolbarBuilder();
 				if (BlueprintEditorToolbar.IsValid())
 				{
-					InSection.InsertPosition = FToolMenuInsert();
-					FToolMenuEntry DiffEntry = FToolMenuEntry::InitComboButton(
+					FToolMenuEntry& DiffEntry = InSection.AddEntry(FToolMenuEntry::InitComboButton(
 						"Diff",
 						FUIAction(),
 						FOnGetContent::CreateStatic(&FBlueprintEditorToolbar::MakeDiffMenu, Context),
 						LOCTEXT("Diff", "Diff"),
 						LOCTEXT("BlueprintEditorDiffToolTip", "Diff against previous revisions"),
 						FSlateIcon(FAppStyle::Get().GetStyleSetName(), "BlueprintDiff.ToolbarIcon")
-					);
+					));
 					DiffEntry.StyleNameOverride = "CalloutToolbar";
-					InSection.AddEntry(DiffEntry);
 				}
 			}
 		}));
