Fix potential world leak because of REINST classes not being properly handled on ALT

diff --git a/Engine/Source/Editor/Kismet/Private/BlueprintCompilationManager.cpp b/Engine/Source/Editor/Kismet/Private/BlueprintCompilationManager.cpp
--- a/Engine/Source/Editor/Kismet/Private/BlueprintCompilationManager.cpp
+++ b/Engine/Source/Editor/Kismet/Private/BlueprintCompilationManager.cpp
@@ -1814,7 +1814,9 @@ void FBlueprintCompilationManagerImpl::FlushReinstancingQueueImpl(bool bFindAndR
 		Options.bReplaceReferencesToOldCDOs = bFindAndReplaceCDOReferences;
 		FBlueprintCompileReinstancer::BatchReplaceInstancesOfClass(ClassesToReinstance, Options);
 
-		if (IsAsyncLoading())
+		// Special case when we run on ALT, we want to cleanup all classes flagged for reinstanciation right away.
+		const bool bIsInActualAsyncLoadingThread = IsInAsyncLoadingThread() && !IsInGameThread();
+		if (IsAsyncLoading() && IsAsyncLoadingMultithreaded() && !bIsInActualAsyncLoadingThread)
 		{
 			// While async loading we only remove classes that have no instances being
 			// async loaded. Those instances will need to be reinstanced once they finish
