Blueprint Editor - When editing the shared defaults of blueprints, ensure that the shared classname appears and not {ClassName}.

diff --git a/Engine/Source/Editor/Kismet/Private/BlueprintEditor.cpp b/Engine/Source/Editor/Kismet/Private/BlueprintEditor.cpp
--- a/Engine/Source/Editor/Kismet/Private/BlueprintEditor.cpp
+++ b/Engine/Source/Editor/Kismet/Private/BlueprintEditor.cpp
@@ -9604,7 +9604,7 @@ FText FBlueprintEditor::GetToolkitToolTipText() const
 
 	FFormatNamedArguments Args;
 	Args.Add( TEXT("NumberOfObjects"), EditingObjs.Num() );
-	Args.Add( TEXT("ObjectName"), FText::FromString( SharedParentClass->GetName() ) );
+	Args.Add( TEXT("ClassName"), FText::FromString( SharedParentClass->GetName() ) );
 	return FText::Format( NSLOCTEXT("KismetEditor", "ToolkitTitle_UniqueLayerName", "{NumberOfObjects} {ClassName} - Class Defaults"), Args );
 }
 
