Improve Blueprint editor startup time by deferring palette/favorites list construction to when the Palette tab is first spawned.

diff --git a/Engine/Source/Editor/Kismet/Private/BlueprintEditor.cpp b/Engine/Source/Editor/Kismet/Private/BlueprintEditor.cpp
--- a/Engine/Source/Editor/Kismet/Private/BlueprintEditor.cpp
+++ b/Engine/Source/Editor/Kismet/Private/BlueprintEditor.cpp
@@ -1979,10 +1979,16 @@ void FBlueprintEditor::CommonInitialization(const TArray<UBlueprint*>& InitBluep
 
 	if (InitBlueprints.Num() == 1)
 	{
-		// Load blueprint libraries
-		if (!bShouldOpenInDefaultsMode && ShouldLoadBPLibrariesFromAssetRegistry())
+		if (!bShouldOpenInDefaultsMode)
 		{
-			LoadLibrariesFromAssetRegistry();
+			// Load blueprint libraries
+			if (ShouldLoadBPLibrariesFromAssetRegistry())
+			{
+				LoadLibrariesFromAssetRegistry();
+			}
+
+			// Init the action DB for the context menu/palette if not already constructed
+			FBlueprintActionDatabase::Get();
 		}
 
 		FLoadObjectsFromAssetRegistryHelper::Load<UUserDefinedEnum>(UserDefinedEnumerators);
@@ -2880,10 +2886,6 @@ void FBlueprintEditor::CreateDefaultTabContents(const TArray<UBlueprint*>& InBlu
 
 	if (InBlueprint)
 	{
-		this->Palette = 
-			SNew(SBlueprintPalette, SharedThis(this))
-				.IsEnabled(this, &FBlueprintEditor::IsFocusedGraphEditable);
-
 		this->BookmarksWidget =
 			SNew(SBlueprintBookmarks)
 				.EditorContext(SharedThis(this));
@@ -3677,6 +3679,19 @@ void FBlueprintEditor::SetKeyboardFocus()
 	FSlateApplication::Get().SetKeyboardFocus(GetMyBlueprintWidget());
 }
 
+TSharedRef<SBlueprintPalette> FBlueprintEditor::GetPalette()
+{
+	// Note: construction is deferred until first access; in large-scale projects this can be an expensive widget to construct during editor
+	// initialization logic. It's an unnecessary cost if the tab it's hosted in is closed and/or unavailable (e.g. as in the default layout).
+	if (!Palette.IsValid())
+	{
+		SAssignNew(Palette, SBlueprintPalette, SharedThis(this))
+			.IsEnabled(this, &FBlueprintEditor::IsFocusedGraphEditable);
+	}
+
+	return Palette.ToSharedRef();
+}
+
 bool FBlueprintEditor::TransactionObjectAffectsBlueprint(UObject* InTransactedObject)
 {
 	check(InTransactedObject);
diff --git a/Engine/Source/Editor/Kismet/Private/SBlueprintPalette.cpp b/Engine/Source/Editor/Kismet/Private/SBlueprintPalette.cpp
--- a/Engine/Source/Editor/Kismet/Private/SBlueprintPalette.cpp
+++ b/Engine/Source/Editor/Kismet/Private/SBlueprintPalette.cpp
@@ -9,7 +9,6 @@
 #include "AnimationStateMachineGraph.h"
 #include "AnimationStateMachineSchema.h"
 #include "AnimationTransitionGraph.h"
-#include "AssetDiscoveryIndicator.h"
 #include "AssetRegistry/AssetData.h"
 #include "AssetRegistry/AssetRegistryModule.h"
 #include "AssetRegistry/IAssetRegistry.h"
@@ -38,7 +37,6 @@
 #include "EdGraphNode_Comment.h"
 #include "EdGraphSchema_K2.h"
 #include "EdGraphSchema_K2_Actions.h"
-#include "EditorWidgetsModule.h"
 #include "Engine/Blueprint.h"
 #include "Engine/MemberReference.h"
 #include "Engine/SCS_Node.h"
@@ -67,6 +65,7 @@
 #include "Misc/FileHelper.h"
 #include "Misc/PackageName.h"
 #include "Misc/Paths.h"
+#include "Misc/ScopedSlowTask.h"
 #include "Modules/ModuleManager.h"
 #include "SBlueprintFavoritesPalette.h"
 #include "SBlueprintLibraryPalette.h"
@@ -2183,9 +2182,11 @@ TSharedPtr<SToolTip> SBlueprintPaletteItem::ConstructToolTipWidget() const
 //------------------------------------------------------------------------------
 void SBlueprintPalette::Construct(const FArguments& InArgs, TWeakPtr<FBlueprintEditor> InBlueprintEditor)
 {
-	// Create the asset discovery indicator
-	FEditorWidgetsModule& EditorWidgetsModule = FModuleManager::LoadModuleChecked<FEditorWidgetsModule>("EditorWidgets");
-	TSharedRef<SWidget> AssetDiscoveryIndicator = EditorWidgetsModule.CreateAssetDiscoveryIndicator(EAssetDiscoveryIndicatorScaleMode::Scale_Vertical);
+	const float NumProgressFrames = 2.0f;
+	const float SecondsToWaitBeforeShowingProgressDialog = 0.25f;
+
+	FScopedSlowTask SlowTask(NumProgressFrames, LOCTEXT("ConstructingPaletteTabContent", "Initializing Palette..."));
+	SlowTask.MakeDialogDelayed(SecondsToWaitBeforeShowingProgressDialog);
 
 	float FavoritesHeightRatio = 0.33f;
 	GConfig->GetFloat(*BlueprintPalette::ConfigSection, *BlueprintPalette::FavoritesHeightConfigKey, FavoritesHeightRatio, GEditorPerProjectIni);
@@ -2195,16 +2196,27 @@ void SBlueprintPalette::Construct(const FArguments& InArgs, TWeakPtr<FBlueprintE
 	bool bUseLegacyLayout = false;
 	GConfig->GetBool(*BlueprintPalette::ConfigSection, TEXT("bUseLegacyLayout"), bUseLegacyLayout, GEditorIni);
 
+	SlowTask.EnterProgressFrame();
+	TSharedRef<SWidget> FavoritesContent = SNew(SBlueprintFavoritesPalette, InBlueprintEditor);
+
+	SlowTask.EnterProgressFrame();
+	TSharedRef<SWidget> LibraryContent = SNew(SBlueprintLibraryPalette, InBlueprintEditor)
+		.UseLegacyLayout(bUseLegacyLayout);
+
 	if (bUseLegacyLayout)
 	{
+		LibraryWrapper = LibraryContent;
+
 		this->ChildSlot
 		[
-			SAssignNew(LibraryWrapper, SBlueprintLibraryPalette, InBlueprintEditor)
-				.UseLegacyLayout(bUseLegacyLayout)
+			LibraryContent
 		];
 	}
 	else 
 	{
+		LibraryContent->AddMetadata<FTagMetaData>(MakeShared<FTagMetaData>(TEXT("BlueprintPaletteLibrary")));
+		FavoritesContent->AddMetadata<FTagMetaData>(MakeShared<FTagMetaData>(TEXT("BlueprintPaletteFavorites")));
+
 		this->ChildSlot
 		[
 			SAssignNew(PaletteSplitter, SSplitter)
@@ -2215,15 +2227,13 @@ void SBlueprintPalette::Construct(const FArguments& InArgs, TWeakPtr<FBlueprintE
 			+ SSplitter::Slot()
 			.Value(FavoritesHeightRatio)
 			[
-				SNew(SBlueprintFavoritesPalette, InBlueprintEditor)
-				.AddMetaData<FTagMetaData>(FTagMetaData(TEXT("BlueprintPaletteFavorites")))
+				FavoritesContent
 			]
 
 			+ SSplitter::Slot()
 			.Value(LibraryHeightRatio)
 			[
-				SNew(SBlueprintLibraryPalette, InBlueprintEditor)
-				.AddMetaData<FTagMetaData>(FTagMetaData(TEXT("BlueprintPaletteLibrary")))
+				LibraryContent
 			]
 		];
 	}	
diff --git a/Engine/Source/Editor/Kismet/Public/BlueprintEditor.h b/Engine/Source/Editor/Kismet/Public/BlueprintEditor.h
--- a/Engine/Source/Editor/Kismet/Public/BlueprintEditor.h
+++ b/Engine/Source/Editor/Kismet/Public/BlueprintEditor.h
@@ -329,7 +329,7 @@ public:
 	/** Getters for the various Kismet2 widgets */
 	TSharedRef<SKismetInspector> GetInspector() const { return Inspector.ToSharedRef(); }
 	TSharedRef<SKismetInspector> GetDefaultEditor() const { return DefaultEditor.ToSharedRef(); }
-	TSharedRef<SBlueprintPalette> GetPalette() const { return Palette.ToSharedRef(); }
+	TSharedRef<SBlueprintPalette> GetPalette();
 	TSharedRef<SBlueprintBookmarks> GetBookmarksWidget() const { return BookmarksWidget.ToSharedRef(); }
 	TSharedRef<SWidget> GetCompilerResults() const { return CompilerResults.ToSharedRef(); }
 	TSharedRef<SFindInBlueprints> GetFindResults() const { return FindResults.ToSharedRef(); }
