Fix the size of the combo boxes in the BP debugging toolbar, make Select&Frame icon toggle enabled state rather than visibility state, and improve the button's tooltip.

diff --git a/Engine/Plugins/Experimental/PCG/Source/PCGEditor/Private/SPCGEditorGraphDebugObjectWidget.cpp b/Engine/Plugins/Experimental/PCG/Source/PCGEditor/Private/SPCGEditorGraphDebugObjectWidget.cpp
--- a/Engine/Plugins/Experimental/PCG/Source/PCGEditor/Private/SPCGEditorGraphDebugObjectWidget.cpp
+++ b/Engine/Plugins/Experimental/PCG/Source/PCGEditor/Private/SPCGEditorGraphDebugObjectWidget.cpp
@@ -31,7 +31,7 @@ void SPCGEditorGraphDebugObjectWidget::Construct(const FArguments& InArgs, TShar
 
 	const TSharedRef<SWidget> BrowseButton = PropertyCustomizationHelpers::MakeBrowseButton(
 		FSimpleDelegate::CreateSP(this, &SPCGEditorGraphDebugObjectWidget::SelectedDebugObject_OnClicked),
-		LOCTEXT("DebugSelectActor", "Select this Actor in level"),
+		LOCTEXT("DebugSelectActor", "Select and frame the debug actor in the Level Editor."),
 		TAttribute<bool>::Create(TAttribute<bool>::FGetter::CreateSP(this, &SPCGEditorGraphDebugObjectWidget::IsSelectDebugObjectButtonEnabled))
 	);
 	
diff --git a/Engine/Source/Editor/Kismet/Private/SBlueprintEditorSelectedDebugObjectWidget.cpp b/Engine/Source/Editor/Kismet/Private/SBlueprintEditorSelectedDebugObjectWidget.cpp
--- a/Engine/Source/Editor/Kismet/Private/SBlueprintEditorSelectedDebugObjectWidget.cpp
+++ b/Engine/Source/Editor/Kismet/Private/SBlueprintEditorSelectedDebugObjectWidget.cpp
@@ -74,7 +74,8 @@ void SBlueprintEditorSelectedDebugObjectWidget::Construct(const FArguments& InAr
 		.InitiallySelectedItem(GetDebugWorldName())
 		.Visibility(this, &SBlueprintEditorSelectedDebugObjectWidget::IsDebugWorldComboVisible)
 		.OnComboBoxOpening(this, &SBlueprintEditorSelectedDebugObjectWidget::GenerateDebugWorldNames, true)
-		.OnSelectionChanged(this, &SBlueprintEditorSelectedDebugObjectWidget::DebugWorldSelectionChanged);
+		.OnSelectionChanged(this, &SBlueprintEditorSelectedDebugObjectWidget::DebugWorldSelectionChanged)
+		.ContentPadding(FMargin(0.f, 4.f));
 
 	DebugObjectsComboBox = SNew(SComboBox<TSharedPtr<FBlueprintDebugObjectInstance>>)
 		.ToolTip(IDocumentation::Get()->CreateToolTip(
@@ -87,6 +88,7 @@ void SBlueprintEditorSelectedDebugObjectWidget::Construct(const FArguments& InAr
 		.OnComboBoxOpening(this, &SBlueprintEditorSelectedDebugObjectWidget::GenerateDebugObjectInstances, true)
 		.OnSelectionChanged(this, &SBlueprintEditorSelectedDebugObjectWidget::DebugObjectSelectionChanged)
 		.OnGenerateWidget(this, &SBlueprintEditorSelectedDebugObjectWidget::CreateDebugObjectItemWidget)
+		.ContentPadding(FMargin(0.f, 4.f))
 		.AddMetaData<FTagMetaData>(TEXT("SelectDebugObjectCobmo"))
 		[
 			SNew(STextBlock)
@@ -146,9 +148,10 @@ TSharedRef<SWidget> SBlueprintEditorSelectedDebugObjectWidget::OnGetActiveDetail
 {
 	const TSharedRef<SWidget> BrowseButton = PropertyCustomizationHelpers::MakeBrowseButton(
 		FSimpleDelegate::CreateSP(this, &SBlueprintEditorSelectedDebugObjectWidget::SelectedDebugObject_OnClicked),
-		LOCTEXT("DebugSelectActor", "Select this Actor in level")
+		LOCTEXT("DebugSelectActor", "Select and frame the debug actor in the Level Editor."),
+		TAttribute<bool>::Create(TAttribute<bool>::FGetter::CreateSP(this, &SBlueprintEditorSelectedDebugObjectWidget::IsDebugObjectSelected))
 	);
-	BrowseButton->SetVisibility(TAttribute<EVisibility>::Create(TAttribute<EVisibility>::FGetter::CreateSP(this, &SBlueprintEditorSelectedDebugObjectWidget::IsSelectDebugObjectButtonVisible)));
+
 
 	TSharedRef<SWidget> DebugObjectSelectionWidget =
 		SNew(SHorizontalBox)
@@ -589,17 +592,17 @@ void SBlueprintEditorSelectedDebugObjectWidget::DebugObjectSelectionChanged(TSha
 	}
 }
 
-EVisibility SBlueprintEditorSelectedDebugObjectWidget::IsSelectDebugObjectButtonVisible() const
+bool SBlueprintEditorSelectedDebugObjectWidget::IsDebugObjectSelected() const
 {
 	check(GetBlueprintObj());
 	if (UObject* DebugObj = GetBlueprintObj()->GetObjectBeingDebugged())
 	{
 		if (AActor* Actor = Cast<AActor>(DebugObj))
 		{
-			return EVisibility::Visible;
+			return true;
 		}
 	}
-	return EVisibility::Collapsed;
+	return false;
 }
 
 void SBlueprintEditorSelectedDebugObjectWidget::SelectedDebugObject_OnClicked()
diff --git a/Engine/Source/Editor/Kismet/Private/SBlueprintEditorSelectedDebugObjectWidget.h b/Engine/Source/Editor/Kismet/Private/SBlueprintEditorSelectedDebugObjectWidget.h
--- a/Engine/Source/Editor/Kismet/Private/SBlueprintEditorSelectedDebugObjectWidget.h
+++ b/Engine/Source/Editor/Kismet/Private/SBlueprintEditorSelectedDebugObjectWidget.h
@@ -112,7 +112,7 @@ private:
 	void SelectedDebugObject_OnClicked();
 
 	/** Returns true if a debug actor is currently selected */
-	EVisibility IsSelectDebugObjectButtonVisible() const;
+	bool IsDebugObjectSelected() const;
 
 	EVisibility IsDebugWorldComboVisible() const;
 
