[BugFix] Blueprint Interface Implemented and Inherited Interface naming is reversed

diff --git a/Engine/Source/Editor/Kismet/Private/BlueprintDetailsCustomization.cpp b/Engine/Source/Editor/Kismet/Private/BlueprintDetailsCustomization.cpp
--- a/Engine/Source/Editor/Kismet/Private/BlueprintDetailsCustomization.cpp
+++ b/Engine/Source/Editor/Kismet/Private/BlueprintDetailsCustomization.cpp
@@ -6048,8 +6048,8 @@ FBlueprintInterfaceLayout::FBlueprintInterfaceLayout(TWeakPtr<class FBlueprintGl
 	, InterfacesProperty(InInterfacesProperty)
 {
 	DisplayOptions.TitleText = InterfacesProperty ?
-		LOCTEXT("BlueprintInheritedInterfaceTitle", "Inherited Interfaces") :
-		LOCTEXT("BlueprintImplementedInterfaceTitle", "Implemented Interfaces");
+		LOCTEXT("BlueprintImplementedInterfaceTitle", "Implemented Interfaces") :
+		LOCTEXT("BlueprintInheritedInterfaceTitle", "Inherited Interfaces");
 	DisplayOptions.NoItemsLabelText = LOCTEXT("NoBlueprintInterface", "No Interfaces");
 	DisplayOptions.BrowseButtonToolTipText = LOCTEXT("BlueprintInterfaceBrowseTooltip", "Opens this interface");
 }
