Include skeleton class of archetype outers for archetype reference replacement, specifically to accomdate SkeletonClass ComponentTempltae members

diff --git a/Engine/Source/Editor/Kismet/Private/BlueprintCompilationManager.cpp b/Engine/Source/Editor/Kismet/Private/BlueprintCompilationManager.cpp
--- a/Engine/Source/Editor/Kismet/Private/BlueprintCompilationManager.cpp
+++ b/Engine/Source/Editor/Kismet/Private/BlueprintCompilationManager.cpp
@@ -2370,9 +2370,10 @@ void FBlueprintCompilationManagerImpl::ReinstanceBatch(TArray<FReinstancingJob>&
 					while(Iter)
 					{
 						UBlueprintGeneratedClass* IterAsBPGC = Cast<UBlueprintGeneratedClass>(Iter);
+						UBlueprint* IterAsBP = Cast<UBlueprint>(Iter);
 						if(Iter->HasAnyFlags(RF_ClassDefaultObject)
 							|| (IterAsBPGC && !IterAsBPGC->HasAnyClassFlags(CLASS_NewerVersionExists))
-							|| Cast<UBlueprint>(Iter) )
+							|| IterAsBP)
 						{
 							ArchetypeReferencers.Add(Iter);
 
@@ -2382,6 +2383,14 @@ void FBlueprintCompilationManagerImpl::ReinstanceBatch(TArray<FReinstancingJob>&
 							if(IterAsBPGC)
 							{
 								ArchetypeReferencers.Add(IterAsBPGC->ClassGeneratedBy);
+								IterAsBP = Cast<UBlueprint>(IterAsBPGC->ClassGeneratedBy);
+							}
+							if (IterAsBP)
+							{
+								if (IterAsBP->SkeletonGeneratedClass)
+								{
+									ArchetypeReferencers.Add(IterAsBP->SkeletonGeneratedClass);
+								}
 							}
 
 							// this handles nested subobjects:
