Improve error message when a blueprint attempts to compile while other blueprints are compiling

diff --git a/Engine/Source/Editor/Kismet/Private/BlueprintCompilationManager.cpp b/Engine/Source/Editor/Kismet/Private/BlueprintCompilationManager.cpp
--- a/Engine/Source/Editor/Kismet/Private/BlueprintCompilationManager.cpp
+++ b/Engine/Source/Editor/Kismet/Private/BlueprintCompilationManager.cpp
@@ -129,6 +129,7 @@ struct FBlueprintCompilationManagerImpl : public FGCObject
 	bool HasBlueprintsToCompile() const;
 	bool IsGeneratedClassLayoutReady() const;
 	void GetDefaultValue(const UClass* ForClass, const FProperty* Property, FString& OutDefaultValueAsString) const;
+	void VerifyNoQueuedRequests(const TArray<FCompilerData>& CurrentlyCompilingBPs);
 
 	static void ReparentHierarchies(const TMap<UClass*, UClass*>& OldClassToNewClass, EReparentClassOptions Options);
 	static void BuildDSOMap(UObject* OldObject, UObject* NewObject, TMap<UObject*, UObject*>& OutOldToNewDSO);
@@ -1733,7 +1734,7 @@ void FBlueprintCompilationManagerImpl::FlushCompilationQueueImpl(bool bSuppressB
 	UE_LOG(LogBlueprint, Display, TEXT("Time Compiling: %f, Time Reinstancing: %f"),  GTimeCompiling, GTimeReinstancing);
 	//GTimeCompiling = 0.0;
 	//GTimeReinstancing = 0.0;
-	ensure(QueuedRequests.Num() == 0);
+	VerifyNoQueuedRequests(CurrentlyCompilingBPs);
 }
 
 void FBlueprintCompilationManagerImpl::ProcessExtensions(const TArray<FCompilerData>& InCurrentlyCompilingBPs)
@@ -1897,6 +1898,27 @@ void FBlueprintCompilationManagerImpl::GetDefaultValue(const UClass* ForClass, c
 	}
 }
 
+void FBlueprintCompilationManagerImpl::VerifyNoQueuedRequests(const TArray<FCompilerData>& CurrentlyCompilingBPs)
+{
+	if (QueuedRequests.Num() != 0)
+	{
+		FString QueuedBlueprints = TEXT("");
+		for (const FBPCompileRequestInternal& Request: QueuedRequests)
+		{
+			QueuedBlueprints += Request.UserData.BPToCompile->GetName() + TEXT(" ");
+		}
+		FString CompilingBlueprints = TEXT("");
+		for (const FCompilerData& CompilerData : CurrentlyCompilingBPs)
+		{
+			CompilingBlueprints += CompilerData.BP->GetName() + TEXT(" ");
+		}
+
+		ensureMsgf(false, 
+			TEXT("Blueprints requested compilation while compiling other blueprints: %s\nWhile Compiling: %s"),
+			*QueuedBlueprints, *CompilingBlueprints);
+	}
+}
+
 void FBlueprintCompilationManagerImpl::ReparentHierarchies(const TMap<UClass*, UClass*>& OldToNewClasses, EReparentClassOptions Options)
 {
 	const bool bReplaceReferencesToOldClasses = (Options & EReparentClassOptions::ReplaceReferencesToOldClasses) != EReparentClassOptions::None;
