Fix a regression in which duplicate Add Component actions are shown in the Blueprint graph right-click context menu when "Context Sensitive" is enabled.

diff --git a/Engine/Source/Editor/Kismet/Private/BlueprintActionMenuUtils.cpp b/Engine/Source/Editor/Kismet/Private/BlueprintActionMenuUtils.cpp
--- a/Engine/Source/Editor/Kismet/Private/BlueprintActionMenuUtils.cpp
+++ b/Engine/Source/Editor/Kismet/Private/BlueprintActionMenuUtils.cpp
@@ -672,7 +672,7 @@ void FBlueprintActionMenuUtils::MakeContextMenu(FBlueprintActionContext const& C
 	FBlueprintActionFilter AddComponentFilter;
 	AddComponentFilter.Context = MainMenuFilter.Context;
 	AddComponentFilter.PermittedNodeTypes.Add(UK2Node_AddComponent::StaticClass());
-	LevelActorsFilter.AddRejectionTest(MAKE_ACTION_FILTER_REJECTION_TEST(IsUnBoundSpawner)->WithFlags(EActionFilterTestFlags::CacheResults));
+	AddComponentFilter.AddRejectionTest(MAKE_ACTION_FILTER_REJECTION_TEST(IsUnBoundSpawner)->WithFlags(EActionFilterTestFlags::CacheResults));
 
 
 	FContentBrowserModule& ContentBrowserModule = FModuleManager::LoadModuleChecked<FContentBrowserModule>("ContentBrowser");
