Added Browse to Type to the context menu of a variable in the Blueprint Editor details menu, available when the type of a variable is a blueprint class.

diff --git a/Engine/Source/Editor/Kismet/Private/BlueprintDetailsCustomization.cpp b/Engine/Source/Editor/Kismet/Private/BlueprintDetailsCustomization.cpp
--- a/Engine/Source/Editor/Kismet/Private/BlueprintDetailsCustomization.cpp
+++ b/Engine/Source/Editor/Kismet/Private/BlueprintDetailsCustomization.cpp
@@ -2,6 +2,7 @@
 
 #include "BlueprintDetailsCustomization.h"
 
+#include "AssetRegistry/AssetData.h"
 #include "BlueprintEditor.h"
 #include "BlueprintEditorModule.h"
 #include "BlueprintEditorSettings.h"
@@ -26,6 +27,7 @@
 #include "EdGraphSchema_K2.h"
 #include "EdGraphSchema_K2_Actions.h"
 #include "EdMode.h"
+#include "Editor/EditorEngine.h"
 #include "Engine/BlueprintGeneratedClass.h"
 #include "Engine/Engine.h"
 #include "Engine/EngineBaseTypes.h"
@@ -444,7 +446,15 @@ void FBlueprintVarActionDetails::CustomizeDetails( IDetailLayoutBuilder& DetailL
 			.Font(DetailFontInfo)
 			.ToolTip(VarTypeTooltip)
 			.CustomFilters(CustomPinTypeFilters)
-		];
+		]
+		.AddCustomContextMenuAction(FUIAction(
+			FExecuteAction::CreateRaw(this, &FBlueprintVarActionDetails::OnBrowseToVarType),
+			FCanExecuteAction::CreateRaw(this, &FBlueprintVarActionDetails::CanBrowseToVarType)
+			),
+			LOCTEXT("BrowseToType", "Browse to Type"),
+			LOCTEXT("BrowseToTypeToolTip", "Browse to this variable type in the Content Browser."),
+			FSlateIcon(FAppStyle::GetAppStyleSetName(), "Icons.BrowseContent")
+		);
 
 	TSharedPtr<SToolTip> ToolTipTooltip = IDocumentation::Get()->CreateToolTip(LOCTEXT("VarToolTipTooltip", "Extra information about this variable, shown when cursor is over it."), NULL, DocLink, TEXT("Description"));
 
@@ -2804,6 +2814,41 @@ EVisibility FBlueprintVarActionDetails::IsTooltipEditVisible() const
 	return EVisibility::Collapsed;
 }
 
+void FBlueprintVarActionDetails::OnBrowseToVarType() const
+{
+	FEdGraphPinType PinType = OnGetVarType();
+	if (const UObject* Object = PinType.PinSubCategoryObject.Get())
+	{
+		if (Object->IsAsset())
+		{
+			FAssetData AssetData(Object, false);
+			if (AssetData.IsValid())
+			{
+				TArray<FAssetData> AssetDataList = { AssetData };
+				GEditor->SyncBrowserToObjects(AssetDataList);
+			}
+		}
+	}
+}
+
+bool FBlueprintVarActionDetails::CanBrowseToVarType() const
+{
+	FEdGraphPinType PinType = OnGetVarType();
+	if (const UObject* Object = PinType.PinSubCategoryObject.Get())
+	{
+		if (Object->IsAsset())
+		{
+			FAssetData AssetData(Object, false);
+			if (AssetData.IsValid())
+			{
+				return true;
+			}
+		}
+	}
+
+	return false;
+}
+
 void FBlueprintVarActionDetails::OnFinishedChangingVariable(const FPropertyChangedEvent& InPropertyChangedEvent)
 {
 	if (InPropertyChangedEvent.GetNumObjectsBeingEdited() == 0)
diff --git a/Engine/Source/Editor/Kismet/Private/BlueprintDetailsCustomization.h b/Engine/Source/Editor/Kismet/Private/BlueprintDetailsCustomization.h
--- a/Engine/Source/Editor/Kismet/Private/BlueprintDetailsCustomization.h
+++ b/Engine/Source/Editor/Kismet/Private/BlueprintDetailsCustomization.h
@@ -178,6 +178,9 @@ private:
 	void OnVarTypeChanged(const FEdGraphPinType& NewPinType);
 	EVisibility IsTooltipEditVisible() const;
 
+	void OnBrowseToVarType() const;
+	bool CanBrowseToVarType() const;
+
 	/**
 	 * Callback when changing a variable property
 	 *
