Synchronize the scrolling between left and right details panels in Blueprint Diff

diff --git a/Engine/Source/Editor/Kismet/Private/DiffControl.cpp b/Engine/Source/Editor/Kismet/Private/DiffControl.cpp
--- a/Engine/Source/Editor/Kismet/Private/DiffControl.cpp
+++ b/Engine/Source/Editor/Kismet/Private/DiffControl.cpp
@@ -10,6 +10,34 @@
 
 #define LOCTEXT_NAMESPACE "SBlueprintDif"
 
+// well known algorithm for computing the Longest Common Subsequence table of two ordered lists
+template <typename RangeType, typename ComparePredicate>
+static TArray<TArray<int32>> CalculateLCSTable(const RangeType& Range1, const RangeType& Range2, ComparePredicate Comparison)
+{
+	TArray<TArray<int32>> LCS;
+	LCS.SetNum(Range1.Num() + 1);
+	for (int32 I = 0; I <= Range1.Num(); I++)
+	{
+		LCS[I].SetNum(Range2.Num() + 1);
+		if (I == 0)
+		{
+			continue;
+		}
+	
+		for (int32 J = 1; J <= Range2.Num(); J++)
+		{
+			if (Comparison(Range1[I - 1], Range2[J - 1]))
+			{
+				LCS[I][J] = LCS[I - 1][J - 1] + 1;
+			}
+			else
+			{
+				LCS[I][J] = FMath::Max(LCS[I - 1][J], LCS[I][J - 1]);
+			}
+		}
+	}
+	return LCS;
+}
 
 /////////////////////////////////////////////////////////////////////////////
 /// IDiffControl
@@ -185,17 +213,137 @@ void FSCSDiffControl::GenerateTreeEntries(TArray<TSharedPtr<FBlueprintDifference
 	));
 }
 
+FDetailsDiffControl::FDetailsDiffControl(const UObject* InOldObject, const UObject* InNewObject,
+	FOnDiffEntryFocused InSelectionCallback, bool bPopulateOutTreeEntries): SelectionCallback(InSelectionCallback)
+	                                                                        , OldDetails(InOldObject, FDetailsDiff::FOnDisplayedPropertiesChanged())
+	                                                                        , NewDetails(InNewObject, FDetailsDiff::FOnDisplayedPropertiesChanged())
+	                                                                        , bPopulateOutTreeEntries(bPopulateOutTreeEntries)
+{
+	OldDetails.DiffAgainst(NewDetails, DifferingProperties, true);
+
+	TSet<FPropertyPath> PropertyPaths;
+	Algo::Transform(DifferingProperties, PropertyPaths,
+	                [&InOldObject](const FSingleObjectDiffEntry& DiffEntry)
+	                {
+		                return DiffEntry.Identifier.ResolvePath(InOldObject);
+	                });
+
+	OldDetails.DetailsWidget()->UpdatePropertyAllowList(PropertyPaths);
+
+	PropertyPaths.Reset();
+	Algo::Transform(DifferingProperties, PropertyPaths,
+	                [&InNewObject](const FSingleObjectDiffEntry& DiffEntry)
+	                {
+		                return DiffEntry.Identifier.ResolvePath(InNewObject);
+	                });
+
+	NewDetails.DetailsWidget()->UpdatePropertyAllowList(PropertyPaths);
+		
+	// Sync the scrolling between the left and right panels
+	const TSharedRef<IDetailsView> OldDetailsView = OldDetails.DetailsWidget();
+	const TSharedRef<IDetailsView> NewDetailsView = NewDetails.DetailsWidget();
+	OldDetailsView->ScrollLockDetailsViews(NewDetailsView, GetLinkedScrollRateAttribute(OldDetailsView, NewDetailsView));
+
+	// Make it so that these details panels don't overwrite the saved expansion states since they're just temp copies
+	OldDetailsView->SetCategoryExpansionSaving(false);
+	OldDetailsView->SetCategoryExpansionSaving(false);
+}
+
+void FDetailsDiffControl::GenerateTreeEntries(TArray<TSharedPtr<FBlueprintDifferenceTreeEntry>>& OutTreeEntries, TArray<TSharedPtr<FBlueprintDifferenceTreeEntry>>& OutRealDifferences)
+{
+	for (const FSingleObjectDiffEntry& Difference : DifferingProperties)
+	{
+			
+		TSharedPtr<FBlueprintDifferenceTreeEntry> Entry = MakeShared<FBlueprintDifferenceTreeEntry>(
+			FOnDiffEntryFocused::CreateSP(TSharedFromThis<FDetailsDiffControl>::AsShared(), &FDetailsDiffControl::OnSelectDiffEntry, Difference.Identifier),
+			FGenerateDiffEntryWidget::CreateStatic(&GenerateObjectDiffWidget, Difference, RightRevision));
+		Children.Push(Entry);
+		OutRealDifferences.Push(Entry);
+		if (bPopulateOutTreeEntries)
+		{
+			OutTreeEntries.Push(Entry);
+		}
+	}
+}
+
+void FDetailsDiffControl::OnSelectDiffEntry(FPropertySoftPath PropertyName)
+{
+	SelectionCallback.ExecuteIfBound();
+	OldDetails.HighlightProperty(PropertyName);
+	NewDetails.HighlightProperty(PropertyName);
+}
+
+TAttribute<FInvertiblePiecewiseLinearFunction> FDetailsDiffControl::GetLinkedScrollRateAttribute(const TSharedRef<IDetailsView>& OldDetailsView, const TSharedRef<IDetailsView>& NewDetailsView)
+{
+	return TAttribute<FInvertiblePiecewiseLinearFunction>::CreateRaw(this, &FDetailsDiffControl::GetLinkedScrollRate, OldDetailsView, NewDetailsView);
+}
+
+FInvertiblePiecewiseLinearFunction FDetailsDiffControl::GetLinkedScrollRate(TSharedRef<IDetailsView> OldDetailsView, TSharedRef<IDetailsView> NewDetailsView) const
+{
+	TArray<TPair<int32, FPropertyPath>> OldProperties = OldDetailsView->GetPropertyRowNumbers();
+	TArray<TPair<int32, FPropertyPath>> NewProperties = NewDetailsView->GetPropertyRowNumbers();
+
+	// use caching to avoid O(n^2) LCS calculation every frame
+	if(LinkedScrollRateCache.OldProperties != OldProperties || LinkedScrollRateCache.NewProperties != NewProperties)
+	{
+		const auto PropertyPathsEqual = [](const TPair<int32, FPropertyPath>& A, const TPair<int32, FPropertyPath>& B)
+		{
+			return A.Value == B.Value;
+		};
+		
+		// Find the Longest Common Subsequence between both sets of properties
+		const TArray<TArray<int32>> LCS = CalculateLCSTable(OldProperties, NewProperties, PropertyPathsEqual);
+	
+		// Using the LCS Table, we can determine which lines should match one another while scrolling. For example, an element
+		// of FixedPoints may be {12.f, 5.f} meaning that line 12 of the left panel is the same as line 5 of the right panel
+		TArray<FVector2f> FixedPoints;
+
+		int32 I = OldProperties.Num();
+		int32 J = NewProperties.Num();
+		while (I > 0 && J > 0)
+		{
+			if (OldProperties[I - 1].Value == NewProperties[J - 1].Value)
+			{
+				const int32 OldLineNum = OldProperties[I - 1].Key;
+				const int32 NewLineNum = NewProperties[J - 1].Key;
+				// add two endpoints for every matched property because the elements have a width of 1 in the panels
+				FixedPoints.Add(FVector2f(OldLineNum + 1.f, NewLineNum + 1.f));
+				FixedPoints.Add(FVector2f((float)OldLineNum, (float)NewLineNum));
+				--I;
+				--J;
+			}
+			else if (LCS[I - 1][J] <= LCS[I][J	 - 1])
+			{
+				--J;
+			}
+			else
+			{
+				--I;
+			}
+		}
+		FixedPoints.Add({0.f, 0.f});
+		Algo::Reverse(FixedPoints);
+        	
+		LinkedScrollRateCache.OldProperties = OldProperties;
+		LinkedScrollRateCache.NewProperties = NewProperties;
+		LinkedScrollRateCache.ScrollRate = FInvertiblePiecewiseLinearFunction (FixedPoints);
+	}
+	
+	return LinkedScrollRateCache.ScrollRate;
+}
+
 
 /////////////////////////////////////////////////////////////////////////////
 /// FCDODiffControl
 
-FCDODiffControl::FCDODiffControl(const UObject* InOldObject, const UObject* InNewObject, FOnDiffEntryFocused InSelectionCallback): TDetailsDiffControl(InOldObject, InNewObject, InSelectionCallback)
+FCDODiffControl::FCDODiffControl(const UObject* InOldObject, const UObject* InNewObject, FOnDiffEntryFocused InSelectionCallback)
+	: FDetailsDiffControl(InOldObject, InNewObject, InSelectionCallback, false)
 {
 }
 
 void FCDODiffControl::GenerateTreeEntries(TArray<TSharedPtr<FBlueprintDifferenceTreeEntry>>& OutTreeEntries, TArray<TSharedPtr<FBlueprintDifferenceTreeEntry>>& OutRealDifferences)
 {
-	TDetailsDiffControl::GenerateTreeEntries(OutTreeEntries, OutRealDifferences);
+	FDetailsDiffControl::GenerateTreeEntries(OutTreeEntries, OutRealDifferences);
 
 	const bool bHasDifferences = Children.Num() != 0;
 	if (!bHasDifferences)
@@ -218,13 +366,14 @@ void FCDODiffControl::GenerateTreeEntries(TArray<TSharedPtr<FBlueprintDifference
 /////////////////////////////////////////////////////////////////////////////
 /// FClassSettingsDiffControl
 
-FClassSettingsDiffControl::FClassSettingsDiffControl(const UObject* InOldObject, const UObject* InNewObject, FOnDiffEntryFocused InSelectionCallback): TDetailsDiffControl(InOldObject, InNewObject, InSelectionCallback)
+FClassSettingsDiffControl::FClassSettingsDiffControl(const UObject* InOldObject, const UObject* InNewObject, FOnDiffEntryFocused InSelectionCallback)
+	: FDetailsDiffControl(InOldObject, InNewObject, InSelectionCallback, false)
 {
 }
 
 void FClassSettingsDiffControl::GenerateTreeEntries(TArray<TSharedPtr<FBlueprintDifferenceTreeEntry>>& OutTreeEntries, TArray<TSharedPtr<FBlueprintDifferenceTreeEntry>>& OutRealDifferences)
 {
-	TDetailsDiffControl::GenerateTreeEntries(OutTreeEntries, OutRealDifferences);
+	FDetailsDiffControl::GenerateTreeEntries(OutTreeEntries, OutRealDifferences);
 
 	const bool bHasDifferences = Children.Num() != 0;
 	if (!bHasDifferences)
diff --git a/Engine/Source/Editor/Kismet/Private/DiffControl.h b/Engine/Source/Editor/Kismet/Private/DiffControl.h
--- a/Engine/Source/Editor/Kismet/Private/DiffControl.h
+++ b/Engine/Source/Editor/Kismet/Private/DiffControl.h
@@ -9,12 +9,14 @@
 #include "IDetailsView.h"
 #include "Algo/Transform.h"
 #include "Widgets/SWidget.h"
+#include "Framework/FInvertiblePiecewiseLinearFunction.h"
 
 struct FDiffResultItem;
 class SBlueprintDiff;
 class UEdGraph;
 struct FEdGraphEditAction;
 
+
 /** Interface responsible for generating FBlueprintDifferenceTreeEntry's for visual diff tools */
 class KISMET_API IDiffControl
 {
@@ -67,63 +69,23 @@ private:
 };
 
 /** Generic wrapper around a details view template parameter determines whether TreeEntries are populated */
-template<bool bPopulateOutTreeEntries>
-class KISMET_API TDetailsDiffControl : public TSharedFromThis<TDetailsDiffControl<bPopulateOutTreeEntries>>, public IDiffControl
+class KISMET_API FDetailsDiffControl : public TSharedFromThis<FDetailsDiffControl>, public IDiffControl
 {
 public:
-	TDetailsDiffControl(const UObject* InOldObject, const UObject* InNewObject, FOnDiffEntryFocused InSelectionCallback)
-		: SelectionCallback(InSelectionCallback)
-		, OldDetails(InOldObject, FDetailsDiff::FOnDisplayedPropertiesChanged())
-		, NewDetails(InNewObject, FDetailsDiff::FOnDisplayedPropertiesChanged())
-	{
-		OldDetails.DiffAgainst(NewDetails, DifferingProperties, true);
-
-		TSet<FPropertyPath> PropertyPaths;
-		Algo::Transform(DifferingProperties, PropertyPaths,
-			[&InOldObject](const FSingleObjectDiffEntry& DiffEntry)
-			{
-				return DiffEntry.Identifier.ResolvePath(InOldObject);
-			});
-
-		OldDetails.DetailsWidget()->UpdatePropertyAllowList(PropertyPaths);
+	FDetailsDiffControl(const UObject* InOldObject, const UObject* InNewObject, FOnDiffEntryFocused InSelectionCallback, bool bPopulateOutTreeEntries);
 
-		PropertyPaths.Reset();
-		Algo::Transform(DifferingProperties, PropertyPaths,
-			[&InNewObject](const FSingleObjectDiffEntry& DiffEntry)
-			{
-				return DiffEntry.Identifier.ResolvePath(InNewObject);
-			});
-
-		NewDetails.DetailsWidget()->UpdatePropertyAllowList(PropertyPaths);
-	}
-
-	virtual void GenerateTreeEntries(TArray< TSharedPtr<FBlueprintDifferenceTreeEntry> >& OutTreeEntries, TArray< TSharedPtr<FBlueprintDifferenceTreeEntry> >& OutRealDifferences) override
-	{
-		for (const FSingleObjectDiffEntry& Difference : DifferingProperties)
-		{
-			
-			TSharedPtr<FBlueprintDifferenceTreeEntry> Entry = MakeShared<FBlueprintDifferenceTreeEntry>(
-				FOnDiffEntryFocused::CreateSP(TSharedFromThis<TDetailsDiffControl<bPopulateOutTreeEntries>>::AsShared(), &TDetailsDiffControl::OnSelectDiffEntry, Difference.Identifier),
-				FGenerateDiffEntryWidget::CreateStatic(&GenerateObjectDiffWidget, Difference, RightRevision));
-			Children.Push(Entry);
-			OutRealDifferences.Push(Entry);
-			if constexpr (bPopulateOutTreeEntries)
-			{
-				OutTreeEntries.Push(Entry);
-			}
-		}
-	}
+	virtual void GenerateTreeEntries(TArray< TSharedPtr<FBlueprintDifferenceTreeEntry> >& OutTreeEntries, TArray< TSharedPtr<FBlueprintDifferenceTreeEntry> >& OutRealDifferences) override;
 
 	TSharedRef<SWidget> OldDetailsWidget() { return OldDetails.DetailsWidget(); }
 	TSharedRef<SWidget> NewDetailsWidget() { return NewDetails.DetailsWidget(); }
 
 protected:
-	virtual void OnSelectDiffEntry(FPropertySoftPath PropertyName)
-	{
-		SelectionCallback.ExecuteIfBound();
-		OldDetails.HighlightProperty(PropertyName);
-		NewDetails.HighlightProperty(PropertyName);
-	}
+	virtual void OnSelectDiffEntry(FPropertySoftPath PropertyName);
+	TAttribute<FInvertiblePiecewiseLinearFunction> GetLinkedScrollRateAttribute(const TSharedRef<IDetailsView>& OldDetailsView, const TSharedRef<IDetailsView>& NewDetailsView);
+	
+
+	// helper function that analyzes two details views and determines the rate they should scroll relative to one another to be in sync
+	FInvertiblePiecewiseLinearFunction GetLinkedScrollRate(TSharedRef<IDetailsView> OldDetailsView, TSharedRef<IDetailsView> NewDetailsView) const;
 
 	FOnDiffEntryFocused SelectionCallback;
 	FDetailsDiff OldDetails;
@@ -131,12 +93,18 @@ protected:
 
 	TArray<FSingleObjectDiffEntry> DifferingProperties;
 	TArray< TSharedPtr<FBlueprintDifferenceTreeEntry> > Children;
-};
+	const bool bPopulateOutTreeEntries;
 
-using FDetailsDiffControl = TDetailsDiffControl<true>;
+	mutable struct
+	{
+		TArray<TPair<int32, FPropertyPath>> OldProperties;
+		TArray<TPair<int32, FPropertyPath>> NewProperties;
+		FInvertiblePiecewiseLinearFunction ScrollRate;
+	} LinkedScrollRateCache;
+};
 
 /** Override for CDO special case */
-class FCDODiffControl : public TDetailsDiffControl<false>
+class FCDODiffControl : public FDetailsDiffControl
 {
 public:
 	FCDODiffControl(const UObject* InOldObject, const UObject* InNewObject, FOnDiffEntryFocused InSelectionCallback);
@@ -145,7 +113,7 @@ public:
 };
 
 /** Override for class class settings */
-class FClassSettingsDiffControl : public TDetailsDiffControl<false>
+class FClassSettingsDiffControl : public FDetailsDiffControl
 {
 public:
 	FClassSettingsDiffControl(const UObject* InOldObject, const UObject* InNewObject, FOnDiffEntryFocused InSelectionCallback);
diff --git a/Engine/Source/Editor/Kismet/Private/SDetailsDiff.cpp b/Engine/Source/Editor/Kismet/Private/SDetailsDiff.cpp
--- a/Engine/Source/Editor/Kismet/Private/SDetailsDiff.cpp
+++ b/Engine/Source/Editor/Kismet/Private/SDetailsDiff.cpp
@@ -306,7 +306,7 @@ void SDetailsDiff::GenerateDifferencesList()
 
 SDetailsDiff::FDiffControl SDetailsDiff::GenerateDetailsPanel()
 {
-	TSharedPtr<FDetailsDiffControl> NewDiffControl = MakeShared<FDetailsDiffControl>(PanelOld.Object, PanelNew.Object, FOnDiffEntryFocused::CreateRaw(this, &SDetailsDiff::SetCurrentMode, DetailsMode));
+	TSharedPtr<FDetailsDiffControl> NewDiffControl = MakeShared<FDetailsDiffControl>(PanelOld.Object, PanelNew.Object, FOnDiffEntryFocused::CreateRaw(this, &SDetailsDiff::SetCurrentMode, DetailsMode), true);
 	NewDiffControl->GenerateTreeEntries(PrimaryDifferencesList, RealDifferences);
 
 	SDetailsDiff::FDiffControl Ret;
diff --git a/Engine/Source/Editor/PropertyEditor/Private/DetailCategoryBuilderImpl.cpp b/Engine/Source/Editor/PropertyEditor/Private/DetailCategoryBuilderImpl.cpp
--- a/Engine/Source/Editor/PropertyEditor/Private/DetailCategoryBuilderImpl.cpp
+++ b/Engine/Source/Editor/PropertyEditor/Private/DetailCategoryBuilderImpl.cpp
@@ -181,6 +181,7 @@ FDetailCategoryImpl::FDetailCategoryImpl(FName InCategoryName, TSharedRef<FDetai
 	, bFavoriteCategory(false)
 	, bShowOnlyChildren(false)
 	, bHasVisibleAdvanced(false)
+	, bExpanded(true)
 {
 	const UStruct* BaseStruct = InDetailLayout->GetRootNode()->GetBaseStructure();
 
@@ -817,6 +818,7 @@ void FDetailCategoryImpl::OnItemExpansionChanged(bool bIsExpanded, bool bShouldS
 		// Save the collapsed state of this section
 		GConfig->SetBool(TEXT("DetailCategories"), *CategoryPathName, bIsExpanded, GEditorPerProjectIni);
 	}
+	bExpanded = bIsExpanded;
 
 	OnExpansionChangedDelegate.ExecuteIfBound(bIsExpanded);
 }
@@ -837,7 +839,7 @@ bool FDetailCategoryImpl::ShouldBeExpanded() const
 	}
 	else
 	{
-		return !bShouldBeInitiallyCollapsed;
+		return bExpanded;
 	}
 }
 
diff --git a/Engine/Source/Editor/PropertyEditor/Private/DetailCategoryBuilderImpl.h b/Engine/Source/Editor/PropertyEditor/Private/DetailCategoryBuilderImpl.h
--- a/Engine/Source/Editor/PropertyEditor/Private/DetailCategoryBuilderImpl.h
+++ b/Engine/Source/Editor/PropertyEditor/Private/DetailCategoryBuilderImpl.h
@@ -428,4 +428,6 @@ private:
 	bool bFavoriteCategory : 1;
 	bool bShowOnlyChildren : 1;
 	bool bHasVisibleAdvanced : 1;
+	/** whether or not this category is expanded */
+	bool bExpanded : 1;
 };
diff --git a/Engine/Source/Editor/PropertyEditor/Private/SDetailsViewBase.cpp b/Engine/Source/Editor/PropertyEditor/Private/SDetailsViewBase.cpp
--- a/Engine/Source/Editor/PropertyEditor/Private/SDetailsViewBase.cpp
+++ b/Engine/Source/Editor/PropertyEditor/Private/SDetailsViewBase.cpp
@@ -207,6 +207,72 @@ TArray< FPropertyPath > SDetailsViewBase::GetPropertiesInOrderDisplayed() const
 	return Ret;
 }
 
+static void GetPropertyRowRowNumbersRecursive(const TArray<TSharedRef<FDetailTreeNode>>& TreeNodes, int32& CurrentRowNumber, TArray<TPair<int32, FPropertyPath>>& OutResults)
+{
+	for (const TSharedRef<FDetailTreeNode>& TreeNode : TreeNodes)
+	{
+		// pre-order traversal (add parents to list before their children)
+		FPropertyPath Path = TreeNode->GetPropertyPath();
+
+		if(Path.IsValid())
+		{
+			OutResults.Add(TPair<int32, FPropertyPath>(CurrentRowNumber, Path));
+		}
+		++CurrentRowNumber;
+		
+		if (!TreeNode->IsLeaf() && TreeNode->ShouldBeExpanded())
+		{
+			TArray<TSharedRef<FDetailTreeNode>> Children;
+			TreeNode->GetChildren(Children);
+			GetPropertyRowRowNumbersRecursive(Children, CurrentRowNumber, OutResults);
+		}
+	}
+}
+
+TArray<TPair<int32, FPropertyPath>> SDetailsViewBase::GetPropertyRowNumbers() const
+{
+	TArray<TPair<int32, FPropertyPath>> Results;
+	int32 OverallIdx = 0;
+	GetPropertyRowRowNumbersRecursive(RootTreeNodes, OverallIdx, Results);
+	return Results;
+}
+
+void SDetailsViewBase::ScrollLockDetailsViews(TSharedRef<IDetailsView> OtherDetailsView,
+	const TAttribute<FInvertiblePiecewiseLinearFunction>& ScrollRate)
+{
+	const TSharedPtr<SDetailTree> NewDetailTree = DetailTree;
+	const TSharedPtr<SDetailTree> OldDetailTree = StaticCastSharedRef<SDetailsViewBase>(OtherDetailsView)->DetailTree;
+	if (NewDetailTree && OldDetailTree)
+	{
+		const TSharedPtr<FScrollLockedTableView> ScrollLock = MakeShared<FScrollLockedTableView>(NewDetailTree, OldDetailTree, ScrollRate);
+		NewDetailTree->SetScrollLockedTable(ScrollLock);
+		OldDetailTree->SetScrollLockedTable(ScrollLock);
+	}
+}
+
+static void SetCategoryExpansionSavingRecursive(const TArray<TSharedRef<FDetailTreeNode>>& TreeNodes, bool bShouldCategoriesSaveExpansion)
+{
+	for (const TSharedRef<FDetailTreeNode>& TreeNode : TreeNodes)
+	{
+		if (!TreeNode->IsLeaf())
+		{
+			TArray<TSharedRef<FDetailTreeNode>> Children;
+			TreeNode->GetChildren(Children);
+			SetCategoryExpansionSavingRecursive(Children, bShouldCategoriesSaveExpansion);
+		}
+
+		if(const TSharedPtr<FDetailCategoryImpl> ParentCategory = TreeNode->GetParentCategory())
+		{
+			ParentCategory->RestoreExpansionState(bShouldCategoriesSaveExpansion);
+		}
+	}
+}
+
+void SDetailsViewBase::SetCategoryExpansionSaving(bool bShouldCategoriesSaveExpansion)
+{
+	SetCategoryExpansionSavingRecursive(RootTreeNodes, bShouldCategoriesSaveExpansion);
+}
+
 // construct a map for fast FPropertyPath.ToString() -> FDetailTreeNode lookup
 static void GetPropertyPathToTreeNodes(const TArray<TSharedRef<FDetailTreeNode>>& RootTreeNodes, TMap<FString, TSharedRef<FDetailTreeNode>> &OutMap)
 {
diff --git a/Engine/Source/Editor/PropertyEditor/Private/SDetailsViewBase.h b/Engine/Source/Editor/PropertyEditor/Private/SDetailsViewBase.h
--- a/Engine/Source/Editor/PropertyEditor/Private/SDetailsViewBase.h
+++ b/Engine/Source/Editor/PropertyEditor/Private/SDetailsViewBase.h
@@ -26,6 +26,7 @@
 #include "Widgets/Views/STableRow.h"
 #include "Widgets/Views/STableViewBase.h"
 #include "Widgets/Views/STreeView.h"
+#include "Framework/FInvertiblePiecewiseLinearFunction.h"
 
 class FAssetThumbnailPool;
 class FDetailCategoryImpl;
@@ -85,6 +86,9 @@ public:
 	 * IDetailsView interface 
 	 */
 	virtual TArray< FPropertyPath > GetPropertiesInOrderDisplayed() const override;
+	virtual TArray<TPair<int32, FPropertyPath>> GetPropertyRowNumbers() const override;
+	virtual void ScrollLockDetailsViews(TSharedRef<IDetailsView> OtherDetailsView, const TAttribute<FInvertiblePiecewiseLinearFunction>& ScrollRate) override;
+	virtual void SetCategoryExpansionSaving(bool bShouldCategoriesSaveExpansion) override;
 	virtual void HighlightProperty(const FPropertyPath& Property) override;
 	virtual void ShowAllAdvancedProperties() override;
 	virtual void SetOnDisplayedPropertiesChanged(FOnDisplayedPropertiesChanged InOnDisplayedPropertiesChangedDelegate) override;
diff --git a/Engine/Source/Editor/PropertyEditor/Public/IDetailsView.h b/Engine/Source/Editor/PropertyEditor/Public/IDetailsView.h
--- a/Engine/Source/Editor/PropertyEditor/Public/IDetailsView.h
+++ b/Engine/Source/Editor/PropertyEditor/Public/IDetailsView.h
@@ -7,6 +7,7 @@
 #include "PropertyEditorDelegates.h"
 #include "Framework/Commands/UICommandList.h"
 #include "DetailsViewArgs.h"
+#include "Framework/FInvertiblePiecewiseLinearFunction.h"
 
 class AActor;
 class FNotifyHook;
@@ -187,6 +188,12 @@ public:
 	 * Returns a list of all the properties displayed (via full path), order in list corresponds to draw order:
 	 */
 	virtual TArray< FPropertyPath > GetPropertiesInOrderDisplayed() const = 0;
+	
+	virtual TArray<TPair<int32, FPropertyPath>> GetPropertyRowNumbers() const  = 0;
+
+	virtual void ScrollLockDetailsViews(TSharedRef<IDetailsView> OtherDetailsView, const TAttribute<FInvertiblePiecewiseLinearFunction>& ScrollRate) = 0;
+	
+	virtual void  SetCategoryExpansionSaving(bool bShouldCategoriesSaveExpansion) = 0;
 
 	/**
 	 * Creates a box around the treenode corresponding to Property and scrolls the treenode into view
diff --git a/Engine/Source/Runtime/Slate/Private/Framework/FInvertiblePiecewiseLinearFunction.cpp b/Engine/Source/Runtime/Slate/Private/Framework/FInvertiblePiecewiseLinearFunction.cpp
new file mode 100644
--- /dev/null
+++ b/Engine/Source/Runtime/Slate/Private/Framework/FInvertiblePiecewiseLinearFunction.cpp
@@ -0,0 +1,122 @@
+// Copyright Epic Games, Inc. All Rights Reserved.
+
+#include "Framework/FInvertiblePiecewiseLinearFunction.h"
+
+FInvertiblePointSlopeFunction::FInvertiblePointSlopeFunction(FVector2f Point, float Slope)
+	: Point(Point)
+	, Slope(Slope)
+{}
+
+float FInvertiblePointSlopeFunction::SolveForY(float X) const
+{
+	const float M = Slope;
+	const float B = Point.Y - (M * Point.X);
+	// y = m*x + b
+	return M * X + B;
+}
+
+float FInvertiblePointSlopeFunction::SolveForX(float Y) const
+{
+	const float M = 1 / Slope;
+	const float B = Point.X - (M * Point.Y);
+	// x = m*y + b
+	return M * Y + B;
+}
+
+FInvertiblePiecewiseLinearFunction::FInvertiblePiecewiseLinearFunction()
+	: SubFunctions{{{0.f,0.f}, 1.f}}
+{}
+
+FInvertiblePiecewiseLinearFunction::FInvertiblePiecewiseLinearFunction(const TArray<FVector2f>& FixedPoints)
+{
+	if (FixedPoints.Num() <= 1)
+	{
+		SubFunctions.Emplace(FVector2f{0.f,0.f}, 1.f);
+		return;
+	}
+	
+	float PrevSlope = 0;
+	for(int32 Index = 0; (Index + 1) < FixedPoints.Num(); ++Index)
+	{
+		const FVector2f& Begin = FixedPoints[Index]; 
+		const FVector2f& End = FixedPoints[Index + 1];
+
+		// duplicate points are allowed but ignored
+		if(Begin == End)
+		{
+			continue;
+		}
+			
+		// for invertibility, both x and y must be increasing
+		check(Begin.X <= End.X && Begin.Y <= End.Y);
+			
+		const float Slope = (End.Y - Begin.Y) / (End.X - Begin.X);
+			
+		// Only add a subfunction if it changes the slope
+		if (PrevSlope != Slope)
+		{
+			SubFunctions.Emplace(Begin, Slope);
+		}
+		PrevSlope = Slope;
+	}
+}
+
+float FInvertiblePiecewiseLinearFunction::SolveForY(float X) const
+{
+	return FindSubFunctionAtX(X).SolveForY(X);
+}
+
+float FInvertiblePiecewiseLinearFunction::SolveForX(float Y) const
+{
+	return FindSubFunctionAtY(Y).SolveForX(Y);
+}
+
+const FInvertiblePiecewiseLinearFunction::SubFunction& FInvertiblePiecewiseLinearFunction::FindSubFunctionAtX(float X) const
+{
+	if(X < SubFunctions[0].Point.X)
+	{
+		return SubFunctions[0];
+	}
+		
+	int Start = 0;
+	int End = SubFunctions.Num() - 1;
+		
+	// Binary search for a subFunction with the correct bounds
+	while (Start <= End) {
+		const int Mid = (Start + End) / 2;
+		// If K is found
+		if (SubFunctions[Mid].Point.X == X)
+			return SubFunctions[Mid];
+		else if (SubFunctions[Mid].Point.X < X)
+			Start = Mid + 1;
+		else
+			End = Mid - 1;
+	}
+
+	return SubFunctions[End];
+}
+
+const FInvertiblePiecewiseLinearFunction::SubFunction& FInvertiblePiecewiseLinearFunction::FindSubFunctionAtY(float Y) const
+{
+	if(Y < SubFunctions[0].Point.Y)
+	{
+		return SubFunctions[0];
+	}
+		
+	int Start = 0;
+	int End = SubFunctions.Num() - 1;
+		
+	// Binary search for a subFunction with the correct bounds
+	while (Start <= End) {
+		const int Mid = (Start + End) / 2;
+		// If K is found
+		if (SubFunctions[Mid].Point.Y == Y)
+			return SubFunctions[Mid];
+		else if (SubFunctions[Mid].Point.Y < Y)
+			Start = Mid + 1;
+		else
+			End = Mid - 1;
+	}
+		
+	return SubFunctions[End];
+}
diff --git a/Engine/Source/Runtime/Slate/Private/Widgets/Views/STableViewBase.cpp b/Engine/Source/Runtime/Slate/Private/Widgets/Views/STableViewBase.cpp
--- a/Engine/Source/Runtime/Slate/Private/Widgets/Views/STableViewBase.cpp
+++ b/Engine/Source/Runtime/Slate/Private/Widgets/Views/STableViewBase.cpp
@@ -274,6 +274,29 @@ EActiveTimerReturnType STableViewBase::EnsureTickToRefresh(double InCurrentTime,
 	return EActiveTimerReturnType::Stop;
 }
 
+FScrollLockedTableView::FScrollLockedTableView(const TSharedPtr<STableViewBase>& XTable,
+	const TSharedPtr<STableViewBase>& YTable, const TAttribute<FInvertiblePiecewiseLinearFunction>& ScrollSyncRate): XTable(XTable)
+	, YTable(YTable)
+	, ScrollSyncRate(ScrollSyncRate)
+{}
+
+void FScrollLockedTableView::MatchScroll(const STableViewBase* Scroller) const
+{
+	if (!XTable.IsValid() || !YTable.IsValid())
+	{
+		return;
+	}
+	if (Scroller == XTable.Pin().Get())
+	{
+		YTable.Pin()->SetScrollOffset(ScrollSyncRate.Get().SolveForY(Scroller->GetScrollOffset()));
+	}
+	else
+	{
+		check(Scroller == YTable.Pin().Get());
+		XTable.Pin()->SetScrollOffset(ScrollSyncRate.Get().SolveForX(Scroller->GetScrollOffset()));
+	}
+}
+
 void STableViewBase::Tick( const FGeometry& AllottedGeometry, const double InCurrentTime, const float InDeltaTime )
 {
 	if (ItemsPanel.IsValid())
@@ -832,6 +855,12 @@ float STableViewBase::ScrollTo( float InScrollOffset)
 
 	SetScrollOffset( NewScrollOffset );
 
+	// if another table is scrolling in sync with this one, update it
+	if (ScrollLockedTable)
+	{
+		ScrollLockedTable->MatchScroll(this);
+	}
+
 	return AmountScrolled;
 }
 
@@ -926,6 +955,11 @@ void STableViewBase::SetBackgroundBrush(const TAttribute<const FSlateBrush*>& In
 	BackgroundBrush.SetImage(*this, InBackgroundBrush);
 }
 
+void STableViewBase::SetScrollLockedTable(TSharedPtr<FScrollLockedTableView> ScrollLockedTableView)
+{
+	ScrollLockedTable = ScrollLockedTableView;
+}
+
 void STableViewBase::InsertWidget( const TSharedRef<ITableRow> & WidgetToInset )
 {
 	ItemsPanel->AddSlot(0)
diff --git a/Engine/Source/Runtime/Slate/Public/Framework/FInvertiblePiecewiseLinearFunction.h b/Engine/Source/Runtime/Slate/Public/Framework/FInvertiblePiecewiseLinearFunction.h
new file mode 100644
--- /dev/null
+++ b/Engine/Source/Runtime/Slate/Public/Framework/FInvertiblePiecewiseLinearFunction.h
@@ -0,0 +1,33 @@
+// Copyright Epic Games, Inc. All Rights Reserved.
+#pragma once
+#include "CoreMinimal.h"
+
+struct SLATE_API FInvertiblePointSlopeFunction
+{
+	FInvertiblePointSlopeFunction(FVector2f Point, float Slope);
+
+	float SolveForY(float X) const;
+	float SolveForX(float Y) const;
+
+	FVector2f Point;
+	float Slope;
+};
+
+class SLATE_API FInvertiblePiecewiseLinearFunction
+{
+public:
+	FInvertiblePiecewiseLinearFunction();
+
+	// construct using a list of fixed points that will be interpolated between
+	explicit FInvertiblePiecewiseLinearFunction(const TArray<FVector2f>& FixedPoints);
+
+	float SolveForY(float X) const;
+	float SolveForX(float Y) const;
+
+private:
+	using SubFunction = FInvertiblePointSlopeFunction;
+	const SubFunction& FindSubFunctionAtX(float X) const;
+	const SubFunction& FindSubFunctionAtY(float Y) const;
+
+	TArray<SubFunction> SubFunctions;
+};
diff --git a/Engine/Source/Runtime/Slate/Public/Widgets/Views/STableViewBase.h b/Engine/Source/Runtime/Slate/Public/Widgets/Views/STableViewBase.h
--- a/Engine/Source/Runtime/Slate/Public/Widgets/Views/STableViewBase.h
+++ b/Engine/Source/Runtime/Slate/Public/Widgets/Views/STableViewBase.h
@@ -3,6 +3,7 @@
 #pragma once
 
 #include "CoreMinimal.h"
+#include "Framework/FInvertiblePiecewiseLinearFunction.h"
 #include "Misc/Attribute.h"
 #include "Layout/Geometry.h"
 #include "Input/CursorReply.h"
@@ -26,6 +27,7 @@ class SScrollBar;
 enum class EConsumeMouseWheel : uint8;
 enum class ESlateVisibility : uint8;
 struct FScrollBarStyle;
+class FScrollLockedTableView;
 
 /** If the list panel is arranging items as tiles, this enum dictates how the items should be aligned (basically, where any extra space is placed) */
 UENUM(BlueprintType)
@@ -91,6 +93,7 @@ class SLATE_API STableViewBase
 	, public IScrollableWidget
 {
 public:
+	friend class FScrollLockedTableView;
 
 	/** Create the child widgets that comprise the list */
 	void ConstructChildren( const TAttribute<float>& InItemWidth, const TAttribute<float>& InItemHeight, const TAttribute<EListItemAlignment>& InItemAlignment, const TSharedPtr<SHeaderRow>& InHeaderRow, const TSharedPtr<SScrollBar>& InScrollBar, EOrientation InScrollOrientation, const FOnTableViewScrolled& InOnTableViewScrolled, const FScrollBarStyle* InScrollBarStyle = nullptr, const bool bInPreventThrottling = false );
@@ -182,6 +185,9 @@ public:
 	/** Sets the Background Brush */
 	void SetBackgroundBrush(const TAttribute<const FSlateBrush*>& InBackgroundBrush);
 
+	/** Used to sync the scrolling between two tableviews so they can scroll simultaneously */
+	void SetScrollLockedTable(TSharedPtr<FScrollLockedTableView> ScrollLockedTableView);
+
 public:
 
 	// SWidget interface
@@ -469,6 +475,9 @@ protected:
 	/** Brush resource representing the background area of the view */
 	FInvalidatableBrushAttribute BackgroundBrush;
 
+	/** Set when another table scrolls automatically in sync with this table */
+	TSharedPtr<FScrollLockedTableView> ScrollLockedTable;
+
 protected:
 
 	/** Check whether the current state of the table warrants inertial scroll by the specified amount */
@@ -506,4 +515,23 @@ namespace TableViewHelpers
 	 * view, and therefore does not to know about parent wires.
 	 */
 	SLATE_API const TBitArray<>& GetEmptyBitArray();
-}
\ No newline at end of file
+}
+
+// Helper class used to give a TableView the authority to auto scroll another adjacent
+// table in sync with it
+class SLATE_API FScrollLockedTableView
+{
+public:
+	FScrollLockedTableView(const TSharedPtr<STableViewBase>& XTable, const TSharedPtr<STableViewBase>& YTable,
+		const TAttribute<FInvertiblePiecewiseLinearFunction>& ScrollSyncRate);
+
+	void MatchScroll(const STableViewBase* Scroller) const;
+
+private:
+	TWeakPtr<STableViewBase> XTable;
+	TWeakPtr<STableViewBase> YTable;
+
+	// this determines the relative speed that each table scrolls to stay in sync with one another.
+	// to scroll at a constant and equal rate, just ust the default constructor
+	TAttribute<FInvertiblePiecewiseLinearFunction> ScrollSyncRate;
+};
