Add support for setting UnsafeDuringActorConstruction metadata in blueprint functions, via the details customization used to edit other blueprint function metadata

diff --git a/Engine/Source/Editor/BlueprintGraph/Classes/K2Node_EditablePinBase.h b/Engine/Source/Editor/BlueprintGraph/Classes/K2Node_EditablePinBase.h
--- a/Engine/Source/Editor/BlueprintGraph/Classes/K2Node_EditablePinBase.h
+++ b/Engine/Source/Editor/BlueprintGraph/Classes/K2Node_EditablePinBase.h
@@ -93,6 +93,9 @@ struct FKismetUserDeclaredFunctionMetadata
 	UPROPERTY()
 	bool bThreadSafe;
 
+	UPROPERTY()
+	bool bIsUnsafeDuringActorConstruction;
+
 	/** Cached value for whether or not the graph has latent functions, positive for TRUE, zero for FALSE, and INDEX_None for undetermined */
 	UPROPERTY(transient)
 	int8 HasLatentFunctions;
@@ -103,6 +106,7 @@ public:
 		, bIsDeprecated(false)
 		, bCallInEditor(false)
 		, bThreadSafe(false)
+		, bIsUnsafeDuringActorConstruction(false)
 		, HasLatentFunctions(INDEX_NONE)
 	{
 	}
diff --git a/Engine/Source/Editor/Kismet/Private/BlueprintDetailsCustomization.cpp b/Engine/Source/Editor/Kismet/Private/BlueprintDetailsCustomization.cpp
--- a/Engine/Source/Editor/Kismet/Private/BlueprintDetailsCustomization.cpp
+++ b/Engine/Source/Editor/Kismet/Private/BlueprintDetailsCustomization.cpp
@@ -4095,6 +4095,23 @@ void FBlueprintGraphActionDetails::CustomizeDetails( IDetailLayoutBuilder& Detai
 					.OnCheckStateChanged( this, &FBlueprintGraphActionDetails::OnIsThreadSafeFunctionModified )
 				];
 			}
+			if (IsUnsafeDuringActorConstructionVisible())
+			{
+				Category.AddCustomRow(LOCTEXT("FunctionUnsafeDuringActorConstruction_Tooltip", "Unsafe During Actor Constrution"), true)
+				.NameContent()
+				[
+					SNew(STextBlock)
+					.Text(LOCTEXT("FunctionUnsafeDuringActorConstruction_Tooltip", "Unsafe During Actor Construction"))
+					.ToolTipText(LOCTEXT("FunctionIsUnsafeDuringActorConstruction_Tooltip", "Mark this function as unsafe during actor construction so that a warning is generated when it is called by a Constructin Script - useful when calling native functions that are also unsafe during construction"))
+					.Font(IDetailLayoutBuilder::GetDetailFont())
+				]
+				.ValueContent()
+				[
+					SNew(SCheckBox)
+					.IsChecked(this, &FBlueprintGraphActionDetails::GetIsUnsafeDuringActorConstruction)
+					.OnCheckStateChanged(this, &FBlueprintGraphActionDetails::OnIsUnsafeDuringActorConstructionModified)
+				];
+			}
 		}
 
 		if (bIsCustomEvent)
@@ -5797,6 +5814,51 @@ ECheckBoxState FBlueprintGraphActionDetails::GetIsThreadSafeFunction() const
 	return EntryNode->MetaData.bThreadSafe ? ECheckBoxState::Checked : ECheckBoxState::Unchecked;
 }
 
+bool FBlueprintGraphActionDetails::IsUnsafeDuringActorConstructionVisible() const
+{
+	bool bSupportedType = false;
+	bool bIsEditable = false;
+
+	UK2Node_EditablePinBase* FunctionEntryNode = FunctionEntryNodePtr.Get();
+	if (FunctionEntryNode)
+	{
+		bSupportedType = FunctionEntryNode->IsA<UK2Node_FunctionEntry>();
+		bIsEditable = FunctionEntryNode->IsEditable();
+	}
+
+	return bIsEditable && bSupportedType;
+}
+
+void FBlueprintGraphActionDetails::OnIsUnsafeDuringActorConstructionModified(const ECheckBoxState NewCheckedState)
+{
+	if (FunctionEntryNodePtr.IsValid())
+	{
+		const bool bIsUnsafeDuringActorConstruction = NewCheckedState == ECheckBoxState::Checked;
+		const FText TransactionType = bIsUnsafeDuringActorConstruction ? 
+			LOCTEXT("DisableIsUnsafeDuringActorConstruction", "Disable Unsafe During Actor Construction") : 
+			LOCTEXT("EnableThreadSafe", "Enable Unsafe During Actor Construction");
+
+		if (UK2Node_FunctionEntry* EntryPoint = Cast<UK2Node_FunctionEntry>(FunctionEntryNodePtr.Get()))
+		{
+			const FScopedTransaction Transaction(TransactionType);
+			EntryPoint->Modify();
+			EntryPoint->MetaData.bIsUnsafeDuringActorConstruction = bIsUnsafeDuringActorConstruction;
+			FBlueprintEditorUtils::MarkBlueprintAsModified(EntryPoint->GetBlueprint());
+		}
+	}
+}
+
+ECheckBoxState FBlueprintGraphActionDetails::GetIsUnsafeDuringActorConstruction() const
+{
+	UK2Node_EditablePinBase* FunctionEntryNode = FunctionEntryNodePtr.Get();
+	UK2Node_FunctionEntry* EntryNode = Cast<UK2Node_FunctionEntry>(FunctionEntryNode);
+	if (!EntryNode)
+	{
+		return ECheckBoxState::Undetermined;
+	}
+	return EntryNode->MetaData.bIsUnsafeDuringActorConstruction ? ECheckBoxState::Checked : ECheckBoxState::Unchecked;
+}
+
 FReply FBaseBlueprintGraphActionDetails::OnAddNewInputClicked()
 {
 	UK2Node_EditablePinBase* FunctionEntryNode = FunctionEntryNodePtr.Get();
diff --git a/Engine/Source/Editor/Kismet/Private/BlueprintDetailsCustomization.h b/Engine/Source/Editor/Kismet/Private/BlueprintDetailsCustomization.h
--- a/Engine/Source/Editor/Kismet/Private/BlueprintDetailsCustomization.h
+++ b/Engine/Source/Editor/Kismet/Private/BlueprintDetailsCustomization.h
@@ -779,6 +779,10 @@ private:
 	void OnIsThreadSafeFunctionModified(const ECheckBoxState NewCheckedState);
 	ECheckBoxState GetIsThreadSafeFunction() const;
 	
+	bool IsUnsafeDuringActorConstructionVisible() const;
+	void OnIsUnsafeDuringActorConstructionModified(const ECheckBoxState NewCheckedState);
+	ECheckBoxState GetIsUnsafeDuringActorConstruction() const;
+
 	/** Determines if the selected event is identified as editor callable */
 	ECheckBoxState GetIsEditorCallableEvent() const;
 
diff --git a/Engine/Source/Editor/KismetCompiler/Private/KismetCompiler.cpp b/Engine/Source/Editor/KismetCompiler/Private/KismetCompiler.cpp
--- a/Engine/Source/Editor/KismetCompiler/Private/KismetCompiler.cpp
+++ b/Engine/Source/Editor/KismetCompiler/Private/KismetCompiler.cpp
@@ -2271,6 +2271,11 @@ void FKismetCompilerContext::PrecompileFunction(FKismetFunctionContext& Context,
 			Context.Function->SetMetaData(FBlueprintMetadata::MD_CallInEditor, TEXT( "true" ));
 		}
 
+		if (FunctionMetaData.bIsUnsafeDuringActorConstruction)
+		{
+			Context.Function->SetMetaData(FBlueprintMetadata::MD_UnsafeForConstructionScripts, TEXT("true"));
+		}
+
 		// Set appropriate metadata if the function is deprecated
 		if (FunctionMetaData.bIsDeprecated)
 		{
@@ -2903,6 +2908,11 @@ void FKismetCompilerContext::SetCalculatedMetaDataAndFlags(UFunction* Function,
 	{
 		Function->SetMetaData(FBlueprintMetadata::MD_ThreadSafe, TEXT("true"));
 	}
+
+	if (EntryNode->MetaData.bIsUnsafeDuringActorConstruction)
+	{
+		Function->SetMetaData(FBlueprintMetadata::MD_UnsafeForConstructionScripts, TEXT("true"));
+	}
 	
 	if (UEdGraphPin* WorldContextPin = EntryNode->GetAutoWorldContextPin())
 	{
