Fix two bad loctext keys (one due to a typo causing keys to fail to be reused, one caused by copy/paste)
#rnx

diff --git a/Engine/Source/Editor/Kismet/Private/BlueprintDetailsCustomization.cpp b/Engine/Source/Editor/Kismet/Private/BlueprintDetailsCustomization.cpp
--- a/Engine/Source/Editor/Kismet/Private/BlueprintDetailsCustomization.cpp
+++ b/Engine/Source/Editor/Kismet/Private/BlueprintDetailsCustomization.cpp
@@ -4097,7 +4097,7 @@ void FBlueprintGraphActionDetails::CustomizeDetails( IDetailLayoutBuilder& Detai
 			}
 			if (IsUnsafeDuringActorConstructionVisible())
 			{
-				Category.AddCustomRow(LOCTEXT("FunctionUnsafeDuringActorConstruction_Tooltip", "Unsafe During Actor Constrution"), true)
+				Category.AddCustomRow(LOCTEXT("FunctionUnsafeDuringActorConstruction_Tooltip", "Unsafe During Actor Construction"), true)
 				.NameContent()
 				[
 					SNew(STextBlock)
@@ -5836,7 +5836,7 @@ void FBlueprintGraphActionDetails::OnIsUnsafeDuringActorConstructionModified(con
 		const bool bIsUnsafeDuringActorConstruction = NewCheckedState == ECheckBoxState::Checked;
 		const FText TransactionType = bIsUnsafeDuringActorConstruction ? 
 			LOCTEXT("DisableIsUnsafeDuringActorConstruction", "Disable Unsafe During Actor Construction") : 
-			LOCTEXT("EnableThreadSafe", "Enable Unsafe During Actor Construction");
+			LOCTEXT("EnableIsUnsafeDuringActorConstruction", "Enable Unsafe During Actor Construction");
 
 		if (UK2Node_FunctionEntry* EntryPoint = Cast<UK2Node_FunctionEntry>(FunctionEntryNodePtr.Get()))
 		{
