EditInstanceOnly properties introduced in a direct parent class are no longer visible on child blueprint classes.

For example, creating a C++ class with a EditInstanceOnly property would be visible in a blueprint class derived from that C++ class. APropertyEditorTestActor is an example of one such C++ class that contains EditInstanceOnly classes.

diff --git a/Engine/Source/Editor/Kismet/Private/SKismetInspector.cpp b/Engine/Source/Editor/Kismet/Private/SKismetInspector.cpp
--- a/Engine/Source/Editor/Kismet/Private/SKismetInspector.cpp
+++ b/Engine/Source/Editor/Kismet/Private/SKismetInspector.cpp
@@ -857,8 +857,7 @@ bool SKismetInspector::IsPropertyVisible( const FPropertyAndParent& PropertyAndP
 		{
 			if (const UClass* OwningClass = Property->GetOwnerClass())
 			{
-				return Blueprint && (OwningClass == Blueprint->ParentClass ||
-					OwningClass->ClassGeneratedBy == Blueprint);
+				return Blueprint && OwningClass->ClassGeneratedBy == Blueprint;
 			}
 			return false;
 		};
