Fix warning spam and editor stalls when searching for blueprint variables
This will show the correct class-specific pin colors in more situations, but they may still be wrong for unloaded blueprints

diff --git a/Engine/Source/Editor/Kismet/Private/FindInBlueprintManager.cpp b/Engine/Source/Editor/Kismet/Private/FindInBlueprintManager.cpp
--- a/Engine/Source/Editor/Kismet/Private/FindInBlueprintManager.cpp
+++ b/Engine/Source/Editor/Kismet/Private/FindInBlueprintManager.cpp
@@ -724,7 +724,8 @@ namespace BlueprintSearchMetaDataHelpers
 
 		if(InPinType.PinSubCategoryObject.IsValid())
 		{
-			InWriter->WriteValue(FFindInBlueprintSearchTags::FiB_ObjectClass, FText::FromString(InPinType.PinSubCategoryObject->GetName()));
+			// Write the full path because this can be an ambiguous blueprint class name
+			InWriter->WriteValue(FFindInBlueprintSearchTags::FiB_ObjectClass, FText::FromString(InPinType.PinSubCategoryObject->GetPathName()));
 		}
 		InWriter->WriteValue(FFindInBlueprintSearchTags::FiB_IsArray, InPinType.IsArray());
 		InWriter->WriteValue(FFindInBlueprintSearchTags::FiB_IsReference, InPinType.bIsReference);
diff --git a/Engine/Source/Editor/Kismet/Private/FindInBlueprints.cpp b/Engine/Source/Editor/Kismet/Private/FindInBlueprints.cpp
--- a/Engine/Source/Editor/Kismet/Private/FindInBlueprints.cpp
+++ b/Engine/Source/Editor/Kismet/Private/FindInBlueprints.cpp
@@ -368,7 +368,19 @@ void FFindInBlueprintsGraphNode::FinalizeSearchData()
 {
 	if(!ClassName.IsEmpty())
 	{
-		Class = UClass::TryFindTypeSlow<UClass>(ClassName, EFindFirstObjectOptions::ExactClass);
+		// Check the node subclasses and look for one with the same short name
+		TArray<UClass*> NodeClasses;
+		GetDerivedClasses(UEdGraphNode::StaticClass(), NodeClasses, /*bRecursive=*/true);
+
+		for (UClass* FoundClass : NodeClasses)
+		{
+			if (FoundClass->GetName() == ClassName)
+			{
+				Class = FoundClass;
+				break;
+			}
+		}
+
 		ClassName.Empty();
 	}
 }
@@ -431,25 +443,35 @@ void FFindInBlueprintsPin::FinalizeSearchData()
 {
 	if(!PinType.PinSubCategory.IsNone())
 	{
-		PinType.PinSubCategoryObject = UClass::TryFindTypeSlow<UClass>(PinType.PinSubCategory.ToString(), EFindFirstObjectOptions::ExactClass);
-		if(!PinType.PinSubCategoryObject.IsValid())
+		// This can either be a full path to an object, or a short name specific to the category
+		if (FPackageName::IsShortPackageName(PinType.PinSubCategory))
 		{
-			PinType.PinSubCategoryObject = FindObject<UScriptStruct>(UObject::StaticClass(), *PinType.PinSubCategory.ToString());
+			// This could also be an old class name without the full path, but it's fine to ignore in that case
 		}
-
-		if (PinType.PinSubCategoryObject.IsValid())
+		else
 		{
-			PinType.PinSubCategory = NAME_None;
+			PinType.PinSubCategoryObject = FindObject<UObject>(UObject::StaticClass(), *PinType.PinSubCategory.ToString());
+			if (PinType.PinSubCategoryObject.IsValid())
+			{
+				PinType.PinSubCategory = NAME_None;
+			}
 		}
 	}
 
 	if(!SchemaName.IsEmpty())
 	{
-		UClass* SchemaClass = UClass::TryFindTypeSlow<UClass>(SchemaName, EFindFirstObjectOptions::ExactClass);
-		if(SchemaClass)
+		// Get all subclasses of schema and find the one with a matching short name
+		TArray<UClass*> SchemaClasses;
+		GetDerivedClasses(UEdGraphSchema::StaticClass(), SchemaClasses, /*bRecursive=*/true);
+
+		for (UClass* FoundClass : SchemaClasses)
 		{
-			UEdGraphSchema* Schema = SchemaClass->GetDefaultObject<UEdGraphSchema>();
-			IconColor = Schema->GetPinTypeColor(PinType);
+			if (FoundClass->GetName() == SchemaName)
+			{
+				UEdGraphSchema* Schema = FoundClass->GetDefaultObject<UEdGraphSchema>();
+				IconColor = Schema->GetPinTypeColor(PinType);
+				break;
+			}
 		}
 
 		SchemaName.Empty();
@@ -547,17 +569,20 @@ FText FFindInBlueprintsProperty::GetCategory() const
 
 void FFindInBlueprintsProperty::FinalizeSearchData()
 {
-	if(!PinType.PinSubCategory.IsNone())
+	if (!PinType.PinSubCategory.IsNone())
 	{
-		PinType.PinSubCategoryObject = UClass::TryFindTypeSlow<UClass>(PinType.PinSubCategory.ToString(), EFindFirstObjectOptions::ExactClass);
-		if(!PinType.PinSubCategoryObject.IsValid())
+		// This can either be a full path to an object, or a short name specific to the category
+		if (FPackageName::IsShortPackageName(PinType.PinSubCategory))
 		{
-			PinType.PinSubCategoryObject = FindObject<UScriptStruct>(UObject::StaticClass(), *PinType.PinSubCategory.ToString());
+			// This could also be an old class name without the full path, but it's fine to ignore in that case
 		}
-
-		if (PinType.PinSubCategoryObject.IsValid())
+		else
 		{
-			PinType.PinSubCategory = NAME_None;
+			PinType.PinSubCategoryObject = FindObject<UObject>(UObject::StaticClass(), *PinType.PinSubCategory.ToString());
+			if (PinType.PinSubCategoryObject.IsValid())
+			{
+				PinType.PinSubCategory = NAME_None;
+			}
 		}
 	}
 }
