Re: Synchronize the scrolling between left and right details panels in Blueprint Diff (Refactor by request of Patrick Boutot)

diff --git a/Engine/Source/Editor/Kismet/Private/DetailsDiff.cpp b/Engine/Source/Editor/Kismet/Private/DetailsDiff.cpp
--- a/Engine/Source/Editor/Kismet/Private/DetailsDiff.cpp
+++ b/Engine/Source/Editor/Kismet/Private/DetailsDiff.cpp
@@ -16,6 +16,7 @@
 #include "PropertyPath.h"
 #include "UObject/WeakObjectPtrTemplates.h"
 #include "Engine/MemberReference.h"
+#include "Widgets/Layout/LinkableScrollBar.h"
 
 class UObject;
 
@@ -27,6 +28,8 @@ FDetailsDiff::FDetailsDiff(const UObject* InObject, FOnDisplayedPropertiesChange
 	FDetailsViewArgs DetailsViewArgs;
 	DetailsViewArgs.bShowDifferingPropertiesOption = true;
 	DetailsViewArgs.NameAreaSettings = FDetailsViewArgs::HideNameArea;
+	ScrollBar = SNew(SLinkableScrollBar);
+	DetailsViewArgs.ExternalScrollbar = ScrollBar;
 
 	FPropertyEditorModule& EditModule = FModuleManager::Get().GetModuleChecked<FPropertyEditorModule>("PropertyEditor");
 	DetailsView = EditModule.CreateDetailView(DetailsViewArgs);
@@ -235,3 +238,9 @@ void FDetailsDiff::DiffAgainst(const FDetailsDiff& Newer, TArray< FSingleObjectD
 		}
 	}
 }
+
+void FDetailsDiff::LinkScrolling(FDetailsDiff& LeftPanel, FDetailsDiff& RightPanel,
+	const TAttribute<TArray<FVector2f>>& ScrollRate)
+{
+	SLinkableScrollBar::LinkScrollBars(LeftPanel.ScrollBar.ToSharedRef(), RightPanel.ScrollBar.ToSharedRef(), ScrollRate);
+}
diff --git a/Engine/Source/Editor/Kismet/Private/DiffControl.cpp b/Engine/Source/Editor/Kismet/Private/DiffControl.cpp
--- a/Engine/Source/Editor/Kismet/Private/DiffControl.cpp
+++ b/Engine/Source/Editor/Kismet/Private/DiffControl.cpp
@@ -10,6 +10,34 @@
 
 #define LOCTEXT_NAMESPACE "SBlueprintDif"
 
+// well known algorithm for computing the Longest Common Subsequence table of two ordered lists
+template <typename RangeType, typename ComparePredicate>
+static TArray<TArray<int32>> CalculateLCSTable(const RangeType& Range1, const RangeType& Range2, ComparePredicate Comparison)
+{
+	TArray<TArray<int32>> LCS;
+	LCS.SetNum(Range1.Num() + 1);
+	for (int32 I = 0; I <= Range1.Num(); I++)
+	{
+		LCS[I].SetNum(Range2.Num() + 1);
+		if (I == 0)
+		{
+			continue;
+		}
+	
+		for (int32 J = 1; J <= Range2.Num(); J++)
+		{
+			if (Comparison(Range1[I - 1], Range2[J - 1]))
+			{
+				LCS[I][J] = LCS[I - 1][J - 1] + 1;
+			}
+			else
+			{
+				LCS[I][J] = FMath::Max(LCS[I - 1][J], LCS[I][J - 1]);
+			}
+		}
+	}
+	return LCS;
+}
 
 /////////////////////////////////////////////////////////////////////////////
 /// IDiffControl
@@ -185,17 +213,137 @@ void FSCSDiffControl::GenerateTreeEntries(TArray<TSharedPtr<FBlueprintDifference
 	));
 }
 
+FDetailsDiffControl::FDetailsDiffControl(const UObject* InOldObject, const UObject* InNewObject,
+	FOnDiffEntryFocused InSelectionCallback, bool bPopulateOutTreeEntries): SelectionCallback(InSelectionCallback)
+	                                                                        , OldDetails(InOldObject, FDetailsDiff::FOnDisplayedPropertiesChanged())
+	                                                                        , NewDetails(InNewObject, FDetailsDiff::FOnDisplayedPropertiesChanged())
+	                                                                        , bPopulateOutTreeEntries(bPopulateOutTreeEntries)
+{
+	OldDetails.DiffAgainst(NewDetails, DifferingProperties, true);
+
+	TSet<FPropertyPath> PropertyPaths;
+	Algo::Transform(DifferingProperties, PropertyPaths,
+	                [&InOldObject](const FSingleObjectDiffEntry& DiffEntry)
+	                {
+		                return DiffEntry.Identifier.ResolvePath(InOldObject);
+	                });
+
+	OldDetails.DetailsWidget()->UpdatePropertyAllowList(PropertyPaths);
+
+	PropertyPaths.Reset();
+	Algo::Transform(DifferingProperties, PropertyPaths,
+	                [&InNewObject](const FSingleObjectDiffEntry& DiffEntry)
+	                {
+		                return DiffEntry.Identifier.ResolvePath(InNewObject);
+	                });
+
+	NewDetails.DetailsWidget()->UpdatePropertyAllowList(PropertyPaths);
+		
+	// Sync the scrolling between the left and right panels
+	const TSharedRef<IDetailsView> OldDetailsView = OldDetails.DetailsWidget();
+	const TSharedRef<IDetailsView> NewDetailsView = NewDetails.DetailsWidget();
+	FDetailsDiff::LinkScrolling(OldDetails, NewDetails, GetLinkedScrollRateAttribute(OldDetailsView, NewDetailsView));;
+}
+
+void FDetailsDiffControl::GenerateTreeEntries(TArray<TSharedPtr<FBlueprintDifferenceTreeEntry>>& OutTreeEntries, TArray<TSharedPtr<FBlueprintDifferenceTreeEntry>>& OutRealDifferences)
+{
+	for (const FSingleObjectDiffEntry& Difference : DifferingProperties)
+	{
+			
+		TSharedPtr<FBlueprintDifferenceTreeEntry> Entry = MakeShared<FBlueprintDifferenceTreeEntry>(
+			FOnDiffEntryFocused::CreateSP(TSharedFromThis<FDetailsDiffControl>::AsShared(), &FDetailsDiffControl::OnSelectDiffEntry, Difference.Identifier),
+			FGenerateDiffEntryWidget::CreateStatic(&GenerateObjectDiffWidget, Difference, RightRevision));
+		Children.Push(Entry);
+		OutRealDifferences.Push(Entry);
+		if (bPopulateOutTreeEntries)
+		{
+			OutTreeEntries.Push(Entry);
+		}
+	}
+}
+
+void FDetailsDiffControl::OnSelectDiffEntry(FPropertySoftPath PropertyName)
+{
+	SelectionCallback.ExecuteIfBound();
+	OldDetails.HighlightProperty(PropertyName);
+	NewDetails.HighlightProperty(PropertyName);
+}
+
+TAttribute<TArray<FVector2f>> FDetailsDiffControl::GetLinkedScrollRateAttribute(const TSharedRef<IDetailsView>& OldDetailsView, const TSharedRef<IDetailsView>& NewDetailsView)
+{
+	return TAttribute<TArray<FVector2f>>::CreateRaw(this, &FDetailsDiffControl::GetLinkedScrollRate, OldDetailsView, NewDetailsView);
+}
+
+TArray<FVector2f> FDetailsDiffControl::GetLinkedScrollRate(TSharedRef<IDetailsView> OldDetailsView, TSharedRef<IDetailsView> NewDetailsView) const
+{
+	TArray<TPair<int32, FPropertyPath>> OldProperties = OldDetailsView->GetPropertyRowNumbers();
+	TArray<TPair<int32, FPropertyPath>> NewProperties = NewDetailsView->GetPropertyRowNumbers();
+
+	const int32 OldRowCount = OldDetailsView->CountRows();
+	const int32 NewRowCount = NewDetailsView->CountRows();
+
+	// use caching to avoid O(n^2) LCS calculation every frame
+	if(LinkedScrollRateCache.OldProperties != OldProperties || LinkedScrollRateCache.NewProperties != NewProperties)
+	{
+		const auto PropertyPathsEqual = [](const TPair<int32, FPropertyPath>& A, const TPair<int32, FPropertyPath>& B)
+		{
+			return A.Value == B.Value;
+		};
+		
+		// Find the Longest Common Subsequence between both sets of properties
+		const TArray<TArray<int32>> LCS = CalculateLCSTable(OldProperties, NewProperties, PropertyPathsEqual);
+	
+		// Using the LCS Table, we can determine which lines should match one another while scrolling. For example, an element
+		// of FixedPoints may be {12.f, 5.f} meaning that line 12 of the left panel is the same as line 5 of the right panel
+		TArray<FVector2f> FixedPoints;
+
+		int32 I = OldProperties.Num();
+		int32 J = NewProperties.Num();
+		while (I > 0 && J > 0)
+		{
+			if (OldProperties[I - 1].Value == NewProperties[J - 1].Value)
+			{
+				const int32 OldLineNum = OldProperties[I - 1].Key;
+				const int32 NewLineNum = NewProperties[J - 1].Key;
+				// add two endpoints for every matched property because the elements have a width of 1 in the panels
+				FixedPoints.Add(FVector2f((OldLineNum + 1.f) / OldRowCount, (NewLineNum + 1.f) / NewRowCount));
+				FixedPoints.Add(FVector2f((float)OldLineNum / OldRowCount, (float)NewLineNum / NewRowCount));
+				--I;
+				--J;
+			}
+			else if (LCS[I - 1][J] <= LCS[I][J	 - 1])
+			{
+				--J;
+			}
+			else
+			{
+				--I;
+			}
+		}
+		FixedPoints.Add({0.f, 0.f});
+		Algo::Reverse(FixedPoints);
+		FixedPoints.Add({1.f, 1.f});
+        	
+		LinkedScrollRateCache.OldProperties = OldProperties;
+		LinkedScrollRateCache.NewProperties = NewProperties;
+		LinkedScrollRateCache.ScrollRate = FixedPoints;
+	}
+	
+	return LinkedScrollRateCache.ScrollRate;
+}
+
 
 /////////////////////////////////////////////////////////////////////////////
 /// FCDODiffControl
 
-FCDODiffControl::FCDODiffControl(const UObject* InOldObject, const UObject* InNewObject, FOnDiffEntryFocused InSelectionCallback): TDetailsDiffControl(InOldObject, InNewObject, InSelectionCallback)
+FCDODiffControl::FCDODiffControl(const UObject* InOldObject, const UObject* InNewObject, FOnDiffEntryFocused InSelectionCallback)
+	: FDetailsDiffControl(InOldObject, InNewObject, InSelectionCallback, false)
 {
 }
 
 void FCDODiffControl::GenerateTreeEntries(TArray<TSharedPtr<FBlueprintDifferenceTreeEntry>>& OutTreeEntries, TArray<TSharedPtr<FBlueprintDifferenceTreeEntry>>& OutRealDifferences)
 {
-	TDetailsDiffControl::GenerateTreeEntries(OutTreeEntries, OutRealDifferences);
+	FDetailsDiffControl::GenerateTreeEntries(OutTreeEntries, OutRealDifferences);
 
 	const bool bHasDifferences = Children.Num() != 0;
 	if (!bHasDifferences)
@@ -218,13 +366,14 @@ void FCDODiffControl::GenerateTreeEntries(TArray<TSharedPtr<FBlueprintDifference
 /////////////////////////////////////////////////////////////////////////////
 /// FClassSettingsDiffControl
 
-FClassSettingsDiffControl::FClassSettingsDiffControl(const UObject* InOldObject, const UObject* InNewObject, FOnDiffEntryFocused InSelectionCallback): TDetailsDiffControl(InOldObject, InNewObject, InSelectionCallback)
+FClassSettingsDiffControl::FClassSettingsDiffControl(const UObject* InOldObject, const UObject* InNewObject, FOnDiffEntryFocused InSelectionCallback)
+	: FDetailsDiffControl(InOldObject, InNewObject, InSelectionCallback, false)
 {
 }
 
 void FClassSettingsDiffControl::GenerateTreeEntries(TArray<TSharedPtr<FBlueprintDifferenceTreeEntry>>& OutTreeEntries, TArray<TSharedPtr<FBlueprintDifferenceTreeEntry>>& OutRealDifferences)
 {
-	TDetailsDiffControl::GenerateTreeEntries(OutTreeEntries, OutRealDifferences);
+	FDetailsDiffControl::GenerateTreeEntries(OutTreeEntries, OutRealDifferences);
 
 	const bool bHasDifferences = Children.Num() != 0;
 	if (!bHasDifferences)
diff --git a/Engine/Source/Editor/Kismet/Private/DiffControl.h b/Engine/Source/Editor/Kismet/Private/DiffControl.h
--- a/Engine/Source/Editor/Kismet/Private/DiffControl.h
+++ b/Engine/Source/Editor/Kismet/Private/DiffControl.h
@@ -15,6 +15,7 @@ class SBlueprintDiff;
 class UEdGraph;
 struct FEdGraphEditAction;
 
+
 /** Interface responsible for generating FBlueprintDifferenceTreeEntry's for visual diff tools */
 class KISMET_API IDiffControl
 {
@@ -67,63 +68,23 @@ private:
 };
 
 /** Generic wrapper around a details view template parameter determines whether TreeEntries are populated */
-template<bool bPopulateOutTreeEntries>
-class KISMET_API TDetailsDiffControl : public TSharedFromThis<TDetailsDiffControl<bPopulateOutTreeEntries>>, public IDiffControl
+class KISMET_API FDetailsDiffControl : public TSharedFromThis<FDetailsDiffControl>, public IDiffControl
 {
 public:
-	TDetailsDiffControl(const UObject* InOldObject, const UObject* InNewObject, FOnDiffEntryFocused InSelectionCallback)
-		: SelectionCallback(InSelectionCallback)
-		, OldDetails(InOldObject, FDetailsDiff::FOnDisplayedPropertiesChanged())
-		, NewDetails(InNewObject, FDetailsDiff::FOnDisplayedPropertiesChanged())
-	{
-		OldDetails.DiffAgainst(NewDetails, DifferingProperties, true);
-
-		TSet<FPropertyPath> PropertyPaths;
-		Algo::Transform(DifferingProperties, PropertyPaths,
-			[&InOldObject](const FSingleObjectDiffEntry& DiffEntry)
-			{
-				return DiffEntry.Identifier.ResolvePath(InOldObject);
-			});
-
-		OldDetails.DetailsWidget()->UpdatePropertyAllowList(PropertyPaths);
+	FDetailsDiffControl(const UObject* InOldObject, const UObject* InNewObject, FOnDiffEntryFocused InSelectionCallback, bool bPopulateOutTreeEntries);
 
-		PropertyPaths.Reset();
-		Algo::Transform(DifferingProperties, PropertyPaths,
-			[&InNewObject](const FSingleObjectDiffEntry& DiffEntry)
-			{
-				return DiffEntry.Identifier.ResolvePath(InNewObject);
-			});
-
-		NewDetails.DetailsWidget()->UpdatePropertyAllowList(PropertyPaths);
-	}
-
-	virtual void GenerateTreeEntries(TArray< TSharedPtr<FBlueprintDifferenceTreeEntry> >& OutTreeEntries, TArray< TSharedPtr<FBlueprintDifferenceTreeEntry> >& OutRealDifferences) override
-	{
-		for (const FSingleObjectDiffEntry& Difference : DifferingProperties)
-		{
-			
-			TSharedPtr<FBlueprintDifferenceTreeEntry> Entry = MakeShared<FBlueprintDifferenceTreeEntry>(
-				FOnDiffEntryFocused::CreateSP(TSharedFromThis<TDetailsDiffControl<bPopulateOutTreeEntries>>::AsShared(), &TDetailsDiffControl::OnSelectDiffEntry, Difference.Identifier),
-				FGenerateDiffEntryWidget::CreateStatic(&GenerateObjectDiffWidget, Difference, RightRevision));
-			Children.Push(Entry);
-			OutRealDifferences.Push(Entry);
-			if constexpr (bPopulateOutTreeEntries)
-			{
-				OutTreeEntries.Push(Entry);
-			}
-		}
-	}
+	virtual void GenerateTreeEntries(TArray< TSharedPtr<FBlueprintDifferenceTreeEntry> >& OutTreeEntries, TArray< TSharedPtr<FBlueprintDifferenceTreeEntry> >& OutRealDifferences) override;
 
 	TSharedRef<SWidget> OldDetailsWidget() { return OldDetails.DetailsWidget(); }
 	TSharedRef<SWidget> NewDetailsWidget() { return NewDetails.DetailsWidget(); }
 
 protected:
-	virtual void OnSelectDiffEntry(FPropertySoftPath PropertyName)
-	{
-		SelectionCallback.ExecuteIfBound();
-		OldDetails.HighlightProperty(PropertyName);
-		NewDetails.HighlightProperty(PropertyName);
-	}
+	virtual void OnSelectDiffEntry(FPropertySoftPath PropertyName);
+	TAttribute<TArray<FVector2f>> GetLinkedScrollRateAttribute(const TSharedRef<IDetailsView>& OldDetailsView, const TSharedRef<IDetailsView>& NewDetailsView);
+	
+
+	// helper function that analyzes two details views and determines the rate they should scroll relative to one another to be in sync
+	TArray<FVector2f> GetLinkedScrollRate(TSharedRef<IDetailsView> OldDetailsView, TSharedRef<IDetailsView> NewDetailsView) const;
 
 	FOnDiffEntryFocused SelectionCallback;
 	FDetailsDiff OldDetails;
@@ -131,12 +92,18 @@ protected:
 
 	TArray<FSingleObjectDiffEntry> DifferingProperties;
 	TArray< TSharedPtr<FBlueprintDifferenceTreeEntry> > Children;
-};
+	const bool bPopulateOutTreeEntries;
 
-using FDetailsDiffControl = TDetailsDiffControl<true>;
+	mutable struct
+	{
+		TArray<TPair<int32, FPropertyPath>> OldProperties;
+		TArray<TPair<int32, FPropertyPath>> NewProperties;
+		TArray<FVector2f> ScrollRate;
+	} LinkedScrollRateCache;
+};
 
 /** Override for CDO special case */
-class FCDODiffControl : public TDetailsDiffControl<false>
+class FCDODiffControl : public FDetailsDiffControl
 {
 public:
 	FCDODiffControl(const UObject* InOldObject, const UObject* InNewObject, FOnDiffEntryFocused InSelectionCallback);
@@ -145,7 +112,7 @@ public:
 };
 
 /** Override for class class settings */
-class FClassSettingsDiffControl : public TDetailsDiffControl<false>
+class FClassSettingsDiffControl : public FDetailsDiffControl
 {
 public:
 	FClassSettingsDiffControl(const UObject* InOldObject, const UObject* InNewObject, FOnDiffEntryFocused InSelectionCallback);
diff --git a/Engine/Source/Editor/Kismet/Private/SDetailsDiff.cpp b/Engine/Source/Editor/Kismet/Private/SDetailsDiff.cpp
--- a/Engine/Source/Editor/Kismet/Private/SDetailsDiff.cpp
+++ b/Engine/Source/Editor/Kismet/Private/SDetailsDiff.cpp
@@ -306,7 +306,7 @@ void SDetailsDiff::GenerateDifferencesList()
 
 SDetailsDiff::FDiffControl SDetailsDiff::GenerateDetailsPanel()
 {
-	TSharedPtr<FDetailsDiffControl> NewDiffControl = MakeShared<FDetailsDiffControl>(PanelOld.Object, PanelNew.Object, FOnDiffEntryFocused::CreateRaw(this, &SDetailsDiff::SetCurrentMode, DetailsMode));
+	TSharedPtr<FDetailsDiffControl> NewDiffControl = MakeShared<FDetailsDiffControl>(PanelOld.Object, PanelNew.Object, FOnDiffEntryFocused::CreateRaw(this, &SDetailsDiff::SetCurrentMode, DetailsMode), true);
 	NewDiffControl->GenerateTreeEntries(PrimaryDifferencesList, RealDifferences);
 
 	SDetailsDiff::FDiffControl Ret;
diff --git a/Engine/Source/Editor/Kismet/Public/DetailsDiff.h b/Engine/Source/Editor/Kismet/Public/DetailsDiff.h
--- a/Engine/Source/Editor/Kismet/Public/DetailsDiff.h
+++ b/Engine/Source/Editor/Kismet/Public/DetailsDiff.h
@@ -8,6 +8,7 @@
 #include "PropertyPath.h"
 #include "Templates/SharedPointer.h"
 #include "Widgets/SWidget.h"
+#include "Widgets/Layout/LinkableScrollBar.h"
 
 class FPropertyPath;
 class IDetailsView;
@@ -37,6 +38,9 @@ public:
 	/** Perform a diff against another view, ordering either by display order or by remove/add/change */
 	void DiffAgainst(const FDetailsDiff& Newer, TArray<FSingleObjectDiffEntry>& OutDifferences, bool bSortByDisplayOrder = false) const;
 
+	/** Link the two details panels so they scroll in sync with one another */
+	static void LinkScrolling(FDetailsDiff& LeftPanel, FDetailsDiff& RightPanel, const TAttribute<TArray<FVector2f>>& ScrollRate);
+
 private:
 	void HandlePropertiesChanged();
 
@@ -46,4 +50,6 @@ private:
 	const UObject* DisplayedObject;
 
 	TSharedPtr< class IDetailsView > DetailsView;
+
+	TSharedPtr<SLinkableScrollBar> ScrollBar;
 };
diff --git a/Engine/Source/Editor/PropertyEditor/Private/SDetailsView.cpp b/Engine/Source/Editor/PropertyEditor/Private/SDetailsView.cpp
--- a/Engine/Source/Editor/PropertyEditor/Private/SDetailsView.cpp
+++ b/Engine/Source/Editor/PropertyEditor/Private/SDetailsView.cpp
@@ -86,7 +86,7 @@ void SDetailsView::Construct(const FArguments& InArgs, const FDetailsViewArgs& I
 	//  - We can't use AlwaysShowScrollbar for this, as this will also show the scrollbar when nothing is selected
 	//  - We can't use the Visibility construction parameter, as it gets translated into user visibility and can hide the scrollbar even when objects are selected
 	// We instead have to explicitly set the visibility after the scrollbar has been constructed to get the exact behavior we want
-	TSharedRef<SScrollBar> ExternalScrollbar = SNew(SScrollBar);
+	const TSharedRef<SScrollBar> ExternalScrollbar = InDetailsViewArgs.ExternalScrollbar ? InDetailsViewArgs.ExternalScrollbar.ToSharedRef() : SNew(SScrollBar);
 	ExternalScrollbar->SetVisibility( TAttribute<EVisibility>( this, &SDetailsView::GetScrollBarVisibility ) );
 
 	FMenuBuilder DetailViewOptions( true, nullptr);
@@ -427,7 +427,7 @@ void SDetailsView::Construct(const FArguments& InArgs, const FDetailsViewArgs& I
 	];
 }
 
-TSharedRef<SDetailTree> SDetailsView::ConstructTreeView( TSharedRef<SScrollBar>& ScrollBar )
+TSharedRef<SDetailTree> SDetailsView::ConstructTreeView(const TSharedRef<SScrollBar>& ScrollBar )
 {
 	check( !DetailTree.IsValid() || DetailTree.IsUnique() );
 
diff --git a/Engine/Source/Editor/PropertyEditor/Private/SDetailsView.h b/Engine/Source/Editor/PropertyEditor/Private/SDetailsView.h
--- a/Engine/Source/Editor/PropertyEditor/Private/SDetailsView.h
+++ b/Engine/Source/Editor/PropertyEditor/Private/SDetailsView.h
@@ -121,7 +121,7 @@ private:
 
 	void SetObjectArrayPrivate(const TArray<UObject*>& InObjects);
 
-	TSharedRef<SDetailTree> ConstructTreeView( TSharedRef<SScrollBar>& ScrollBar );
+	TSharedRef<SDetailTree> ConstructTreeView(const TSharedRef<SScrollBar>& ScrollBar );
 
 	/**
 	 * Returns whether or not new objects need to be set. If the new objects being set are identical to the objects 
diff --git a/Engine/Source/Editor/PropertyEditor/Private/SDetailsViewBase.cpp b/Engine/Source/Editor/PropertyEditor/Private/SDetailsViewBase.cpp
--- a/Engine/Source/Editor/PropertyEditor/Private/SDetailsViewBase.cpp
+++ b/Engine/Source/Editor/PropertyEditor/Private/SDetailsViewBase.cpp
@@ -207,6 +207,58 @@ TArray< FPropertyPath > SDetailsViewBase::GetPropertiesInOrderDisplayed() const
 	return Ret;
 }
 
+static void GetPropertyRowRowNumbersRecursive(const TArray<TSharedRef<FDetailTreeNode>>& TreeNodes, int32& CurrentRowNumber, const TSharedPtr<SDetailTree>& DetailTree, TArray<TPair<int32, FPropertyPath>>& OutResults)
+{
+	for (const TSharedRef<FDetailTreeNode>& TreeNode : TreeNodes)
+	{
+		// pre-order traversal (add parents to list before their children)
+		FPropertyPath Path = TreeNode->GetPropertyPath();
+
+		if(Path.IsValid())
+		{
+			OutResults.Add(TPair<int32, FPropertyPath>(CurrentRowNumber, Path));
+		}
+		++CurrentRowNumber;
+
+		if (!TreeNode->IsLeaf() && DetailTree->IsItemExpanded(TreeNode))
+		{
+			TArray<TSharedRef<FDetailTreeNode>> Children;
+			TreeNode->GetChildren(Children);
+			GetPropertyRowRowNumbersRecursive(Children, CurrentRowNumber, DetailTree, OutResults);
+		}
+	}
+}
+
+TArray<TPair<int32, FPropertyPath>> SDetailsViewBase::GetPropertyRowNumbers() const
+{
+	TArray<TPair<int32, FPropertyPath>> Results;
+	int32 OverallIdx = 0;
+	GetPropertyRowRowNumbersRecursive(RootTreeNodes, OverallIdx, DetailTree, Results);
+	return Results;
+}
+
+static int32 CountRowsRecursive(const TArray<TSharedRef<FDetailTreeNode>>& TreeNodes, const TSharedPtr<SDetailTree>& DetailTree)
+{
+	int32 Count = 0;
+	for (const TSharedRef<FDetailTreeNode>& TreeNode : TreeNodes)
+	{
+		++Count;
+		
+		if (!TreeNode->IsLeaf() && DetailTree->IsItemExpanded(TreeNode))
+		{
+			TArray<TSharedRef<FDetailTreeNode>> Children;
+			TreeNode->GetChildren(Children);
+			Count += CountRowsRecursive(Children, DetailTree);
+		}
+	}
+	return Count;
+}
+
+int32 SDetailsViewBase::CountRows() const
+{
+	return CountRowsRecursive(RootTreeNodes, DetailTree);
+}
+
 // construct a map for fast FPropertyPath.ToString() -> FDetailTreeNode lookup
 static void GetPropertyPathToTreeNodes(const TArray<TSharedRef<FDetailTreeNode>>& RootTreeNodes, TMap<FString, TSharedRef<FDetailTreeNode>> &OutMap)
 {
diff --git a/Engine/Source/Editor/PropertyEditor/Private/SDetailsViewBase.h b/Engine/Source/Editor/PropertyEditor/Private/SDetailsViewBase.h
--- a/Engine/Source/Editor/PropertyEditor/Private/SDetailsViewBase.h
+++ b/Engine/Source/Editor/PropertyEditor/Private/SDetailsViewBase.h
@@ -85,6 +85,8 @@ public:
 	 * IDetailsView interface 
 	 */
 	virtual TArray< FPropertyPath > GetPropertiesInOrderDisplayed() const override;
+	virtual TArray<TPair<int32, FPropertyPath>> GetPropertyRowNumbers() const override;
+	virtual int32 CountRows() const override;
 	virtual void HighlightProperty(const FPropertyPath& Property) override;
 	virtual void ShowAllAdvancedProperties() override;
 	virtual void SetOnDisplayedPropertiesChanged(FOnDisplayedPropertiesChanged InOnDisplayedPropertiesChangedDelegate) override;
diff --git a/Engine/Source/Editor/PropertyEditor/Public/DetailsViewArgs.h b/Engine/Source/Editor/PropertyEditor/Public/DetailsViewArgs.h
--- a/Engine/Source/Editor/PropertyEditor/Public/DetailsViewArgs.h
+++ b/Engine/Source/Editor/PropertyEditor/Public/DetailsViewArgs.h
@@ -3,6 +3,7 @@
 #pragma once
 
 #include "CoreMinimal.h"
+#include "Widgets/Layout/SScrollBar.h"
 
 class FDetailsViewObjectFilter;
 class FNotifyHook;
@@ -47,6 +48,8 @@ struct FDetailsViewArgs
 	TSharedPtr<FDetailsViewObjectFilter> ObjectFilter;
 	/** Optional custom filter(s) to apply to the class viewer widget for class object property values. */
 	TArray<TSharedRef<IClassViewerFilter>> ClassViewerFilters;
+	/** Externally constructed Scrollbar (internally constructed if not provided) */
+	TSharedPtr<SScrollBar> ExternalScrollbar;
 
 	/** Identifier for this details view; NAME_None if this view is anonymous */
 	FName ViewIdentifier;
diff --git a/Engine/Source/Editor/PropertyEditor/Public/IDetailsView.h b/Engine/Source/Editor/PropertyEditor/Public/IDetailsView.h
--- a/Engine/Source/Editor/PropertyEditor/Public/IDetailsView.h
+++ b/Engine/Source/Editor/PropertyEditor/Public/IDetailsView.h
@@ -187,6 +187,18 @@ public:
 	 * Returns a list of all the properties displayed (via full path), order in list corresponds to draw order:
 	 */
 	virtual TArray< FPropertyPath > GetPropertiesInOrderDisplayed() const = 0;
+	
+	/**
+	 * Returns the row number in the underlying TableView for each property being displayed. Note this won't include
+	 * any properties that are apart of a collapsed section of the property tree because those properties don't have
+	 * row numbers. If you need a list of all properties independent of expansion state, use GetPropertiesInOrderDisplayed
+	 */
+	virtual TArray<TPair<int32, FPropertyPath>> GetPropertyRowNumbers() const  = 0;
+	
+	/**
+	 * returns the number of rows in the underlying TableView
+	 */
+	virtual int32 CountRows() const  = 0;
 
 	/**
 	 * Creates a box around the treenode corresponding to Property and scrolls the treenode into view
diff --git a/Engine/Source/Runtime/Slate/Private/Framework/FInvertiblePiecewiseLinearFunction.cpp b/Engine/Source/Runtime/Slate/Private/Framework/FInvertiblePiecewiseLinearFunction.cpp
new file mode 100644
--- /dev/null
+++ b/Engine/Source/Runtime/Slate/Private/Framework/FInvertiblePiecewiseLinearFunction.cpp
@@ -0,0 +1,129 @@
+// Copyright Epic Games, Inc. All Rights Reserved.
+
+#include "Framework/FInvertiblePiecewiseLinearFunction.h"
+
+FInvertiblePointSlopeFunction::FInvertiblePointSlopeFunction(FVector2f Point, float Slope)
+	: Point(Point)
+	, Slope(Slope)
+{}
+
+float FInvertiblePointSlopeFunction::SolveForY(float X) const
+{
+	const float M = Slope;
+	const float B = Point.Y - (M * Point.X);
+	// y = m*x + b
+	return M * X + B;
+}
+
+float FInvertiblePointSlopeFunction::SolveForX(float Y) const
+{
+	const float M = 1 / Slope;
+	const float B = Point.X - (M * Point.Y);
+	// x = m*y + b
+	return M * Y + B;
+}
+
+FInvertiblePiecewiseLinearFunction::FInvertiblePiecewiseLinearFunction()
+	: SubFunctions{{{0.f,0.f}, 1.f}}
+{}
+
+FInvertiblePiecewiseLinearFunction::FInvertiblePiecewiseLinearFunction(const TArray<FVector2f>& FixedPoints)
+{
+	if (FixedPoints.Num() <= 1)
+	{
+		SubFunctions.Emplace(FVector2f{0.f,0.f}, 1.f);
+		return;
+	}
+	
+	float PrevSlope = 0;
+	for(int32 Index = 0; (Index + 1) < FixedPoints.Num(); ++Index)
+	{
+		const FVector2f& Begin = FixedPoints[Index]; 
+		const FVector2f& End = FixedPoints[Index + 1];
+
+		// duplicate points are allowed but ignored
+		if(Begin == End)
+		{
+			continue;
+		}
+			
+		// for invertibility, both x and y must be increasing
+		check(Begin.X <= End.X && Begin.Y <= End.Y);
+			
+		const float Slope = (End.Y - Begin.Y) / (End.X - Begin.X);
+			
+		// Only add a subfunction if it changes the slope
+		if (PrevSlope != Slope)
+		{
+			SubFunctions.Emplace(Begin, Slope);
+		}
+		PrevSlope = Slope;
+	}
+}
+
+float FInvertiblePiecewiseLinearFunction::SolveForY(float X) const
+{
+	return FindSubFunctionAtX(X).SolveForY(X);
+}
+
+float FInvertiblePiecewiseLinearFunction::SolveForX(float Y) const
+{
+	return FindSubFunctionAtY(Y).SolveForX(Y);
+}
+
+const FInvertiblePiecewiseLinearFunction::SubFunction& FInvertiblePiecewiseLinearFunction::FindSubFunctionAtX(float X) const
+{
+	if(X < SubFunctions[0].Point.X)
+	{
+		return SubFunctions[0];
+	}
+		
+	int32 Start = 0;
+	int32 End = SubFunctions.Num() - 1;
+		
+	// Binary search for a subFunction with the correct bounds
+	while (Start <= End)
+	{
+		const int Mid = (Start + End) / 2;
+		// If K is found
+		if (SubFunctions[Mid].Point.X == X)
+		{
+			return SubFunctions[Mid];
+		}
+		else if (SubFunctions[Mid].Point.X < X)
+		{
+			Start = Mid + 1;
+		}
+		else
+		{
+			End = Mid - 1;
+		}
+	}
+
+	return SubFunctions[End];
+}
+
+const FInvertiblePiecewiseLinearFunction::SubFunction& FInvertiblePiecewiseLinearFunction::FindSubFunctionAtY(float Y) const
+{
+	if(Y < SubFunctions[0].Point.Y)
+	{
+		return SubFunctions[0];
+	}
+		
+	int Start = 0;
+	int End = SubFunctions.Num() - 1;
+		
+	// Binary search for a subFunction with the correct bounds
+	while (Start <= End) {
+		const int Mid = (Start + End) / 2;
+		// If K is found
+		if (SubFunctions[Mid].Point.Y == Y)
+			return SubFunctions[Mid];
+		else if (SubFunctions[Mid].Point.Y < Y)
+			Start = Mid + 1;
+		else
+			End = Mid - 1;
+	}
+		
+	return SubFunctions[End];
+}
diff --git a/Engine/Source/Runtime/Slate/Private/Framework/FInvertiblePiecewiseLinearFunction.h b/Engine/Source/Runtime/Slate/Private/Framework/FInvertiblePiecewiseLinearFunction.h
new file mode 100644
--- /dev/null
+++ b/Engine/Source/Runtime/Slate/Private/Framework/FInvertiblePiecewiseLinearFunction.h
@@ -0,0 +1,33 @@
+// Copyright Epic Games, Inc. All Rights Reserved.
+#pragma once
+#include "CoreMinimal.h"
+
+struct FInvertiblePointSlopeFunction
+{
+	FInvertiblePointSlopeFunction(FVector2f Point, float Slope);
+
+	float SolveForY(float X) const;
+	float SolveForX(float Y) const;
+
+	FVector2f Point;
+	float Slope;
+};
+
+class FInvertiblePiecewiseLinearFunction
+{
+public:
+	FInvertiblePiecewiseLinearFunction();
+
+	// construct using a list of fixed points that will be interpolated between
+	explicit FInvertiblePiecewiseLinearFunction(const TArray<FVector2f>& FixedPoints);
+
+	float SolveForY(float X) const;
+	float SolveForX(float Y) const;
+
+private:
+	using SubFunction = FInvertiblePointSlopeFunction;
+	const SubFunction& FindSubFunctionAtX(float X) const;
+	const SubFunction& FindSubFunctionAtY(float Y) const;
+
+	TArray<SubFunction> SubFunctions;
+};
diff --git a/Engine/Source/Runtime/Slate/Private/Widgets/Layout/LinkableScrollBar.cpp b/Engine/Source/Runtime/Slate/Private/Widgets/Layout/LinkableScrollBar.cpp
new file mode 100644
--- /dev/null
+++ b/Engine/Source/Runtime/Slate/Private/Widgets/Layout/LinkableScrollBar.cpp
@@ -0,0 +1,71 @@
+// Copyright Epic Games, Inc. All Rights Reserved.
+
+#include "Widgets/Layout/LinkableScrollBar.h"
+
+#include "Widgets/Layout/SScrollBarTrack.h"
+#include "Framework/FInvertiblePiecewiseLinearFunction.h"
+
+void SLinkableScrollBar::SetState(float InOffsetFraction, float InThumbSizeFraction, bool bCallOnUserScrolled)
+{
+	// fallback to default behavior if this isn't linked to another scrollbar
+	if (!LinkedScrollBar.IsValid() || !ScrollSyncRate.IsSet())
+	{
+		SScrollBar::SetState(InOffsetFraction, InThumbSizeFraction, bCallOnUserScrolled);
+		return;
+	}
+
+	const float PrevDistanceFromTop = FMath::Clamp(Track->DistanceFromTop(), 0.f, 1.f);
+	const float PrevDistanceFromBottom = FMath::Clamp(Track->DistanceFromBottom(), 0.f, 1.f);
+	
+	SScrollBar::SetState(InOffsetFraction, InThumbSizeFraction, bCallOnUserScrolled);
+	
+	const float NewDistanceFromTop = FMath::Clamp(Track->DistanceFromTop(), 0.f, 1.f);
+	const float NewDistanceFromBottom = FMath::Clamp(Track->DistanceFromBottom(), 0.f, 1.f);
+	const bool bDirty = NewDistanceFromTop != PrevDistanceFromTop && NewDistanceFromBottom != PrevDistanceFromBottom;
+	
+	// if another scrollbar is linked to this one, Scroll it aswell
+	if (bDirty)
+	{
+		const TSharedPtr<SScrollBar> OtherScrollbar = LinkedScrollBar.Pin();
+		const FInvertiblePiecewiseLinearFunction ScrollRate = FInvertiblePiecewiseLinearFunction(ScrollSyncRate.Get());
+		
+		float OtherScrollOffset = -1.f;
+		switch(LinkId)
+		{
+		case ELinkId::LeftLink:
+			OtherScrollOffset = ScrollRate.SolveForY(InOffsetFraction);
+			break;
+		case ELinkId::RightLink:
+			OtherScrollOffset = ScrollRate.SolveForX(InOffsetFraction);
+			break;
+		
+		case ELinkId::Disabled:;
+		default: ;
+		}
+
+		if (NewDistanceFromTop > PrevDistanceFromTop) // scrolling down
+		{
+			if (FMath::IsNearlyZero(OtherScrollbar->DistanceFromBottom(), 0.01f))
+			{
+				return;
+			}
+		}
+		
+		if (OtherScrollOffset <= 1.f && OtherScrollOffset >= 0.f)
+		{
+			OtherScrollbar->SetState(OtherScrollOffset, OtherScrollbar->ThumbSizeFraction(), true);
+		}
+	}
+}
+
+void SLinkableScrollBar::LinkScrollBars(TSharedRef<SLinkableScrollBar> Left, TSharedRef<SLinkableScrollBar> Right,
+	TAttribute<TArray<FVector2f>> ScrollSyncRate)
+{
+	Left->LinkId = ELinkId::LeftLink;
+	Left->LinkedScrollBar = Right.ToWeakPtr();
+	Left->ScrollSyncRate = ScrollSyncRate;
+	
+	Right->LinkId = ELinkId::RightLink;
+	Right->LinkedScrollBar = Left.ToWeakPtr();
+	Right->ScrollSyncRate = ScrollSyncRate;
+}
diff --git a/Engine/Source/Runtime/Slate/Private/Widgets/Layout/SScrollBar.cpp b/Engine/Source/Runtime/Slate/Private/Widgets/Layout/SScrollBar.cpp
--- a/Engine/Source/Runtime/Slate/Private/Widgets/Layout/SScrollBar.cpp
+++ b/Engine/Source/Runtime/Slate/Private/Widgets/Layout/SScrollBar.cpp
@@ -91,7 +91,7 @@ void SScrollBar::SetOnUserScrolled( const FOnUserScrolled& InHandler )
 	OnUserScrolled = InHandler;
 }
 
-void SScrollBar::SetState( float InOffsetFraction, float InThumbSizeFraction )
+void SScrollBar::SetState( float InOffsetFraction, float InThumbSizeFraction, bool bCallOnUserScrolled )
 {
 	if ( Track->DistanceFromTop() != InOffsetFraction || Track->GetThumbSizeFraction() != InThumbSizeFraction )
 	{
@@ -101,6 +101,10 @@ void SScrollBar::SetState( float InOffsetFraction, float InThumbSizeFraction )
 		GetVisibilityAttribute().UpdateValue();
 
 		LastInteractionTime = FSlateApplication::Get().GetCurrentTime();
+		if (bCallOnUserScrolled)
+		{
+			OnUserScrolled.ExecuteIfBound(InOffsetFraction);
+		}
 	}
 }
 
@@ -230,6 +234,11 @@ float SScrollBar::DistanceFromBottom() const
 	return Track->DistanceFromBottom();
 }
 
+float SScrollBar::ThumbSizeFraction() const
+{
+	return Track->GetThumbSizeFraction();
+}
+
 SScrollBar::SScrollBar()
 	: bDraggingThumb(false)
 	, DragGrabOffset( 0.0f )
diff --git a/Engine/Source/Runtime/Slate/Public/Widgets/Layout/LinkableScrollBar.h b/Engine/Source/Runtime/Slate/Public/Widgets/Layout/LinkableScrollBar.h
new file mode 100644
--- /dev/null
+++ b/Engine/Source/Runtime/Slate/Public/Widgets/Layout/LinkableScrollBar.h
@@ -0,0 +1,29 @@
+// Copyright Epic Games, Inc. All Rights Reserved.
+
+#pragma once
+
+#include "CoreMinimal.h"
+#include "Widgets/Layout/SScrollBar.h"
+
+class SLATE_API SLinkableScrollBar : public SScrollBar
+{
+public:
+	enum class ELinkId : uint8
+	{
+		LeftLink,
+		RightLink,
+		Disabled
+	};
+	
+	virtual void SetState(float InOffsetFraction, float InThumbSizeFraction, bool bCallOnUserScrolled = false) override;
+
+	static void LinkScrollBars(TSharedRef<SLinkableScrollBar> Left, TSharedRef<SLinkableScrollBar> Right, TAttribute<TArray<FVector2f>> ScrollSyncRate);
+
+private:
+	TWeakPtr<SLinkableScrollBar> LinkedScrollBar;
+	ELinkId LinkId = ELinkId::Disabled; // LeftLink is associated w/ x value in SyncScrollRate while RightLink is associated w/ y
+
+	// list of scroll distance pairs that determine variable scroll rate.
+	// scroll values of Left panel will match X components while the right panel will match Y components
+	TAttribute<TArray<FVector2f>> ScrollSyncRate;
+};
diff --git a/Engine/Source/Runtime/Slate/Public/Widgets/Layout/SScrollBar.h b/Engine/Source/Runtime/Slate/Public/Widgets/Layout/SScrollBar.h
--- a/Engine/Source/Runtime/Slate/Public/Widgets/Layout/SScrollBar.h
+++ b/Engine/Source/Runtime/Slate/Public/Widgets/Layout/SScrollBar.h
@@ -81,8 +81,9 @@ public:
 	 *
 	 * @param InOffsetFraction     Offset of the thumbnail from the top as a fraction of the total available scroll space.
 	 * @param InThumbSizeFraction  Size of thumbnail as a fraction of the total available scroll space.
+	 * @param bCallOnUserScrolled  If true, OnUserScrolled will be called with InOffsetFraction
 	 */
-	void SetState( float InOffsetFraction, float InThumbSizeFraction );
+	virtual void SetState( float InOffsetFraction, float InThumbSizeFraction, bool bCallOnUserScrolled = false);
 
 	// SWidget
 	virtual void Tick(const FGeometry& AllottedGeometry, const double InCurrentTime, const float InDeltaTime) override;
@@ -102,6 +103,9 @@ public:
 	/** @return normalized distance from bottom */
 	float DistanceFromBottom() const;
 
+	/** @return normalized percentage of track covered by thumb bar */
+	float ThumbSizeFraction() const;
+
 	/** @return the scrollbar's visibility as a product of internal rules and user-specified visibility */
 	EVisibility ShouldBeVisible() const;
 
