Add an ensure to immediately detect when a blueprint is enqueued for compilation while compilation is in progress

diff --git a/Engine/Source/Editor/Kismet/Private/BlueprintCompilationManager.cpp b/Engine/Source/Editor/Kismet/Private/BlueprintCompilationManager.cpp
--- a/Engine/Source/Editor/Kismet/Private/BlueprintCompilationManager.cpp
+++ b/Engine/Source/Editor/Kismet/Private/BlueprintCompilationManager.cpp
@@ -224,6 +224,19 @@ void FBlueprintCompilationManagerImpl::QueueForCompilation(const FBPCompileReque
 {
 	if(!CompileJob.UserData.BPToCompile->bQueuedForCompilation)
 	{
+		if(GCompilingBlueprint)
+		{
+			FString CurrentlyCompiling;
+			for (const FBPCompileRequestInternal& CompilerData : QueuedRequests)
+			{
+				CurrentlyCompiling += CompilerData.UserData.BPToCompile->GetName() + TEXT(" ");
+			}
+			ensureMsgf(false, 
+				TEXT("Attempting to enqueue %s for compile while compiling: %s"), 
+				*CompileJob.UserData.BPToCompile->GetName(),
+				*CurrentlyCompiling);
+		}
+
 		CompileJob.UserData.BPToCompile->bQueuedForCompilation = true;
 		QueuedRequests.Add(CompileJob);
 	}
