PR #9991: Fix Blueprint debugger crash when attempting to observe a TimerHandle variable (Contributed by [REDACTED_PERSON])

diff --git a/Engine/Source/Editor/Kismet/Private/BlueprintEditor.cpp b/Engine/Source/Editor/Kismet/Private/BlueprintEditor.cpp
--- a/Engine/Source/Editor/Kismet/Private/BlueprintEditor.cpp
+++ b/Engine/Source/Editor/Kismet/Private/BlueprintEditor.cpp
@@ -756,6 +756,7 @@ const FSlateBrush* FBlueprintEditor::GetGlyphForGraph(const UEdGraph* Graph, boo
 
 FSlateBrush const* FBlueprintEditor::GetVarIconAndColor(const UStruct* VarScope, FName VarName, FSlateColor& IconColorOut, FSlateBrush const*& SecondaryBrushOut, FSlateColor& SecondaryColorOut)
 {
+	SecondaryBrushOut = nullptr;
 	if (VarScope != nullptr)
 	{
 		FProperty* Property = FindFProperty<FProperty>(VarScope, VarName);
@@ -766,6 +767,7 @@ FSlateBrush const* FBlueprintEditor::GetVarIconAndColor(const UStruct* VarScope,
 
 FSlateBrush const* FBlueprintEditor::GetVarIconAndColorFromProperty(const FProperty* Property, FSlateColor& IconColorOut, FSlateBrush const*& SecondaryBrushOut, FSlateColor& SecondaryColorOut)
 {
+	SecondaryBrushOut = nullptr;
 	if (Property != nullptr)
 	{
 		const UEdGraphSchema_K2* K2Schema = GetDefault<UEdGraphSchema_K2>();
diff --git a/Engine/Source/Editor/Kismet/Private/Debugging/SKismetDebugTreeView.cpp b/Engine/Source/Editor/Kismet/Private/Debugging/SKismetDebugTreeView.cpp
--- a/Engine/Source/Editor/Kismet/Private/Debugging/SKismetDebugTreeView.cpp
+++ b/Engine/Source/Editor/Kismet/Private/Debugging/SKismetDebugTreeView.cpp
@@ -961,7 +961,7 @@ public:
 	{
 		FSlateColor BaseColor;
 		FSlateColor SecondaryColor;
-		FSlateBrush const* SecondaryIcon;
+		FSlateBrush const* SecondaryIcon = nullptr;
 		const FSlateBrush* Icon = FBlueprintEditor::GetVarIconAndColorFromProperty(
 			Data->Property.Get(),
 			BaseColor,
diff --git a/Engine/Source/Editor/Kismet/Private/SSCSEditor.cpp b/Engine/Source/Editor/Kismet/Private/SSCSEditor.cpp
--- a/Engine/Source/Editor/Kismet/Private/SSCSEditor.cpp
+++ b/Engine/Source/Editor/Kismet/Private/SSCSEditor.cpp
@@ -404,8 +404,8 @@ void FSCSRowDragDropOp::HoverTargetChanged()
 		}
 		else if (FProperty* VariableProperty = GetVariableProperty())
 		{
-			const FSlateBrush* PrimarySymbol;
-			const FSlateBrush* SecondarySymbol;
+			const FSlateBrush* PrimarySymbol = nullptr;
+			const FSlateBrush* SecondarySymbol = nullptr;
 			FSlateColor PrimaryColor;
 			FSlateColor SecondaryColor;
 			GetDefaultStatusSymbol(/*out*/ PrimarySymbol, /*out*/ PrimaryColor, /*out*/ SecondarySymbol, /*out*/ SecondaryColor);
diff --git a/Engine/Source/Editor/SubobjectEditor/Private/SSubobjectEditor.cpp b/Engine/Source/Editor/SubobjectEditor/Private/SSubobjectEditor.cpp
--- a/Engine/Source/Editor/SubobjectEditor/Private/SSubobjectEditor.cpp
+++ b/Engine/Source/Editor/SubobjectEditor/Private/SSubobjectEditor.cpp
@@ -558,8 +558,8 @@ void FSubobjectRowDragDropOp::HoverTargetChanged()
 		}
 		else if (FProperty* VariableProperty = GetVariableProperty())
 		{
-			const FSlateBrush* PrimarySymbol;
-			const FSlateBrush* SecondarySymbol;
+			const FSlateBrush* PrimarySymbol = nullptr;
+			const FSlateBrush* SecondarySymbol = nullptr;
 			FSlateColor PrimaryColor;
 			FSlateColor SecondaryColor;
 			GetDefaultStatusSymbol(/*out*/ PrimarySymbol, /*out*/ PrimaryColor, /*out*/ SecondarySymbol, /*out*/ SecondaryColor);
