Non-blocking Blueprint graph context menu invocation (experimental).

Change summary:
- When enabled, invoking the context menu in a Blueprint graph will not block the UI for long periods waiting for its menu content building to be completed. Instead, it is spread across multiple frames and can be interrupted or restarted at any time. This allows for a better workflow/UX for creators, while we also continue to improve the efficiency of the menu building process.
- Added a progress bar to the context menu UI, visible only when time slicing is active. This will update as the menu is being populated as a visual indicator.
- Added an advanced BP editor setting to allow non-blocking behavior to be enabled (marked as experimental).
- Couple of misc. tweaks to stat reporting and filter adds.

diff --git a/Engine/Source/Editor/BlueprintGraph/Private/BlueprintActionFilter.cpp b/Engine/Source/Editor/BlueprintGraph/Private/BlueprintActionFilter.cpp
--- a/Engine/Source/Editor/BlueprintGraph/Private/BlueprintActionFilter.cpp
+++ b/Engine/Source/Editor/BlueprintGraph/Private/BlueprintActionFilter.cpp
@@ -2245,7 +2245,7 @@ FBlueprintActionFilter::FBlueprintActionFilter(const EFlags InFlags /*= BPFILTER
 //------------------------------------------------------------------------------
 void FBlueprintActionFilter::AddUnique(TArray<FTargetClassFilterData>& ToArray, UClass* TargetClass)
 {
-	for (auto ClassData : ToArray)
+	for (const FTargetClassFilterData& ClassData : ToArray)
 	{
 		if (ClassData.TargetClass == TargetClass)
 		{
@@ -2261,7 +2261,7 @@ void FBlueprintActionFilter::Add(TArray<FTargetClassFilterData>& ToArray, UClass
 	TArray<FString> ClassHideCategories;
 	FEditorCategoryUtils::GetClassHideCategories(TargetClass, ClassHideCategories);
 	FTargetClassFilterData Data = { TargetClass, MoveTemp(ClassHideCategories) };
-	ToArray.Add(Data);
+	ToArray.Add(MoveTemp(Data));
 }
 
 //------------------------------------------------------------------------------
diff --git a/Engine/Source/Editor/BlueprintGraph/Private/BlueprintEditorSettings.cpp b/Engine/Source/Editor/BlueprintGraph/Private/BlueprintEditorSettings.cpp
--- a/Engine/Source/Editor/BlueprintGraph/Private/BlueprintEditorSettings.cpp
+++ b/Engine/Source/Editor/BlueprintGraph/Private/BlueprintEditorSettings.cpp
@@ -54,6 +54,7 @@ UBlueprintEditorSettings::UBlueprintEditorSettings(const FObjectInitializer& Obj
 	, bEnableNamespaceEditorFeatures(true)
 	// Experimental
 	, bFavorPureCastNodes(false)
+	, bEnableContextMenuTimeSlicing(false)
 	// Compiler Settings
 	, SaveOnCompile(SoC_Never)
 	, bJumpToNodeErrors(false)
diff --git a/Engine/Source/Editor/BlueprintGraph/Public/BlueprintEditorSettings.h b/Engine/Source/Editor/BlueprintGraph/Public/BlueprintEditorSettings.h
--- a/Engine/Source/Editor/BlueprintGraph/Public/BlueprintEditorSettings.h
+++ b/Engine/Source/Editor/BlueprintGraph/Public/BlueprintEditorSettings.h
@@ -198,6 +198,10 @@ public:
 
 // Experimental
 public:
+	/** If enabled, when the Blueprint graph context menu is invoked (e.g. by dragging off a pin), it will not block the UI while it populates the available actions list. */
+	UPROPERTY(EditAnywhere, config, Category = Experimental, meta = (DisplayName = "Enable Non-Blocking Context Menu"))
+	bool bEnableContextMenuTimeSlicing;
+
 	/** If enabled, then placed cast nodes will default to their "pure" form (meaning: without execution pins). */
 	UPROPERTY(EditAnywhere, config, AdvancedDisplay, Category = Experimental, meta = (DisplayName = "Default to Using Pure Cast Nodes"))
 	bool bFavorPureCastNodes;
diff --git a/Engine/Source/Editor/Kismet/Private/BlueprintActionMenuBuilder.cpp b/Engine/Source/Editor/Kismet/Private/BlueprintActionMenuBuilder.cpp
--- a/Engine/Source/Editor/Kismet/Private/BlueprintActionMenuBuilder.cpp
+++ b/Engine/Source/Editor/Kismet/Private/BlueprintActionMenuBuilder.cpp
@@ -268,6 +268,8 @@ namespace FBlueprintActionMenuBuilderImpl
 
 		/** Sets the grouping for menu items belonging to this section. */
 		void SetSectionSortOrder(int32 const MenuGrouping);
+		/** Gets the grouping for menu items belonging to this section. */
+		int32 GetSectionSortOrder() const;
 		
 		/**
 		 * Filters the supplied action and if it passes, spawns a new 
@@ -301,6 +303,60 @@ namespace FBlueprintActionMenuBuilderImpl
 		TMap<FProperty const*, TSharedPtr<FBlueprintDragDropMenuItem>> ConsolidatedProperties;
 	};
 
+	/** A utility for building the menu item list based on a set of action descriptors */
+	struct FMenuItemListAddHelper
+	{
+		/** Reset for a new menu build */
+		void Reset(int32 NewSize)
+		{
+			NextIndex = 0;
+			PendingActionList.Reset(NewSize);
+		}
+
+		/** Add a new pending action */
+		void AddPendingAction(FBlueprintActionInfo&& Action)
+		{
+			PendingActionList.Add(Forward<FBlueprintActionInfo>(Action));
+		}
+
+		/** @return the next pending action and advance */
+		FBlueprintActionInfo* GetNextAction()
+		{
+			return PendingActionList.IsValidIndex(NextIndex) ? &PendingActionList[NextIndex++] : nullptr;
+		}
+
+		/** @return the allocated size of the pending action list */
+		SIZE_T GetAllocatedSize() const
+		{
+			return PendingActionList.GetAllocatedSize();
+		}
+
+		/** @return the total number of actions that are still pending */
+		int32 GetNumPendingActions() const
+		{
+			return PendingActionList.IsValidIndex(NextIndex) ? PendingActionList.Num() - NextIndex : 0;
+		}
+
+		/** @return the total number of actions that were added to the pending list */
+		int32 GetNumTotalAddedActions() const
+		{
+			return PendingActionList.Num();
+		}
+
+	private:
+		/** Keeps track of the next action list item to process */
+		int32 NextIndex = 0;
+
+		/** All actions pending menu items for the current context */
+		TArray<FBlueprintActionInfo> PendingActionList;
+	};
+
+	static TAutoConsoleVariable<float> CVarBPContextMenuTimeSlicingThreshold(
+		TEXT("BP.ContextMenuTimeSlicingThreshold"),
+		0.025f,
+		TEXT("The amount of time (in seconds) allowed per frame for time slicing.")
+	);
+
 #if ENABLE_BLUEPRINT_ACTION_FILTER_PROFILING
 	static TAutoConsoleVariable<bool> CVarBPEnableActionMenuDumpToFile(
 		TEXT("BP.EnableActionMenuDumpToFile"),
@@ -335,6 +391,12 @@ void FBlueprintActionMenuBuilderImpl::FMenuSectionDefinition::SetSectionSortOrde
 {
 	ItemFactory.MenuGrouping = MenuGrouping;
 }
+
+//------------------------------------------------------------------------------
+int32 FBlueprintActionMenuBuilderImpl::FMenuSectionDefinition::GetSectionSortOrder() const
+{
+	return ItemFactory.MenuGrouping;
+}
 // 
 //------------------------------------------------------------------------------
 void FBlueprintActionMenuBuilderImpl::FMenuSectionDefinition::AddBoundMenuItems(FBlueprintActionInfo& DatabaseActionInfo, TArray<FFieldVariant> const& PerspectiveBindings, MenuItemList& MenuItemsOut)
@@ -473,6 +535,13 @@ void FBlueprintActionMenuBuilderImpl::FMenuSectionDefinition::Empty()
  * FBlueprintActionMenuBuilder
  ******************************************************************************/
 
+//------------------------------------------------------------------------------
+FBlueprintActionMenuBuilder::FBlueprintActionMenuBuilder(EConfigFlags ConfigFlags)
+{
+	bUsePendingActionList = !!(ConfigFlags & EConfigFlags::UseTimeSlicing);
+	MenuItemListAddHelper = MakeShared<FBlueprintActionMenuBuilderImpl::FMenuItemListAddHelper>();
+}
+
 //------------------------------------------------------------------------------
 // @todo_deprecated - Remove this definition along w/ its declaration in a future release.
 PRAGMA_DISABLE_DEPRECATION_WARNINGS
@@ -521,7 +590,12 @@ void FBlueprintActionMenuBuilder::RebuildActionList()
 #if ENABLE_BLUEPRINT_ACTION_FILTER_PROFILING
 	int32 TotalActionCount = 0;
 #endif	// ENABLE_BLUEPRINT_ACTION_FILTER_PROFILING
-	
+
+	if (bUsePendingActionList)
+	{
+		MenuItemListAddHelper->Reset(ActionRegistry.Num());
+	}
+
 	for (auto Iterator(ActionRegistry.CreateConstIterator()); Iterator; ++Iterator)
 	{
 		const FObjectKey& ObjKey = Iterator->Key;
@@ -537,12 +611,13 @@ void FBlueprintActionMenuBuilder::RebuildActionList()
 				++TotalActionCount;
 #endif	// ENABLE_BLUEPRINT_ACTION_FILTER_PROFILING
 
-				for (TSharedRef<FMenuSectionDefinition> const& MenuSection : MenuSections)
+				if (bUsePendingActionList)
 				{
-					for (TSharedPtr<FEdGraphSchemaAction> MenuEntry : MenuSection->MakeMenuItems(BlueprintAction))
-					{
-						AddAction(MenuEntry);
-					}
+					MenuItemListAddHelper->AddPendingAction(MoveTemp(BlueprintAction));
+				}
+				else
+				{
+					MakeMenuItems(BlueprintAction);
 				}
 			}
 		}
@@ -557,12 +632,16 @@ void FBlueprintActionMenuBuilder::RebuildActionList()
 	if (FBlueprintActionFilter::IsFilterTestStatsLoggingEnabled())
 	{
 		UE_LOG(LogBlueprintActionMenuItemFactory, Log, TEXT("=== UNFILTERED ACTIONS: %d"), TotalActionCount);
+		if (MenuItemListAddHelper->GetNumPendingActions() > 0)
+		{
+			UE_LOG(LogBlueprintActionMenuItemFactory, Log, TEXT("==  PENDING ACTION LIST MEMSIZE: %0.02f MB"), MenuItemListAddHelper->GetAllocatedSize() / 1024.0f / 1024.0f);
+		}
 
 		// Dump detailed stats information about each filter test that was involved with building each menu section.
 		for (int32 SectionIdx = 0; SectionIdx < MenuSections.Num(); ++SectionIdx)
 		{
 			const TSharedRef<FMenuSectionDefinition>& MenuSection = MenuSections[SectionIdx];
-			FString DisplayName = FString::Printf(TEXT("MenuSection[%d]"), SectionIdx);
+			FString DisplayName = FString::Printf(TEXT("MenuSection[%d:%d]"), SectionIdx, MenuSection->GetSectionSortOrder());
 			const FText& SectionHeading = MenuSection->GetSectionHeading();
 			if (!SectionHeading.IsEmptyOrWhitespace())
 			{
@@ -577,6 +656,80 @@ void FBlueprintActionMenuBuilder::RebuildActionList()
 			}
 		}
 	}
+#endif	// ENABLE_BLUEPRINT_ACTION_FILTER_PROFILING
+
+	if (MenuItemListAddHelper->GetNumPendingActions() == 0)
+	{
+		BuildCompleted();
+	}
+}
+
+void FBlueprintActionMenuBuilder::MakeMenuItems(FBlueprintActionInfo& InAction)
+{
+	using namespace FBlueprintActionMenuBuilderImpl;
+
+	for (TSharedRef<FMenuSectionDefinition> const& MenuSection : MenuSections)
+	{
+		for (TSharedPtr<FEdGraphSchemaAction> MenuEntry : MenuSection->MakeMenuItems(InAction))
+		{
+			AddAction(MenuEntry);
+		}
+	}
+}
+
+int32 FBlueprintActionMenuBuilder::GetNumPendingActions() const
+{
+	return MenuItemListAddHelper->GetNumPendingActions();
+}
+
+float FBlueprintActionMenuBuilder::GetPendingActionsProgress() const
+{
+	const float NumPendingActions = static_cast<float>(MenuItemListAddHelper->GetNumPendingActions());
+	const float NumTotalAddedActions = static_cast<float>(MenuItemListAddHelper->GetNumTotalAddedActions());
+
+	check(NumTotalAddedActions > 0.0f);
+	return 1.0f - (NumPendingActions / NumTotalAddedActions);
+}
+
+bool FBlueprintActionMenuBuilder::ProcessPendingActions()
+{
+	TRACE_CPUPROFILER_EVENT_SCOPE(FBlueprintActionMenuBuilder::ProcessPendingActions);
+
+	using namespace FBlueprintActionMenuBuilderImpl;
+
+	bool bProcessedActions = false;
+	double StartTime = FPlatformTime::Seconds();
+
+	FBlueprintActionInfo* CurrentAction = MenuItemListAddHelper->GetNextAction();
+	while (CurrentAction)
+	{
+		bProcessedActions = true;
+
+		MakeMenuItems(*CurrentAction);
+
+		const float MaxTimeThresholdSeconds = CVarBPContextMenuTimeSlicingThreshold.GetValueOnGameThread();
+		if ((FPlatformTime::Seconds() - StartTime) < MaxTimeThresholdSeconds)
+		{
+			CurrentAction = MenuItemListAddHelper->GetNextAction();
+		}
+		else
+		{
+			break;
+		}
+	}
+
+	if (bProcessedActions && !CurrentAction)
+	{
+		BuildCompleted();
+	}
+
+	return bProcessedActions;
+}
+
+void FBlueprintActionMenuBuilder::BuildCompleted()
+{
+#if ENABLE_BLUEPRINT_ACTION_FILTER_PROFILING
+	using namespace FBlueprintActionMenuBuilderImpl;
 
 	if (CVarBPEnableActionMenuDumpToFile.GetValueOnGameThread())
 	{
diff --git a/Engine/Source/Editor/Kismet/Private/SBlueprintActionMenu.cpp b/Engine/Source/Editor/Kismet/Private/SBlueprintActionMenu.cpp
--- a/Engine/Source/Editor/Kismet/Private/SBlueprintActionMenu.cpp
+++ b/Engine/Source/Editor/Kismet/Private/SBlueprintActionMenu.cpp
@@ -57,6 +57,7 @@
 #include "Widgets/Input/SCheckBox.h"
 #include "Widgets/Input/SComboButton.h"
 #include "Widgets/Layout/SBox.h"
+#include "Widgets/Notifications/SProgressBar.h"
 #include "Widgets/SBoxPanel.h"
 #include "Widgets/SCompoundWidget.h"
 #include "Widgets/SNullWidget.h"
@@ -416,11 +417,31 @@ void SBlueprintActionMenu::Construct( const FArguments& InArgs, TSharedPtr<FBlue
 					SAssignNew(GraphActionMenu, SGraphActionMenu)
 						.OnActionSelected(this, &SBlueprintActionMenu::OnActionSelected)
 						.OnCreateWidgetForAction(SGraphActionMenu::FOnCreateWidgetForAction::CreateSP(this, &SBlueprintActionMenu::OnCreateWidgetForAction))
-						.OnCollectAllActions(this, &SBlueprintActionMenu::CollectAllActions)
+						.OnGetActionList(this, &SBlueprintActionMenu::OnGetActionList)
 						.OnCreateCustomRowExpander_Static(&CreateCustomBlueprintActionExpander)
 						.DraggedFromPins(DraggedFromPins)
 						.GraphObj(GraphObj)
 				]
+
+				// PROGRESS BAR
+				+SVerticalBox::Slot()
+				.AutoHeight()
+				[
+					SNew(SBox)
+					.HeightOverride(2)
+					.Visibility_Lambda([this]()
+					{
+						return ContextMenuBuilder.IsValid() && ContextMenuBuilder->GetNumPendingActions() > 0 ? EVisibility::SelfHitTestInvisible : EVisibility::Collapsed;
+					})
+					[
+						SNew(SProgressBar)
+						.BorderPadding(FVector2D(0, 0))
+						.Percent_Lambda([this]()
+						{
+							return ContextMenuBuilder.IsValid() ? ContextMenuBuilder->GetPendingActionsProgress() : 0.0f;
+						})
+					]
+				]
 			]
 		]
 	);
@@ -506,7 +527,7 @@ ECheckBoxState SBlueprintActionMenu::ContextToggleIsChecked() const
 	return EditorPtr.Pin()->GetIsContextSensitive() ? ECheckBoxState::Checked : ECheckBoxState::Unchecked;
 }
 
-void SBlueprintActionMenu::CollectAllActions(FGraphActionListBuilderBase& OutAllActions)
+TSharedRef<FGraphActionListBuilderBase> SBlueprintActionMenu::OnGetActionList()
 {	
 	check(EditorPtr.IsValid());
 	TSharedPtr<FBlueprintEditor> BlueprintEditor = EditorPtr.Pin();
@@ -520,23 +541,31 @@ void SBlueprintActionMenu::CollectAllActions(FGraphActionListBuilderBase& OutAll
 
 	FBlueprintActionContext FilterContext;
 	ConstructActionContext(FilterContext);
-	
-	FBlueprintActionMenuBuilder MenuBuilder;
+
+	FBlueprintActionMenuBuilder::EConfigFlags ConfigFlags = FBlueprintActionMenuBuilder::DefaultConfig;
+	if (GetDefault<UBlueprintEditorSettings>()->bEnableContextMenuTimeSlicing)
+	{
+		ConfigFlags |= FBlueprintActionMenuBuilder::UseTimeSlicing;
+	}
+
+	ContextMenuBuilder = MakeShared<FBlueprintActionMenuBuilder>(ConfigFlags);
+
 	// NOTE: cannot call GetGraphContextActions() during serialization and GC due to its use of FindObject()
-	if(!GIsSavingPackage && !IsGarbageCollecting() && FilterContext.Blueprints.Num() > 0)
+	if (!GIsSavingPackage && !IsGarbageCollecting() && FilterContext.Blueprints.Num() > 0)
 	{
-		FBlueprintActionMenuUtils::MakeContextMenu(FilterContext, bIsContextSensitive, ContextTargetMask, MenuBuilder);
+		FBlueprintActionMenuUtils::MakeContextMenu(FilterContext, bIsContextSensitive, ContextTargetMask, *ContextMenuBuilder);
 	}
-	// copy the added options back to the main list
-	OutAllActions.Append(MenuBuilder); // @TODO: Avoid this copy
+	
 	// also try adding promote to variable if we can do so.
-	TryInsertPromoteToVariable(FilterContext, OutAllActions);
+	TryInsertPromoteToVariable(FilterContext, *ContextMenuBuilder);
 
 	// give the schema the opportunity to add another action
-	if(const UEdGraphSchema* Schema = Cast<const UEdGraphSchema>(GraphObj->GetSchema()))
+	if (const UEdGraphSchema* Schema = Cast<const UEdGraphSchema>(GraphObj->GetSchema()))
 	{
-		Schema->InsertAdditionalActions(FilterContext.Blueprints, FilterContext.Graphs, FilterContext.Pins, OutAllActions);
+		Schema->InsertAdditionalActions(FilterContext.Blueprints, FilterContext.Graphs, FilterContext.Pins, *ContextMenuBuilder);
 	}
+
+	return ContextMenuBuilder.ToSharedRef();
 }
 
 void SBlueprintActionMenu::ConstructActionContext(FBlueprintActionContext& ContextDescOut)
@@ -702,6 +731,16 @@ void SBlueprintActionMenu::OnNamespaceSelectedForImport(const FString& InNamespa
 	}
 }
 
+void SBlueprintActionMenu::Tick(const FGeometry& AllottedGeometry, const double InCurrentTime, const float InDeltaTime)
+{
+	SBorder::Tick(AllottedGeometry, InCurrentTime, InDeltaTime);
+
+	if (ContextMenuBuilder.IsValid() && ContextMenuBuilder->ProcessPendingActions())
+	{
+		GraphActionMenu->GenerateFilteredItems(true);
+	}
+}
+
 /*******************************************************************************
 * SBlueprintActionMenuExpander
 *******************************************************************************/
diff --git a/Engine/Source/Editor/Kismet/Private/SBlueprintActionMenu.h b/Engine/Source/Editor/Kismet/Private/SBlueprintActionMenu.h
--- a/Engine/Source/Editor/Kismet/Private/SBlueprintActionMenu.h
+++ b/Engine/Source/Editor/Kismet/Private/SBlueprintActionMenu.h
@@ -32,6 +32,7 @@ struct FBlueprintActionContext;
 struct FCustomExpanderData;
 struct FEdGraphSchemaAction;
 struct FGraphActionListBuilderBase;
+struct FBlueprintActionMenuBuilder;
 
 /*******************************************************************************
 * SBlueprintActionMenu
@@ -63,6 +64,10 @@ public:
 
 	TSharedRef<SEditableTextBox> GetFilterTextBox();
 
+	// SWidget interface
+	virtual void Tick(const FGeometry& AllottedGeometry, const double InCurrentTime, const float InDeltaTime) override;
+	// End of SWidget interface
+
 protected:
 	/** UI Callback functions */
 	EVisibility GetTypeImageVisibility() const;
@@ -76,7 +81,7 @@ protected:
 	TSharedRef<SWidget> OnCreateWidgetForAction(struct FCreateWidgetForActionData* const InCreateData);
 
 	/** Callback used to populate all actions list in SGraphActionMenu */
-	void CollectAllActions(FGraphActionListBuilderBase& OutAllActions);
+	TSharedRef<FGraphActionListBuilderBase> OnGetActionList();
 
 	/**  */
 	void ConstructActionContext(FBlueprintActionContext& ContextDescOut);
@@ -102,6 +107,7 @@ private:
 	TSharedPtr<SGraphActionMenu> GraphActionMenu;
 	TWeakPtr<FBlueprintEditor> EditorPtr;
 	TSharedPtr<SBlueprintContextTargetMenu> ContextTargetSubMenu;
+	TSharedPtr<FBlueprintActionMenuBuilder> ContextMenuBuilder;
 
 	bool bActionExecuted;
 };
diff --git a/Engine/Source/Editor/Kismet/Public/BlueprintActionMenuBuilder.h b/Engine/Source/Editor/Kismet/Public/BlueprintActionMenuBuilder.h
--- a/Engine/Source/Editor/Kismet/Public/BlueprintActionMenuBuilder.h
+++ b/Engine/Source/Editor/Kismet/Public/BlueprintActionMenuBuilder.h
@@ -14,12 +14,14 @@ class FBlueprintEditor;
 
 namespace FBlueprintActionMenuBuilderImpl
 {
-	// internal type, forward declared to hide implementation details
+	// internal types, forward declared to hide implementation details
 	struct FMenuSectionDefinition;
+	struct FMenuItemListAddHelper;
 };
 class  FBlueprintActionFilter;
 class  FBlueprintEditor;
 struct FBlueprintActionContext;
+struct FBlueprintActionInfo;
 
 /**
  * Responsible for constructing a list of viable blueprint actions. Runs the 
@@ -30,6 +32,18 @@ struct FBlueprintActionContext;
 struct KISMET_API FBlueprintActionMenuBuilder : public FGraphActionListBuilderBase
 {
 public:
+	/** Flags used to configure the builder. */
+	enum EConfigFlags
+	{
+		// If set, the builder will be configured to process the registered
+		// action set over time, rather than process the entire set in a
+		// single blocking frame.
+		UseTimeSlicing		= (1<<0),
+
+		// The default configuration for this builder type.
+		DefaultConfig		= 0,
+	};
+
 	/** Flags used to customize specific sections of the menu. */
 	enum ESectionFlag
 	{
@@ -46,7 +60,12 @@ public:
 	};
 	
 public:
-	FBlueprintActionMenuBuilder() = default;
+	/**
+	 * Default constructor.
+	 * 
+	 * @param Flags	- Optional configuration flags.
+	 */
+	FBlueprintActionMenuBuilder(EConfigFlags Flags = EConfigFlags::DefaultConfig);
 
 	UE_DEPRECATED(5.2, "The action filter now stores a reference to the authoritative editor context. Please use the default constructor instead.")
 	FBlueprintActionMenuBuilder(TWeakPtr<FBlueprintEditor> BlueprintEditorPtr);
@@ -76,6 +95,22 @@ public:
 	 */
 	void RebuildActionList();
 
+	/** Returns the current number of actions that are still pending */
+	int32 GetNumPendingActions() const;
+
+	/** Processes any actions that may be added asynchronously or across multiple frames. Returns true if one or more actions were added into the list. */
+	bool ProcessPendingActions();
+
+	/** Returns the normalized completion state when processing pending actions (e.g. for a status indicator) */
+	float GetPendingActionsProgress() const;
+
+protected:
+	// Adds menu items for the given database action.
+	void MakeMenuItems(FBlueprintActionInfo& InAction);
+
+	// Called when all menu items have been constructed.
+	void BuildCompleted();
+
 private:
 	/** 
 	 * Defines all the separate sections of the menu (filter, sort order, etc.).
@@ -83,4 +118,15 @@ private:
 	 * this API clean).
 	 */
 	TArray< TSharedRef<FBlueprintActionMenuBuilderImpl::FMenuSectionDefinition> > MenuSections;
+
+	/**
+	 * Defines a utility that assists with building the list of action menu items
+	 * for each menu section based on a set of Blueprint action descriptor records.
+	 */
+	TSharedPtr<FBlueprintActionMenuBuilderImpl::FMenuItemListAddHelper> MenuItemListAddHelper;
+
+	/** If enabled, actions will be added to the pending list rather than processed immediately. */
+	bool bUsePendingActionList;
 };
+
+ENUM_CLASS_FLAGS(FBlueprintActionMenuBuilder::EConfigFlags);
\ No newline at end of file
