Adjust the default BP graph context menu time slicing threshold to allow ~2x more actions per frame.

#rnx

diff --git a/Engine/Source/Editor/Kismet/Private/BlueprintActionMenuBuilder.cpp b/Engine/Source/Editor/Kismet/Private/BlueprintActionMenuBuilder.cpp
--- a/Engine/Source/Editor/Kismet/Private/BlueprintActionMenuBuilder.cpp
+++ b/Engine/Source/Editor/Kismet/Private/BlueprintActionMenuBuilder.cpp
@@ -353,7 +353,7 @@ namespace FBlueprintActionMenuBuilderImpl
 
 	static TAutoConsoleVariable<float> CVarBPContextMenuTimeSlicingThreshold(
 		TEXT("BP.ContextMenuTimeSlicingThreshold"),
-		0.025f,
+		0.05f,
 		TEXT("The amount of time (in seconds) allowed per frame for time slicing.")
 	);
 
