[BugFix] Blueprint Collapse Nodes results in random/unexpected pin order

diff --git a/Engine/Source/Editor/Kismet/Private/BlueprintEditor.cpp b/Engine/Source/Editor/Kismet/Private/BlueprintEditor.cpp
--- a/Engine/Source/Editor/Kismet/Private/BlueprintEditor.cpp
+++ b/Engine/Source/Editor/Kismet/Private/BlueprintEditor.cpp
@@ -179,6 +179,7 @@
 #include "AnimationBlendSpaceSampleGraph.h"
 #include "SSubobjectEditor.h"
 #include "BlueprintActionDatabase.h"
+#include "Algo/MinElement.h"
 #include "Editor/EditorEngine.h"
 
 DEFINE_LOG_CATEGORY_STATIC(LogBlueprintEditor, Log, All);
@@ -8407,6 +8408,63 @@ void FBlueprintEditor::ExtractEventTemplateForFunction(class UK2Node_CustomEvent
 	}
 }
 
+static void SortPinsByConnectionPosition(TArray<UEdGraphPin*>& Pins)
+{
+	// Maps a pin to it's index in the owning node
+	TMap<UEdGraphPin*, int32> GetPinIndexCache;
+	auto GetPinIndex = [&GetPinIndexCache](const UEdGraphPin* Pin)
+	{
+		if (!GetPinIndexCache.Contains(Pin))
+		{
+			const UEdGraphNode* Parent = Pin->GetOwningNode();
+			for (int I = 0; I < Parent->Pins.Num(); ++I)
+			{
+				// since we're already iterating all the siblings,
+				// go ahead and store them as well for faster lookup
+				GetPinIndexCache.Add(Parent->Pins[I], I);
+			}
+		}
+		return GetPinIndexCache[Pin];
+	};
+
+	// determines whether the Y value of PinA is less than that of PinB
+	auto PinYComp = [&GetPinIndex](const UEdGraphPin* PinA, const UEdGraphPin* PinB){
+		const float NodeAPosY = PinA->GetOwningNode()->NodePosY;
+		const float NodeBPosY = PinB->GetOwningNode()->NodePosY;
+		const float NodeAPosX = PinA->GetOwningNode()->NodePosX;
+		const float NodeBPosX = PinB->GetOwningNode()->NodePosX;
+
+		// Exec pins should always appear first
+		if (UEdGraphSchema_K2::IsExecPin(*PinA))
+		{
+			return true;
+		}
+		if (UEdGraphSchema_K2::IsExecPin(*PinB))
+		{
+			return false;
+		}
+
+		// sort by Y, then X, then Pin index
+		if (NodeAPosY != NodeBPosY)
+		{
+			return NodeAPosY < NodeBPosY;
+		}
+		if (NodeAPosX != NodeBPosX)
+		{
+			return NodeAPosX < NodeBPosX;
+		}
+		return GetPinIndex(PinA) < GetPinIndex(PinB);
+	};
+
+	Pins.Sort([&PinYComp](const UEdGraphPin& PinA, const UEdGraphPin& PinB)
+	{
+		// since we care more about making the calling convention pretty
+		// and less about making the internal implementation pretty, we'll
+		// sort by linked pins instead of the gateway pins directly.
+		return PinYComp(*Algo::MinElement(PinA.LinkedTo, PinYComp), *Algo::MinElement(PinB.LinkedTo, PinYComp));
+	});
+}
+
 void FBlueprintEditor::CollapseNodesIntoGraph(UEdGraphNode* InGatewayNode, UK2Node_EditablePinBase* InEntryNode, UK2Node_EditablePinBase* InResultNode, UEdGraph* InSourceGraph, UEdGraph* InDestinationGraph, TSet<UEdGraphNode*>& InCollapsableNodes, bool bCanDiscardEmptyReturnNode, bool bCanHaveWeakObjPtrParam)
 {
 	const UEdGraphSchema_K2* K2Schema = GetDefault<UEdGraphSchema_K2>();
@@ -8451,6 +8509,8 @@ void FBlueprintEditor::CollapseNodesIntoGraph(UEdGraphNode* InGatewayNode, UK2No
 			}
 		}
 	}
+	
+	TArray<UEdGraphPin*> GatewayPins;
 
 	// Move the nodes over, which may create cross-graph references that we need fix up ASAP
 	for (TSet<UEdGraphNode*>::TConstIterator NodeIt(InCollapsableNodes); NodeIt; ++NodeIt)
@@ -8515,114 +8575,122 @@ void FBlueprintEditor::CollapseNodesIntoGraph(UEdGraphNode* InGatewayNode, UK2No
 				}
 			}
 
-			// Thunk cross-graph links thru the gateway
 			if (bIsGatewayPin)
 			{
-				// Local port is either the entry or the result node in the collapsed graph
-				// Remote port is the node placed in the source graph
-				UK2Node_EditablePinBase* LocalPort = (LocalPin->Direction == EGPD_Input) ? InEntryNode : InResultNode;
+				GatewayPins.Add(LocalPin);
+			}
+		}
 
-				// Add a new pin to the entry/exit node and to the composite node
-				UEdGraphPin* LocalPortPin = nullptr;
-				UEdGraphPin* RemotePortPin = nullptr;
+		if (OutputGatewayExecPins.Num() > 0)
+		{
+			UEdGraphPin* LocalResultPortPin = K2Schema->FindExecutionPin(*InResultNode, EGPD_Input);
 
-				// Function graphs have a single exec path through them, so only one exec pin for input and another for output. In this fashion, they must not be handled by name.
-				if (InGatewayNode->GetClass() == UK2Node_CallFunction::StaticClass() && LocalPin->PinType.PinCategory == UEdGraphSchema_K2::PC_Exec)
+			// If the Result Node already contains links, then we don't need to make these connections as the intended connections have already been
+			// transferred from original graph.
+			if (LocalResultPortPin != nullptr && LocalResultPortPin->LinkedTo.Num() == 0)
+			{
+				// TODO: Some of these pins may not necessarily be terminal pins. We should prompt the user to choose which of these connections should
+				// be made to the return node.
+				for (UEdGraphPin* LocalPin : OutputGatewayExecPins)
 				{
-					LocalPortPin = LocalPort->Pins[0];
-					RemotePortPin = K2Schema->FindExecutionPin(*InGatewayNode, (LocalPortPin->Direction == EGPD_Input)? EGPD_Output : EGPD_Input);
+					// Connect the gateway pin to the node, there is no remote pin to hook up because the exec pin was not originally connected
+					LocalPin->Modify();
+					LocalPin->MakeLinkTo(LocalResultPortPin);
 				}
-				else
-				{
-					// If there is a custom event being used as a template, we must check to see if any connected pins have already been built
-					if (InterfaceTemplateNode && LocalPin->Direction == EGPD_Input)
-					{
-						// Find the pin on the entry node, we will use that pin's name to find the pin on the remote port
-						UEdGraphPin* EntryNodePin = InEntryNode->FindPin(LocalPin->LinkedTo[0]->PinName);
-						if(EntryNodePin)
-						{
-							LocalPin->BreakAllPinLinks();
-							LocalPin->MakeLinkTo(EntryNodePin);
-							continue;
-						}
-					}
-
-					if (LocalPin->LinkedTo[0]->GetOwningNode() != InEntryNode)
-					{
-						const FName UniquePortName = InGatewayNode->CreateUniquePinName(LocalPin->PinName);
+			}
+		}
+	}
 
-						if (!RemotePortPin && !LocalPortPin)
-						{
-							if (LocalPin->Direction == EGPD_Output)
-							{
-								bDiscardReturnNode = false;
-							}
-							
-							FEdGraphPinType PinType = LocalPin->PinType;
-							if (PinType.bIsWeakPointer && !PinType.IsContainer() && !bCanHaveWeakObjPtrParam)
-							{
-								PinType.bIsWeakPointer = false;
-							}
-							RemotePortPin = InGatewayNode->CreatePin(LocalPin->Direction, PinType, UniquePortName);
-							LocalPortPin = LocalPort->CreateUserDefinedPin(UniquePortName, PinType, (LocalPin->Direction == EGPD_Input) ? EGPD_Output : EGPD_Input);
-						}
-					}
-				}
+	// put all exec pins first, then sort by Y, then X, then Pin index
+	SortPinsByConnectionPosition(GatewayPins);
 
-				check(LocalPortPin);
-				check(RemotePortPin);
+	// Thunk cross-graph links thru the gateway
+	for (UEdGraphPin* LocalPin : GatewayPins)
+	{
+		// Local port is either the entry or the result node in the collapsed graph
+		// Remote port is the node placed in the source graph
+		UK2Node_EditablePinBase* LocalPort = (LocalPin->Direction == EGPD_Input) ? InEntryNode : InResultNode;
 
-				LocalPin->Modify();
+		// Add a new pin to the entry/exit node and to the composite node
+		UEdGraphPin* LocalPortPin = nullptr;
+		UEdGraphPin* RemotePortPin = nullptr;
 
-				// Route the links
-				for (int32 LinkIndex = 0; LinkIndex < LocalPin->LinkedTo.Num(); ++LinkIndex)
+		// Function graphs have a single exec path through them, so only one exec pin for input and another for output. In this fashion, they must not be handled by name.
+		if (InGatewayNode->GetClass() == UK2Node_CallFunction::StaticClass() && LocalPin->PinType.PinCategory == UEdGraphSchema_K2::PC_Exec)
+		{
+			LocalPortPin = LocalPort->Pins[0];
+			RemotePortPin = K2Schema->FindExecutionPin(*InGatewayNode, (LocalPortPin->Direction == EGPD_Input)? EGPD_Output : EGPD_Input);
+		}
+		else
+		{
+			// If there is a custom event being used as a template, we must check to see if any connected pins have already been built
+			if (InterfaceTemplateNode && LocalPin->Direction == EGPD_Input)
+			{
+				// Find the pin on the entry node, we will use that pin's name to find the pin on the remote port
+				UEdGraphPin* EntryNodePin = InEntryNode->FindPin(LocalPin->LinkedTo[0]->PinName);
+				if(EntryNodePin)
 				{
-					UEdGraphPin* RemotePin = LocalPin->LinkedTo[LinkIndex];
-					RemotePin->Modify();
-
-					if (!InCollapsableNodes.Contains(RemotePin->GetOwningNode()) && RemotePin->GetOwningNode() != InEntryNode && RemotePin->GetOwningNode() != InResultNode)
-					{
-						// Fix up the remote pin
-						RemotePin->LinkedTo.Remove(LocalPin);
-						// When given a composite node, we could possibly be given a pin with a different outer
-						// which is bad! Then there would be a pin connecting to itself and cause an ensure
-						if (RemotePin->GetOwningNode()->GetOuter() == RemotePortPin->GetOwningNode()->GetOuter())
-						{
-							RemotePin->MakeLinkTo(RemotePortPin);
-						}
+					LocalPin->BreakAllPinLinks();
+					LocalPin->MakeLinkTo(EntryNodePin);
+					continue;
+				}
+			}
 
-						// The Entry Node only supports a single link, so if we made links above
-						// we need to break them now, to make room for the new link.
-						if (LocalPort == InEntryNode)
-						{
-							LocalPortPin->BreakAllPinLinks();
-						}
+			if (LocalPin->LinkedTo[0]->GetOwningNode() != InEntryNode)
+			{
+				const FName UniquePortName = InGatewayNode->CreateUniquePinName(LocalPin->PinName);
 
-						// Fix up the local pin
-						LocalPin->LinkedTo.Remove(RemotePin);
-						--LinkIndex;
-						LocalPin->MakeLinkTo(LocalPortPin);
+				if (!RemotePortPin && !LocalPortPin)
+				{
+					if (LocalPin->Direction == EGPD_Output)
+					{
+						bDiscardReturnNode = false;
+					}
+					
+					FEdGraphPinType PinType = LocalPin->PinType;
+					if (PinType.bIsWeakPointer && !PinType.IsContainer() && !bCanHaveWeakObjPtrParam)
+					{
+						PinType.bIsWeakPointer = false;
 					}
+					RemotePortPin = InGatewayNode->CreatePin(LocalPin->Direction, PinType, UniquePortName);
+					LocalPortPin = LocalPort->CreateUserDefinedPin(UniquePortName, PinType, (LocalPin->Direction == EGPD_Input) ? EGPD_Output : EGPD_Input);
 				}
 			}
 		}
 
-		if (OutputGatewayExecPins.Num() > 0)
+		check(LocalPortPin);
+		check(RemotePortPin);
+
+		LocalPin->Modify();
+
+		// Route the links
+		for (int32 LinkIndex = 0; LinkIndex < LocalPin->LinkedTo.Num(); ++LinkIndex)
 		{
-			UEdGraphPin* LocalResultPortPin = K2Schema->FindExecutionPin(*InResultNode, EGPD_Input);
+			UEdGraphPin* RemotePin = LocalPin->LinkedTo[LinkIndex];
+			RemotePin->Modify();
 
-			// If the Result Node already contains links, then we don't need to make these connections as the intended connections have already been
-			// transferred from original graph.
-			if (LocalResultPortPin != nullptr && LocalResultPortPin->LinkedTo.Num() == 0)
+			if (!InCollapsableNodes.Contains(RemotePin->GetOwningNode()) && RemotePin->GetOwningNode() != InEntryNode && RemotePin->GetOwningNode() != InResultNode)
 			{
-				// TODO: Some of these pins may not necessarily be terminal pins. We should prompt the user to choose which of these connections should
-				// be made to the return node.
-				for (UEdGraphPin* LocalPin : OutputGatewayExecPins)
+				// Fix up the remote pin
+				RemotePin->LinkedTo.Remove(LocalPin);
+				// When given a composite node, we could possibly be given a pin with a different outer
+				// which is bad! Then there would be a pin connecting to itself and cause an ensure
+				if (RemotePin->GetOwningNode()->GetOuter() == RemotePortPin->GetOwningNode()->GetOuter())
 				{
-					// Connect the gateway pin to the node, there is no remote pin to hook up because the exec pin was not originally connected
-					LocalPin->Modify();
-					LocalPin->MakeLinkTo(LocalResultPortPin);
+					RemotePin->MakeLinkTo(RemotePortPin);
+				}
+
+				// The Entry Node only supports a single link, so if we made links above
+				// we need to break them now, to make room for the new link.
+				if (LocalPort == InEntryNode)
+				{
+					LocalPortPin->BreakAllPinLinks();
 				}
+
+				// Fix up the local pin
+				LocalPin->LinkedTo.Remove(RemotePin);
+				--LinkIndex;
+				LocalPin->MakeLinkTo(LocalPortPin);
 			}
 		}
 	}
