Move CVar to editor config (BP context menu).

#rnx

diff --git a/Engine/Source/Editor/BlueprintGraph/Private/BlueprintEditorSettings.cpp b/Engine/Source/Editor/BlueprintGraph/Private/BlueprintEditorSettings.cpp
--- a/Engine/Source/Editor/BlueprintGraph/Private/BlueprintEditorSettings.cpp
+++ b/Engine/Source/Editor/BlueprintGraph/Private/BlueprintEditorSettings.cpp
@@ -54,6 +54,7 @@ UBlueprintEditorSettings::UBlueprintEditorSettings(const FObjectInitializer& Obj
 	, bEnableNamespaceEditorFeatures(true)
 	// Experimental
 	, bEnableContextMenuTimeSlicing(false)
+	, ContextMenuTimeSlicingThresholdMs(50)
 	, bFavorPureCastNodes(false)
 	// Compiler Settings
 	, SaveOnCompile(SoC_Never)
diff --git a/Engine/Source/Editor/BlueprintGraph/Public/BlueprintEditorSettings.h b/Engine/Source/Editor/BlueprintGraph/Public/BlueprintEditorSettings.h
--- a/Engine/Source/Editor/BlueprintGraph/Public/BlueprintEditorSettings.h
+++ b/Engine/Source/Editor/BlueprintGraph/Public/BlueprintEditorSettings.h
@@ -202,6 +202,10 @@ public:
 	UPROPERTY(EditAnywhere, config, Category = Experimental, meta = (DisplayName = "Enable Non-Blocking Context Menu"))
 	bool bEnableContextMenuTimeSlicing;
 
+	/** The maximum amount of time (in milliseconds) allowed per frame for Blueprint graph context menu building when the non-blocking option is enabled. Larger values will complete the menu build in fewer frames, but will also make the UI feel less responsive. Smaller values will make the UI feel more responsive, but will also take longer to fully populate the menu. */
+	UPROPERTY(EditAnywhere, config, Category = Experimental, AdvancedDisplay, meta = (EditCondition = "bEnableContextMenuTimeSlicing", DisplayName = "Context Menu: Non-Blocking Per-Frame Threshold (ms)", ClampMin = "1"))
+	int32 ContextMenuTimeSlicingThresholdMs;
+
 	/** If enabled, then placed cast nodes will default to their "pure" form (meaning: without execution pins). */
 	UPROPERTY(EditAnywhere, config, AdvancedDisplay, Category = Experimental, meta = (DisplayName = "Default to Using Pure Cast Nodes"))
 	bool bFavorPureCastNodes;
diff --git a/Engine/Source/Editor/Kismet/Private/BlueprintActionMenuBuilder.cpp b/Engine/Source/Editor/Kismet/Private/BlueprintActionMenuBuilder.cpp
--- a/Engine/Source/Editor/Kismet/Private/BlueprintActionMenuBuilder.cpp
+++ b/Engine/Source/Editor/Kismet/Private/BlueprintActionMenuBuilder.cpp
@@ -2,6 +2,7 @@
 
 #include "BlueprintActionMenuBuilder.h"
 #include "UObject/UnrealType.h"
+#include "BlueprintEditorSettings.h"
 #include "Settings/EditorStyleSettings.h"
 #include "Engine/Blueprint.h"
 #include "Editor/EditorEngine.h"
@@ -351,12 +352,6 @@ namespace FBlueprintActionMenuBuilderImpl
 		TArray<FBlueprintActionInfo> PendingActionList;
 	};
 
-	static TAutoConsoleVariable<float> CVarBPContextMenuTimeSlicingThreshold(
-		TEXT("BP.ContextMenuTimeSlicingThreshold"),
-		0.05f,
-		TEXT("The amount of time (in seconds) allowed per frame for time slicing.")
-	);
-
 #if ENABLE_BLUEPRINT_ACTION_FILTER_PROFILING
 	static TAutoConsoleVariable<bool> CVarBPEnableActionMenuDumpToFile(
 		TEXT("BP.EnableActionMenuDumpToFile"),
@@ -698,7 +693,8 @@ bool FBlueprintActionMenuBuilder::ProcessPendingActions()
 	using namespace FBlueprintActionMenuBuilderImpl;
 
 	bool bProcessedActions = false;
-	double StartTime = FPlatformTime::Seconds();
+	const double StartTime = FPlatformTime::Seconds();
+	const float MaxTimeThresholdSeconds = GetDefault<UBlueprintEditorSettings>()->ContextMenuTimeSlicingThresholdMs / 1000.0f;
 
 	FBlueprintActionInfo* CurrentAction = MenuItemListAddHelper->GetNextAction();
 	while (CurrentAction)
@@ -707,7 +703,6 @@ bool FBlueprintActionMenuBuilder::ProcessPendingActions()
 
 		MakeMenuItems(*CurrentAction);
 
-		const float MaxTimeThresholdSeconds = CVarBPContextMenuTimeSlicingThreshold.GetValueOnGameThread();
 		if ((FPlatformTime::Seconds() - StartTime) < MaxTimeThresholdSeconds)
 		{
 			CurrentAction = MenuItemListAddHelper->GetNextAction();
