Fix up references to trashed function objects after generating delegate properties during a batched BP compile with inter-batch dependencies.

diff --git a/Engine/Source/Editor/Kismet/Private/BlueprintCompilationManager.cpp b/Engine/Source/Editor/Kismet/Private/BlueprintCompilationManager.cpp
--- a/Engine/Source/Editor/Kismet/Private/BlueprintCompilationManager.cpp
+++ b/Engine/Source/Editor/Kismet/Private/BlueprintCompilationManager.cpp
@@ -125,6 +125,7 @@ struct FBlueprintCompilationManagerImpl : public FGCObject
 	void QueueForCompilation(const FBPCompileRequestInternal& CompileJob);
 	void CompileSynchronouslyImpl(const FBPCompileRequestInternal& Request);
 	void FlushCompilationQueueImpl(bool bSuppressBroadcastCompiled, TArray<UBlueprint*>* BlueprintsCompiled, TArray<UBlueprint*>* BlueprintsCompiledOrSkeletonCompiled, FUObjectSerializeContext* InLoadContext);
+	void FixupDelegateProperties(const TArray<FCompilerData>& CurrentlyCompilingBPs);
 	void ProcessExtensions(const TArray<FCompilerData>& InCurrentlyCompilingBPs);
 	void FlushReinstancingQueueImpl(bool bFindAndReplaceCDOReferences = false);
 	bool HasBlueprintsToCompile() const;
@@ -1397,6 +1398,21 @@ void FBlueprintCompilationManagerImpl::FlushCompilationQueueImpl(bool bSuppressB
 				CompilerData.Compiler->SetNewClass( CastChecked<UBlueprintGeneratedClass>(BP->GeneratedClass) );
 			}
 		}
+
+		// If we're compiling more than one Blueprint as part of a batch operation, we may have generated
+		// delegate properties (e.g. function parameters) that reference a function dependency in one of
+		// the other Blueprints in the batched set. Depending on the compile order, those functions may get
+		// regenerated in a subsequent pass above, invalidating the delegate property references as a result.
+		// Invalidating references that exist outside of the class layout (e.g. in bytecode) is fine, as
+		// those will be fixed up after we've finished compiling. However, the next stage (compiling functions)
+		// will attempt to validate generated fields for function parameters against their source pin types,
+		// and that logic will throw an error for generated delegate properties if they reference a signature
+		// function that has not yet been replaced. As a result, we run a small pass here to find/fix those up.
+		// 
+		// Note: It is *not* necessary to also verify ScriptAndPropertyObjectReferences here, as that will be
+		// updated in the next stage via reference collection (@see FKismetCompilerContext::CompileFunctions).
+		FixupDelegateProperties(CurrentlyCompilingBPs);
+
 		bGeneratedClassLayoutReady = true;
 		
 		ProcessExtensions(CurrentlyCompilingBPs);
@@ -1751,6 +1767,58 @@ void FBlueprintCompilationManagerImpl::FlushCompilationQueueImpl(bool bSuppressB
 	VerifyNoQueuedRequests(CurrentlyCompilingBPs);
 }
 
+void FBlueprintCompilationManagerImpl::FixupDelegateProperties(const TArray<FCompilerData>& InCurrentlyCompilingBPs)
+{
+	DECLARE_SCOPE_HIERARCHICAL_COUNTER_FUNC();
+
+	if (InCurrentlyCompilingBPs.Num() <= 1)
+	{
+		// It's safe to bypass the fixup logic in this case.
+		return;
+	}
+
+	// Gather old->new field mappings for each regenerated class.
+	TMap<FFieldVariant, FFieldVariant> AllFieldMappings;
+	for (const FCompilerData& CompilerData : InCurrentlyCompilingBPs)
+	{
+		// If there's no reinstancer, then we don't need to include it (e.g. skeleton-only compiles and/or when the layout wasn't regenerated).
+		if (CompilerData.Reinstancer.IsValid())
+		{
+			TMap<FFieldVariant, FFieldVariant> FieldMappings;
+			CompilerData.Reinstancer->GenerateFieldMappings(FieldMappings);
+			AllFieldMappings.Append(MoveTemp(FieldMappings));
+		}
+	}
+
+	auto FixupDelegateProperty = [&AllFieldMappings](FDelegateProperty* DelegateProperty)
+	{
+		// See if the current signature references a stale function (i.e. one that has been regenerated).
+		UFunction** PossiblyStaleFunctionPtr = &DelegateProperty->SignatureFunction;
+		if (FFieldVariant* GeneratedFunctionMapping = AllFieldMappings.Find(*PossiblyStaleFunctionPtr))
+		{
+			// Update the stale reference to point to the regenerated function instead.
+			*PossiblyStaleFunctionPtr = CastChecked<UFunction>(GeneratedFunctionMapping->ToUObject());
+		}
+	};
+
+	for (const TPair<FFieldVariant, FFieldVariant>& FieldMapping : AllFieldMappings)
+	{
+		// Check signatures for any delegate properties owned by the regenerated class.
+		if (FDelegateProperty* ClassDelegateProperty = CastField<FDelegateProperty>(FieldMapping.Value.ToField()))
+		{
+			FixupDelegateProperty(ClassDelegateProperty);
+		}
+		else if (UFunction* Function = Cast<UFunction>(FieldMapping.Value.ToUObject()))
+		{
+			// Check signatures for any delegate properties owned by a regenerated class function (e.g. parameters).
+			for (FDelegateProperty* LocalDelegateProperty : TFieldRange<FDelegateProperty>(Function, EFieldIterationFlags::IncludeDeprecated))
+			{
+				FixupDelegateProperty(LocalDelegateProperty);
+			}
+		}
+	}
+}
+
 void FBlueprintCompilationManagerImpl::ProcessExtensions(const TArray<FCompilerData>& InCurrentlyCompilingBPs)
 {
 	DECLARE_SCOPE_HIERARCHICAL_COUNTER_FUNC()
diff --git a/Engine/Source/Editor/UnrealEd/Private/Kismet2/KismetReinstanceUtilities.cpp b/Engine/Source/Editor/UnrealEd/Private/Kismet2/KismetReinstanceUtilities.cpp
--- a/Engine/Source/Editor/UnrealEd/Private/Kismet2/KismetReinstanceUtilities.cpp
+++ b/Engine/Source/Editor/UnrealEd/Private/Kismet2/KismetReinstanceUtilities.cpp
@@ -421,8 +421,11 @@ void FBlueprintCompileReinstancer::GenerateFieldMappings(TMap<FFieldVariant, FFi
 		FieldMapping.Add(Func.Value, NewFunction);
 	}
 
-	UObject* NewCDO = ClassToReinstance->GetDefaultObject();
-	FieldMapping.Add(OriginalCDO, NewCDO);
+	if(!ClassToReinstance->bLayoutChanging)
+	{
+		UObject* NewCDO = ClassToReinstance->GetDefaultObject();
+		FieldMapping.Add(OriginalCDO, NewCDO);
+	}
 }
 
 void FBlueprintCompileReinstancer::AddReferencedObjects(FReferenceCollector& Collector)
