[Feature/Bugfix] Support for diffing against nullptr (treated as an empty or nonexistent file)

"Crash occurs when scrolling down through files in Review Changelist window upon loading of asset"
  - certain widget blueprints were crashing the review tool because we were working around the inability to diff against nullptr by constructing temprorary empty objects. This approach circumvents that problem entirely

"Changelist Review tool allows diffing against previous revision for Deletion changes"
  - while this jira suggests that the ability to diff deletion changes is a bug, it was actually an intended feature. The real bug was that it would crash or fail to work in many cases. This has been fixed by either asset diffing against nullptr or text diffing against an empty file (depending on circumstance)

"Reviewing Changelist with Deletion Changes results in LogLinker warnings & LoadErrors"
 - deletion changes were trying to load the new revision of files but obviously it doesn't exist because it's a deletion. Fixed by only loading the previous revision

note: since the review tool is going to be "Production Ready"  in 5.2, these changes are neccesary for stability

diff --git a/Engine/Plugins/Animation/ControlRig/Source/ControlRigEditor/Private/ControlRigBlueprintActions.cpp b/Engine/Plugins/Animation/ControlRig/Source/ControlRigEditor/Private/ControlRigBlueprintActions.cpp
--- a/Engine/Plugins/Animation/ControlRig/Source/ControlRigEditor/Private/ControlRigBlueprintActions.cpp
+++ b/Engine/Plugins/Animation/ControlRig/Source/ControlRigEditor/Private/ControlRigBlueprintActions.cpp
@@ -94,47 +94,48 @@ TSharedPtr<SWidget> FControlRigBlueprintActions::GetThumbnailOverlay(const FAsse
 
 void FControlRigBlueprintActions::PerformAssetDiff(UObject* OldAsset, UObject* NewAsset, const FRevisionInfo& OldRevision, const FRevisionInfo& NewRevision) const
 {
-	UBlueprint* OldBlueprint = CastChecked<UBlueprint>(OldAsset);
-	UBlueprint* NewBlueprint = CastChecked<UBlueprint>(NewAsset);
+	UBlueprint* OldBlueprint = Cast<UBlueprint>(OldAsset);
+	UBlueprint* NewBlueprint = Cast<UBlueprint>(NewAsset);
 
 	static const FText DiffWindowMessage = LOCTEXT("ControlRigDiffWindow", "Opening a diff window will close the control rig editor. {0}.\nAre you sure?");
 	
 	UAssetEditorSubsystem* AssetEditorSubsystem = GEditor->GetEditorSubsystem<UAssetEditorSubsystem>();
-	for (IAssetEditorInstance* Editor : AssetEditorSubsystem->FindEditorsForAsset(OldAsset))
+	if (OldAsset)
 	{
-		const EAppReturnType::Type Answer = FMessageDialog::Open( EAppMsgType::YesNo,
-				FText::Format(DiffWindowMessage, FText::FromString(OldBlueprint->GetName())));
-		if(Answer == EAppReturnType::No)
+		for (IAssetEditorInstance* Editor : AssetEditorSubsystem->FindEditorsForAsset(OldAsset))
 		{
-		   return;
+			const EAppReturnType::Type Answer = FMessageDialog::Open( EAppMsgType::YesNo,
+					FText::Format(DiffWindowMessage, FText::FromString(OldBlueprint->GetName())));
+			if(Answer == EAppReturnType::No)
+			{
+			   return;
+			}
 		}
 	}
-	for (IAssetEditorInstance* Editor : AssetEditorSubsystem->FindEditorsForAsset(NewAsset))
+	if (NewAsset)
 	{
-		const EAppReturnType::Type Answer = FMessageDialog::Open( EAppMsgType::YesNo,
-				FText::Format(DiffWindowMessage, FText::FromString(NewBlueprint->GetName())));
-		if(Answer == EAppReturnType::No)
+		for (IAssetEditorInstance* Editor : AssetEditorSubsystem->FindEditorsForAsset(NewAsset))
 		{
-			return;
+			const EAppReturnType::Type Answer = FMessageDialog::Open( EAppMsgType::YesNo,
+					FText::Format(DiffWindowMessage, FText::FromString(NewBlueprint->GetName())));
+			if(Answer == EAppReturnType::No)
+			{
+				return;
+			}
 		}
 	}
-	
-	AssetEditorSubsystem->CloseAllEditorsForAsset(OldAsset);
-	AssetEditorSubsystem->CloseAllEditorsForAsset(NewAsset);
 
-	// sometimes we're comparing different revisions of one single asset (other 
-	// times we're comparing two completely separate assets altogether)
-	bool bIsSingleAsset = (NewBlueprint->GetName() == OldBlueprint->GetName());
-
-	FText WindowTitle = LOCTEXT("ControlRigBlueprintDiff", "Control Rig Blueprint Diff");
-	// if we're diffing one asset against itself 
-	if (bIsSingleAsset)
+	if (OldAsset)
+	{
+		AssetEditorSubsystem->CloseAllEditorsForAsset(OldAsset);
+	}
+	
+	if (NewAsset)
 	{
-		// identify the assumed single asset in the window's title
-		WindowTitle = FText::Format(LOCTEXT("Control Rig Diff", "{0} - Control Rig Diff"), FText::FromString(NewBlueprint->GetName()));
+		AssetEditorSubsystem->CloseAllEditorsForAsset(NewAsset);
 	}
 
-	SBlueprintDiff::CreateDiffWindow(WindowTitle, OldBlueprint, NewBlueprint, OldRevision, NewRevision);
+	SBlueprintDiff::CreateDiffWindow(OldBlueprint, NewBlueprint, OldRevision, NewRevision, GetSupportedClass());
 }
 
 void FControlRigBlueprintActions::ExtendSketalMeshToolMenu()
diff --git a/Engine/Plugins/Editor/ChangelistReview/Source/Private/SSourceControlReview.cpp b/Engine/Plugins/Editor/ChangelistReview/Source/Private/SSourceControlReview.cpp
--- a/Engine/Plugins/Editor/ChangelistReview/Source/Private/SSourceControlReview.cpp
+++ b/Engine/Plugins/Editor/ChangelistReview/Source/Private/SSourceControlReview.cpp
@@ -50,7 +50,7 @@ const UClass* FChangelistFileData::GetIconClass()
 	{
 		CachedIconClass = nullptr;
 		const FAssetRegistryModule& AssetRegistryModule = FModuleManager::LoadModuleChecked<FAssetRegistryModule>(TEXT("AssetRegistry"));
-		const FString TempPackageName = ReviewFileTempPath.GetPackageName();
+		const FString TempPackageName = FPackagePath::FromLocalPath(ReviewFileTempPath).GetPackageName();
 		if (!TempPackageName.IsEmpty())
 		{
 			TArray<FAssetData> OutAssetData;
@@ -374,9 +374,6 @@ void SSourceControlReview::OnChangelistLoadComplete(const FSourceControlOperatio
 	while (ChangelistRecord.Contains(RecordFileMapKey) && ChangelistRecord.Contains(RecordRevisionMapKey))
 	{
 		const FString &FileDepotPath = ChangelistRecord[RecordFileMapKey];
-		
-		//For each 1 file we are loading 2 revisions so files to load is always incremented by two per file
-		FilesToLoad++; 
 
 		const bool bIsShelved = ChangelistRecord[ReviewHelpers::ChangelistStatusKey] == ReviewHelpers::ChangelistPendingStatusKey;
 		const int32 AssetRevision = FCString::Atoi(*ChangelistRecord[RecordRevisionMapKey]);
@@ -402,22 +399,27 @@ void SSourceControlReview::OnChangelistLoadComplete(const FSourceControlOperatio
 		const int32 PreviousAssetRevision = bIsShelved && (ChangelistFileData->FileSourceControlAction == ESourceControlAction::Delete || ChangelistFileData->FileSourceControlAction == ESourceControlAction::Edit) ? AssetRevision : AssetRevision - 1;
 		const FString PreviousAssetRevisionStr = FString::FromInt(PreviousAssetRevision);
 		
-		// retrieve files directly from source control into a temp location
-		TSharedRef<FGetFile> GetFileToReviewCommand = ISourceControlOperation::Create<FGetFile>(Changelist, ChangelistRecord[RecordRevisionMapKey], ChangelistRecord[RecordFileMapKey], bIsShelved);
-
 		TWeakPtr<SSourceControlReview> WeakReviewWidget = SharedThis(this);
-		const auto GetFileCommandResponse = [WeakReviewWidget, ChangelistFileData](const FSourceControlOperationRef& InOperation, ECommandResult::Type InResult, FPackagePath* OutFilePath)
+		const auto GetFileCommandResponse = [WeakReviewWidget, ChangelistFileData](const FSourceControlOperationRef& InOperation, ECommandResult::Type InResult, FString* OutFilePath)
 		{
 			if (WeakReviewWidget.IsValid())
 			{
 				const TSharedRef<FGetFile, ESPMode::ThreadSafe> Operation = StaticCastSharedRef<FGetFile>(InOperation);
-				*OutFilePath = FPackagePath::FromLocalPath(Operation->GetOutPackageFilename());
+				*OutFilePath = Operation->GetOutPackageFilename();
 				WeakReviewWidget.Pin()->OnGetFileFromSourceControl(ChangelistFileData);
 			}
 		};
-		
-		ISourceControlModule::Get().GetProvider().Execute(GetFileToReviewCommand, EConcurrency::Asynchronous, FSourceControlOperationComplete::CreateLambda(GetFileCommandResponse, &ChangelistFileData->ReviewFileTempPath));
 
+		// retrieve files directly from source control into a temp location
+		if (ChangelistFileData->FileSourceControlAction != ESourceControlAction::Delete)
+		{
+			FilesToLoad++;
+			
+			TSharedRef<FGetFile> GetFileToReviewCommand = ISourceControlOperation::Create<FGetFile>(Changelist, ChangelistRecord[RecordRevisionMapKey], ChangelistRecord[RecordFileMapKey], bIsShelved);
+
+			ISourceControlModule::Get().GetProvider().Execute(GetFileToReviewCommand, EConcurrency::Asynchronous, FSourceControlOperationComplete::CreateLambda(GetFileCommandResponse, &ChangelistFileData->ReviewFileTempPath));
+		}
+		
 		if (ChangelistFileData->FileSourceControlAction != ESourceControlAction::Add && ChangelistFileData->FileSourceControlAction != ESourceControlAction::Branch)
 		{
 			FilesToLoad++;
@@ -607,7 +609,7 @@ void SSourceControlReview::OnGetFileFromSourceControl(TSharedPtr<FChangelistFile
 		{
 			if(FileData)
 			{
-				ChangelistFilePaths.Add(FileData->ReviewFileTempPath.GetLocalFullPath());
+				ChangelistFilePaths.Add(FileData->ReviewFileTempPath);
 			}
 		}
 		
@@ -856,7 +858,7 @@ void SSourceControlReview::FixupRedirectors()
 	TMap<FString, TWeakPtr<FChangelistFileData>> RedirectorsFound;
 	for (TSharedPtr<FChangelistFileData> ChangelistFile : ChangelistFiles)
 	{
-		if (FLinkerLoad* ReviewFileLoad = LoadPackageLinker(nullptr, ChangelistFile->ReviewFileTempPath, LOAD_ForDiff | LOAD_DisableCompileOnLoad | LOAD_DisableEngineVersionChecks))
+		if (FLinkerLoad* ReviewFileLoad = LoadPackageLinker(nullptr, FPackagePath::FromLocalPath(ChangelistFile->ReviewFileTempPath), LOAD_ForDiff | LOAD_DisableCompileOnLoad | LOAD_DisableEngineVersionChecks))
 		{
 			if (IsObjectRedirector(ReviewFileLoad, ChangelistFile->AssetName))
 			{
diff --git a/Engine/Plugins/Editor/ChangelistReview/Source/Private/SSourceControlReview.h b/Engine/Plugins/Editor/ChangelistReview/Source/Private/SSourceControlReview.h
--- a/Engine/Plugins/Editor/ChangelistReview/Source/Private/SSourceControlReview.h
+++ b/Engine/Plugins/Editor/ChangelistReview/Source/Private/SSourceControlReview.h
@@ -44,14 +44,26 @@ namespace SourceControlReview
 
 		bool IsDataValidForEntry() const
 		{
-			return FileSourceControlAction != ESourceControlAction::Unset && !ReviewFileTempPath.IsEmpty() &&
-				(!PreviousFileTempPath.IsEmpty() || FileSourceControlAction == ESourceControlAction::Add || FileSourceControlAction == ESourceControlAction::Branch);
+			if (FileSourceControlAction == ESourceControlAction::Unset)
+			{
+				return false;
+			}
+			if (ReviewFileTempPath.IsEmpty() && FileSourceControlAction != ESourceControlAction::Delete)
+			{
+				return false;
+			}
+			if (PreviousFileTempPath.IsEmpty() && FileSourceControlAction != ESourceControlAction::Add && FileSourceControlAction != ESourceControlAction::Branch)
+			{
+				return false;
+			}
+			
+			return true;
 		}
 
 		FString AssetName;
 
 		// path where the temporary asset was downloaded to
-		FPackagePath ReviewFileTempPath;
+		FString ReviewFileTempPath;
 
 		FString ReviewFileRevisionNum;
 
@@ -60,7 +72,7 @@ namespace SourceControlReview
 		FString PreviousAssetName;
 
 		// path where the temporary asset was downloaded to
-		FPackagePath PreviousFileTempPath;
+		FString PreviousFileTempPath;
 
 		FString PreviousFileRevisionNum;
 
diff --git a/Engine/Plugins/Editor/ChangelistReview/Source/Private/SSourceControlReviewEntry.cpp b/Engine/Plugins/Editor/ChangelistReview/Source/Private/SSourceControlReviewEntry.cpp
--- a/Engine/Plugins/Editor/ChangelistReview/Source/Private/SSourceControlReviewEntry.cpp
+++ b/Engine/Plugins/Editor/ChangelistReview/Source/Private/SSourceControlReviewEntry.cpp
@@ -10,6 +10,7 @@
 #include "AssetRegistry/AssetRegistryModule.h"
 #include "IAssetTools.h"
 #include "Kismet2/KismetEditorUtilities.h"
+#include "Misc/FileHelper.h"
 #include "Misc/MessageDialog.h"
 #include "Settings/EditorLoadingSavingSettings.h"
 #include "Styling/StarshipCoreStyle.h"
@@ -207,21 +208,35 @@ bool SSourceControlReviewEntry::CanDiff() const
 	return DiffMethod.IsBound();
 }
 
+static bool IsAssetPath(const FString& Path)
+{
+	if (Path.IsEmpty())
+	{
+		return false;
+	}
+	const EPackageExtension Extension = FPackagePath::ParseExtension(Path);
+	return Extension != EPackageExtension::Unspecified && Extension != EPackageExtension::Custom && Extension != EPackageExtension::EmptyString;
+}
+
 void SSourceControlReviewEntry::TryBindUAssetDiff()
 {
-	UObject* ReviewAsset = nullptr;
-	if (UPackage* ReviewFilePkg = LoadPackage(nullptr, ChangelistFileData.ReviewFileTempPath, LOAD_ForDiff | LOAD_DisableCompileOnLoad | LOAD_DisableEngineVersionChecks))
+	UObject* ReviewAsset = nullptr;	
+	if (IsAssetPath(ChangelistFileData.ReviewFileTempPath))
 	{
-		ReviewAsset = FindObject<UObject>(ReviewFilePkg, *ChangelistFileData.AssetName);
-		if(ReviewAsset && ReviewAsset->IsA<UObjectRedirector>())
+		if (UPackage* ReviewFilePkg = LoadPackage(nullptr, FPackagePath::FromLocalPath(ChangelistFileData.ReviewFileTempPath), LOAD_ForDiff | LOAD_DisableCompileOnLoad | LOAD_DisableEngineVersionChecks))
 		{
-			ReviewAsset = Cast<UObjectRedirector>(ReviewAsset)->DestinationObject;
+			ReviewAsset = FindObject<UObject>(ReviewFilePkg, *ChangelistFileData.AssetName);
+			if(ReviewAsset && ReviewAsset->IsA<UObjectRedirector>())
+			{
+				ReviewAsset = Cast<UObjectRedirector>(ReviewAsset)->DestinationObject;
+			}
 		}
 	}
+	
 	UObject* PreviousAsset = nullptr;
-	if (!ChangelistFileData.PreviousFileTempPath.IsEmpty())
+	if (IsAssetPath(ChangelistFileData.PreviousFileTempPath))
 	{
-		if (UPackage* PreviousFilePkg = LoadPackage(nullptr, ChangelistFileData.PreviousFileTempPath, LOAD_ForDiff | LOAD_DisableCompileOnLoad | LOAD_DisableEngineVersionChecks))
+		if (UPackage* PreviousFilePkg = LoadPackage(nullptr, FPackagePath::FromLocalPath(ChangelistFileData.PreviousFileTempPath), LOAD_ForDiff | LOAD_DisableCompileOnLoad | LOAD_DisableEngineVersionChecks))
 		{
 			if (ChangelistFileData.PreviousAssetName.IsEmpty())
 			{
@@ -239,50 +254,29 @@ void SSourceControlReviewEntry::TryBindUAssetDiff()
 		return;
 	}
 
-	auto MakeDefaultAsset = [this](UObject* OtherAsset)->UObject*
+	DiffMethod.BindLambda([this, PreviousAsset, ReviewAsset]
 	{
-		if (const UBlueprint* BlueprintReviewAsset = Cast<UBlueprint>(OtherAsset))
+		const UBlueprint* ReviewBlueprint = Cast<UBlueprint>(ReviewAsset);
+		const UBlueprint* PreviousBlueprint = Cast<UBlueprint>(PreviousAsset);
+		if ((ReviewBlueprint && !ReviewBlueprint->ParentClass) || (PreviousBlueprint && !PreviousBlueprint->ParentClass))
 		{
-			return GetOrCreateBlueprintForDiff(BlueprintReviewAsset->GeneratedClass, BlueprintReviewAsset->BlueprintType);
+			FMessageDialog::Open(EAppMsgType::Ok, LOCTEXT("ChangelistNotFoundError", "This Blueprint is missing its parent class. Diff results may be incomplete.\n\nDo you need to load this file's plugin/module?\nFor more details, search 'Can't find file.' in the log"));
 		}
-		else
-		{
-			//In case when file is not blueprint we are falling back on trying to create a UObject 
-			const FName EmptyObjectName = OtherAsset->GetFName();
-			const UClass* EmptyObjectAssetClass = OtherAsset->GetClass();
-			return NewObject<UObject>(OtherAsset, EmptyObjectAssetClass, EmptyObjectName, OtherAsset->GetFlags());
-		}
-	};
-
-	if (!ReviewAsset)
-	{
-		ReviewAsset = MakeDefaultAsset(PreviousAsset);
-	}
-
-	if (!PreviousAsset)
-	{
-		PreviousAsset = MakeDefaultAsset(ReviewAsset);
-	}
+		
+		FAssetToolsModule::GetModule().Get().DiffAssets(
+			PreviousAsset,
+			ReviewAsset,
+			GetPreviousFileRevisionInfo(),
+			GetReviewFileRevisionInfo()
+		);
+	});
+}
 
-	if (ReviewAsset && PreviousAsset)
-	{
-		DiffMethod.BindLambda([this, PreviousAsset, ReviewAsset]
-		{
-			const UBlueprint* ReviewBlueprint = Cast<UBlueprint>(ReviewAsset);
-			const UBlueprint* PreviousBlueprint = Cast<UBlueprint>(PreviousAsset);
-			if ((ReviewBlueprint && !ReviewBlueprint->ParentClass) || (PreviousBlueprint && !PreviousBlueprint->ParentClass))
-			{
-				FMessageDialog::Open(EAppMsgType::Ok, LOCTEXT("ChangelistNotFoundError", "This Blueprint is missing its parent class. Diff results may be incomplete.\n\nDo you need to load this file's plugin/module?\nFor more details, search 'Can't find file.' in the log"));
-			}
-			
-			FAssetToolsModule::GetModule().Get().DiffAssets(
-				PreviousAsset,
-				ReviewAsset,
-				GetPreviousFileRevisionInfo(),
-				GetReviewFileRevisionInfo()
-			);
-		});
-	}
+static FString MakeEmptyTempFile(const FString& AssetName, const FString& Revision)
+{
+	const FString Path = FString::Printf(TEXT("%sTemp%s-%s.empty"), *FPaths::DiffDir(), *AssetName, *Revision);
+	FFileHelper::SaveStringToFile(TEXT(""), *Path);
+	return Path;
 }
 
 void SSourceControlReviewEntry::TryBindTextDiff()
@@ -291,7 +285,22 @@ void SSourceControlReviewEntry::TryBindTextDiff()
 	{
 		const FString& DiffCommand = GetDefault<UEditorLoadingSavingSettings>()->TextDiffToolPath.FilePath;
 		const FAssetToolsModule& AssetToolsModule = FAssetToolsModule::GetModule();
-		AssetToolsModule.Get().CreateDiffProcess(DiffCommand, ChangelistFileData.PreviousFileTempPath.GetLocalFullPath(), ChangelistFileData.ReviewFileTempPath.GetLocalFullPath());
+		
+		FString OldFilePath = ChangelistFileData.PreviousFileTempPath;
+		FString NewFilePath = ChangelistFileData.ReviewFileTempPath;
+
+		// if this is an add or delete change, diff against an empty file
+		if (OldFilePath.IsEmpty())
+		{
+			const FString& PreviousAssetName = ChangelistFileData.PreviousAssetName.IsEmpty() ? ChangelistFileData.AssetName : ChangelistFileData.PreviousAssetName;
+			OldFilePath = MakeEmptyTempFile(PreviousAssetName, ChangelistFileData.PreviousFileRevisionNum);
+		}
+		if (NewFilePath.IsEmpty())
+		{
+			NewFilePath = MakeEmptyTempFile(ChangelistFileData.AssetName, ChangelistFileData.PreviousFileRevisionNum);
+		}
+		
+		AssetToolsModule.Get().CreateDiffProcess(DiffCommand, OldFilePath, NewFilePath);
 	});
 }
 
diff --git a/Engine/Plugins/Editor/EngineAssetDefinitions/Source/Private/Animation/AssetDefinition_AnimBlueprint.cpp b/Engine/Plugins/Editor/EngineAssetDefinitions/Source/Private/Animation/AssetDefinition_AnimBlueprint.cpp
--- a/Engine/Plugins/Editor/EngineAssetDefinitions/Source/Private/Animation/AssetDefinition_AnimBlueprint.cpp
+++ b/Engine/Plugins/Editor/EngineAssetDefinitions/Source/Private/Animation/AssetDefinition_AnimBlueprint.cpp
@@ -121,23 +121,15 @@ TSharedPtr<SWidget> UAssetDefinition_AnimBlueprint::GetThumbnailOverlay(const FA
 
 EAssetCommandResult UAssetDefinition_AnimBlueprint::PerformAssetDiff(const FAssetDiffArgs& DiffArgs) const
 {
-	UBlueprint* OldBlueprint = CastChecked<UBlueprint>(DiffArgs.OldAsset);
-	UBlueprint* NewBlueprint = CastChecked<UBlueprint>(DiffArgs.NewAsset);
+	const UBlueprint* OldBlueprint = Cast<UBlueprint>(DiffArgs.OldAsset);
+	const UBlueprint* NewBlueprint = Cast<UBlueprint>(DiffArgs.NewAsset);
 
-	// sometimes we're comparing different revisions of one single asset (other 
-	// times we're comparing two completely separate assets altogether)
-	bool bIsSingleAsset = (NewBlueprint->GetName() == OldBlueprint->GetName());
-
-	FText WindowTitle = LOCTEXT("NamelessAnimationBlueprintDiff", "Animation Blueprint Diff");
-	// if we're diffing one asset against itself 
-	if (bIsSingleAsset)
+	if (NewBlueprint == nullptr && OldBlueprint == nullptr)
 	{
-		// identify the assumed single asset in the window's title
-		WindowTitle = FText::Format(LOCTEXT("AnimationBlueprintDiff", "{0} - Animation Blueprint Diff"), FText::FromString(NewBlueprint->GetName()));
+		return EAssetCommandResult::Unhandled;
 	}
-
-	SBlueprintDiff::CreateDiffWindow(WindowTitle, OldBlueprint, NewBlueprint, DiffArgs.OldRevision, DiffArgs.NewRevision);
-
+	
+	SBlueprintDiff::CreateDiffWindow(OldBlueprint, NewBlueprint, DiffArgs.OldRevision, DiffArgs.NewRevision, UAnimBlueprint::StaticClass());
 	return EAssetCommandResult::Handled;
 }
 
diff --git a/Engine/Plugins/Editor/EngineAssetDefinitions/Source/Private/AssetDefinition_Curve.cpp b/Engine/Plugins/Editor/EngineAssetDefinitions/Source/Private/AssetDefinition_Curve.cpp
--- a/Engine/Plugins/Editor/EngineAssetDefinitions/Source/Private/AssetDefinition_Curve.cpp
+++ b/Engine/Plugins/Editor/EngineAssetDefinitions/Source/Private/AssetDefinition_Curve.cpp
@@ -36,20 +36,25 @@ EAssetCommandResult UAssetDefinition_Curve::PerformAssetDiff(const FAssetDiffArg
 	UCurveBase* OldCurve = Cast<UCurveBase>(DiffArgs.OldAsset);
 	UCurveBase* NewCurve = Cast<UCurveBase>(DiffArgs.NewAsset);
 
-	if (!ensure(OldCurve != nullptr && NewCurve != nullptr))
+	if (NewCurve == nullptr && OldCurve == nullptr)
 	{
 		return EAssetCommandResult::Unhandled;
 	}
 
 	// Build names for temp json files
-	const FString RelOldTempFileName = FString::Printf(TEXT("%sTemp%s-%s.json"), *FPaths::DiffDir(), *DiffArgs.OldAsset->GetName(), *DiffArgs.OldRevision.Revision);
+	const FString OldAssetName = DiffArgs.OldAsset ? DiffArgs.OldAsset->GetName() : DiffArgs.NewAsset->GetName();
+	const FString RelOldTempFileName = FString::Printf(TEXT("%sTemp%s-%s.json"), *FPaths::DiffDir(), *OldAssetName, *DiffArgs.OldRevision.Revision);
 	const FString AbsoluteOldTempFileName = FPaths::ConvertRelativePathToFull(RelOldTempFileName);
-	const FString RelNewTempFileName = FString::Printf(TEXT("%sTemp%s-%s.json"), *FPaths::DiffDir(), *DiffArgs.NewAsset->GetName(), *DiffArgs.NewRevision.Revision);
+	
+	const FString NewAssetName = DiffArgs.NewAsset ? DiffArgs.NewAsset->GetName() : DiffArgs.OldAsset->GetName();
+	const FString RelNewTempFileName = FString::Printf(TEXT("%sTemp%s-%s.json"), *FPaths::DiffDir(), *NewAssetName, *DiffArgs.NewRevision.Revision);
 	const FString AbsoluteNewTempFileName = FPaths::ConvertRelativePathToFull(RelNewTempFileName);
 
 	// save temp files
-	const bool OldResult = FFileHelper::SaveStringToFile(OldCurve->ExportAsJSONString(), *AbsoluteOldTempFileName);
-	const bool NewResult = FFileHelper::SaveStringToFile(NewCurve->ExportAsJSONString(), *AbsoluteNewTempFileName);
+	const FString OldJson = OldCurve ? OldCurve->ExportAsJSONString() : TEXT("");
+	const bool OldResult = FFileHelper::SaveStringToFile(OldJson, *AbsoluteOldTempFileName);
+	const FString NewJson = NewCurve ? NewCurve->ExportAsJSONString() : TEXT("");
+	const bool NewResult = FFileHelper::SaveStringToFile(NewJson, *AbsoluteNewTempFileName);
 
 	if (OldResult && NewResult)
 	{
diff --git a/Engine/Plugins/Editor/EngineAssetDefinitions/Source/Private/AssetDefinition_DataAsset.cpp b/Engine/Plugins/Editor/EngineAssetDefinitions/Source/Private/AssetDefinition_DataAsset.cpp
--- a/Engine/Plugins/Editor/EngineAssetDefinitions/Source/Private/AssetDefinition_DataAsset.cpp
+++ b/Engine/Plugins/Editor/EngineAssetDefinitions/Source/Private/AssetDefinition_DataAsset.cpp
@@ -167,15 +167,12 @@ FText UAssetDefinition_DataAsset::GetAssetDisplayName(const FAssetData& AssetDat
 
 EAssetCommandResult UAssetDefinition_DataAsset::PerformAssetDiff(const FAssetDiffArgs& DiffArgs) const
 {
-	// sometimes we're comparing different revisions of one single asset (other 
-	// times we're comparing two completely separate assets altogether)
-	const bool bIsSingleAsset = (DiffArgs.OldAsset->GetFName() == DiffArgs.NewAsset->GetFName());
-	static const FText BasicWindowTitle = LOCTEXT("NamelessDataAssetDiff", "DataAsset Diff");
-
-	const FText WindowTitle = !bIsSingleAsset ? BasicWindowTitle : FText::Format(LOCTEXT("DataAsset Diff", "{0} - DataAsset Diff"), FText::FromString(DiffArgs.NewAsset->GetName()));
-
-	SDetailsDiff::CreateDiffWindow(WindowTitle, DiffArgs.OldAsset, DiffArgs.NewAsset, DiffArgs.OldRevision, DiffArgs.NewRevision);
+	if (DiffArgs.OldAsset == nullptr && DiffArgs.NewAsset == nullptr)
+	{
+		return EAssetCommandResult::Unhandled;
+	}
 	
+	SDetailsDiff::CreateDiffWindow(DiffArgs.OldAsset, DiffArgs.NewAsset, DiffArgs.OldRevision, DiffArgs.NewRevision, UDataAsset::StaticClass());
 	return EAssetCommandResult::Handled;
 }
 
diff --git a/Engine/Plugins/Editor/EngineAssetDefinitions/Source/Private/Script/AssetDefinition_Blueprint.cpp b/Engine/Plugins/Editor/EngineAssetDefinitions/Source/Private/Script/AssetDefinition_Blueprint.cpp
--- a/Engine/Plugins/Editor/EngineAssetDefinitions/Source/Private/Script/AssetDefinition_Blueprint.cpp
+++ b/Engine/Plugins/Editor/EngineAssetDefinitions/Source/Private/Script/AssetDefinition_Blueprint.cpp
@@ -103,23 +103,15 @@ EAssetCommandResult UAssetDefinition_Blueprint::Merge(const FAssetManualMergeArg
 
 EAssetCommandResult UAssetDefinition_Blueprint::PerformAssetDiff(const FAssetDiffArgs& DiffArgs) const
 {
-	UBlueprint* OldBlueprint = CastChecked<UBlueprint>(DiffArgs.OldAsset);
-	UBlueprint* NewBlueprint = CastChecked<UBlueprint>(DiffArgs.NewAsset);
+	const UBlueprint* OldBlueprint = Cast<UBlueprint>(DiffArgs.OldAsset);
+	const UBlueprint* NewBlueprint = Cast<UBlueprint>(DiffArgs.NewAsset);
 
-	// sometimes we're comparing different revisions of one single asset (other 
-	// times we're comparing two completely separate assets altogether)
-	bool bIsSingleAsset = (NewBlueprint->GetName() == OldBlueprint->GetName());
-
-	FText WindowTitle = LOCTEXT("NamelessBlueprintDiff", "Blueprint Diff");
-	// if we're diffing one asset against itself 
-	if (bIsSingleAsset)
+	if (NewBlueprint == nullptr && OldBlueprint == nullptr)
 	{
-		// identify the assumed single asset in the window's title
-		WindowTitle = FText::Format(LOCTEXT("Blueprint Diff", "{0} - Blueprint Diff"), FText::FromString(NewBlueprint->GetName()));
+		return EAssetCommandResult::Unhandled;
 	}
-
-	SBlueprintDiff::CreateDiffWindow(WindowTitle, OldBlueprint, NewBlueprint, DiffArgs.OldRevision, DiffArgs.NewRevision);
-
+	
+	SBlueprintDiff::CreateDiffWindow(OldBlueprint, NewBlueprint, DiffArgs.OldRevision, DiffArgs.NewRevision, UBlueprint::StaticClass());
 	return EAssetCommandResult::Handled;
 }
 
diff --git a/Engine/Plugins/Editor/EngineAssetDefinitions/Source/Private/Table/AssetDefinition_CurveTable.cpp b/Engine/Plugins/Editor/EngineAssetDefinitions/Source/Private/Table/AssetDefinition_CurveTable.cpp
--- a/Engine/Plugins/Editor/EngineAssetDefinitions/Source/Private/Table/AssetDefinition_CurveTable.cpp
+++ b/Engine/Plugins/Editor/EngineAssetDefinitions/Source/Private/Table/AssetDefinition_CurveTable.cpp
@@ -23,18 +23,29 @@
 
 EAssetCommandResult UAssetDefinition_CurveTable::PerformAssetDiff(const FAssetDiffArgs& DiffArgs) const
 {
-	UCurveTable* OldCurveTable = CastChecked<UCurveTable>(DiffArgs.OldAsset);
-	UCurveTable* NewCurveTable = CastChecked<UCurveTable>(DiffArgs.NewAsset);
+	const UCurveTable* OldCurveTable = Cast<UCurveTable>(DiffArgs.OldAsset);
+	const UCurveTable* NewCurveTable = Cast<UCurveTable>(DiffArgs.NewAsset);
+
+	if (NewCurveTable == nullptr && OldCurveTable == nullptr)
+	{
+		return EAssetCommandResult::Unhandled;
+	}
+
 
 	// Build names for temp csv files
-	const FString RelOldTempFileName = FString::Printf(TEXT("%sTemp%s-%s.csv"), *FPaths::DiffDir(), *DiffArgs.OldAsset->GetName(), *DiffArgs.OldRevision.Revision);
+	const FString OldAssetName = DiffArgs.OldAsset ? DiffArgs.OldAsset->GetName() : DiffArgs.NewAsset->GetName();
+	const FString RelOldTempFileName = FString::Printf(TEXT("%sTemp%s-%s.csv"), *FPaths::DiffDir(), *OldAssetName, *DiffArgs.OldRevision.Revision);
 	const FString AbsoluteOldTempFileName = FPaths::ConvertRelativePathToFull(RelOldTempFileName);
-	const FString RelNewTempFileName = FString::Printf(TEXT("%sTemp%s-%s.csv"), *FPaths::DiffDir(), *DiffArgs.NewAsset->GetName(), *DiffArgs.NewRevision.Revision);
+	
+	const FString NewAssetName = DiffArgs.NewAsset ? DiffArgs.NewAsset->GetName() : DiffArgs.OldAsset->GetName();
+	const FString RelNewTempFileName = FString::Printf(TEXT("%sTemp%s-%s.csv"), *FPaths::DiffDir(), *NewAssetName, *DiffArgs.NewRevision.Revision);
 	const FString AbsoluteNewTempFileName = FPaths::ConvertRelativePathToFull(RelNewTempFileName);
 
 	// save temp files
-	const bool OldResult = FFileHelper::SaveStringToFile(OldCurveTable->GetTableAsCSV(), *AbsoluteOldTempFileName);
-	const bool NewResult = FFileHelper::SaveStringToFile(NewCurveTable->GetTableAsCSV(), *AbsoluteNewTempFileName);
+	const FString OldCurveCSV = OldCurveTable? OldCurveTable->GetTableAsCSV() : TEXT("");
+	const bool OldResult = FFileHelper::SaveStringToFile(OldCurveCSV, *AbsoluteOldTempFileName);
+	const FString NewCurveCSV = NewCurveTable? NewCurveTable->GetTableAsCSV() : TEXT("");
+	const bool NewResult = FFileHelper::SaveStringToFile(NewCurveCSV, *AbsoluteNewTempFileName);
 
 	if (OldResult && NewResult)
 	{
diff --git a/Engine/Plugins/Editor/EngineAssetDefinitions/Source/Private/Table/AssetDefinition_DataTable.cpp b/Engine/Plugins/Editor/EngineAssetDefinitions/Source/Private/Table/AssetDefinition_DataTable.cpp
--- a/Engine/Plugins/Editor/EngineAssetDefinitions/Source/Private/Table/AssetDefinition_DataTable.cpp
+++ b/Engine/Plugins/Editor/EngineAssetDefinitions/Source/Private/Table/AssetDefinition_DataTable.cpp
@@ -46,18 +46,28 @@ FText UAssetDefinition_DataTable::GetAssetDisplayName(const FAssetData& AssetDat
 // Attempts to export temporary CSV files and diff those. If that fails we fall back to diffing the data table assets directly.
 EAssetCommandResult UAssetDefinition_DataTable::PerformAssetDiff(const FAssetDiffArgs& DiffArgs) const
 {
-	const UDataTable* OldDataTable = CastChecked<UDataTable>(DiffArgs.OldAsset);
-	const UDataTable* NewDataTable = CastChecked<UDataTable>(DiffArgs.NewAsset);
+	const UDataTable* OldDataTable = Cast<UDataTable>(DiffArgs.OldAsset);
+	const UDataTable* NewDataTable = Cast<UDataTable>(DiffArgs.NewAsset);
 
+	if (OldDataTable == nullptr && NewDataTable == nullptr)
+	{
+		return EAssetCommandResult::Unhandled;
+	}
+	
 	// Build names for temp csv files
-	const FString RelOldTempFileName = FString::Printf(TEXT("%sTemp%s-%s.csv"), *FPaths::DiffDir(), *DiffArgs.OldAsset->GetName(), *DiffArgs.OldRevision.Revision);
+	const FString OldAssetName = DiffArgs.OldAsset ? DiffArgs.OldAsset->GetName() : DiffArgs.NewAsset->GetName();
+	const FString RelOldTempFileName = FString::Printf(TEXT("%sTemp%s-%s.csv"), *FPaths::DiffDir(), *OldAssetName, *DiffArgs.OldRevision.Revision);
 	const FString AbsoluteOldTempFileName = FPaths::ConvertRelativePathToFull(RelOldTempFileName);
-	const FString RelNewTempFileName = FString::Printf(TEXT("%sTemp%s-%s.csv"), *FPaths::DiffDir(), *DiffArgs.NewAsset->GetName(), *DiffArgs.NewRevision.Revision);
+	
+	const FString NewAssetName = DiffArgs.NewAsset ? DiffArgs.NewAsset->GetName() : DiffArgs.OldAsset->GetName();
+	const FString RelNewTempFileName = FString::Printf(TEXT("%sTemp%s-%s.csv"), *FPaths::DiffDir(), *NewAssetName, *DiffArgs.NewRevision.Revision);
 	const FString AbsoluteNewTempFileName = FPaths::ConvertRelativePathToFull(RelNewTempFileName);
 
 	// save temp files
-	const bool OldResult = FFileHelper::SaveStringToFile(OldDataTable->GetTableAsCSV(EDataTableExportFlags::UseSimpleText), *AbsoluteOldTempFileName);
-	const bool NewResult = FFileHelper::SaveStringToFile(NewDataTable->GetTableAsCSV(EDataTableExportFlags::UseSimpleText), *AbsoluteNewTempFileName);
+	const FString OldTableCSV = OldDataTable? OldDataTable->GetTableAsCSV() : TEXT("");
+	const bool OldResult = FFileHelper::SaveStringToFile(OldTableCSV, *AbsoluteOldTempFileName);
+	const FString NewTableCSV = NewDataTable? NewDataTable->GetTableAsCSV() : TEXT("");
+	const bool NewResult = FFileHelper::SaveStringToFile(NewTableCSV, *AbsoluteNewTempFileName);
 
 	if (OldResult && NewResult)
 	{
diff --git a/Engine/Plugins/Experimental/CommonConversation/Source/CommonConversationEditor/Private/AssetTypeActions_ConversationDatabase.cpp b/Engine/Plugins/Experimental/CommonConversation/Source/CommonConversationEditor/Private/AssetTypeActions_ConversationDatabase.cpp
--- a/Engine/Plugins/Experimental/CommonConversation/Source/CommonConversationEditor/Private/AssetTypeActions_ConversationDatabase.cpp
+++ b/Engine/Plugins/Experimental/CommonConversation/Source/CommonConversationEditor/Private/AssetTypeActions_ConversationDatabase.cpp
@@ -35,17 +35,16 @@ UClass* FAssetTypeActions_ConversationDatabase::GetSupportedClass() const
 void FAssetTypeActions_ConversationDatabase::PerformAssetDiff(UObject* OldAsset, UObject* NewAsset, const struct FRevisionInfo& OldRevision, const struct FRevisionInfo& NewRevision) const
 {
 	UConversationDatabase* OldBank = Cast<UConversationDatabase>(OldAsset);
-	check(OldBank != nullptr);
-
 	UConversationDatabase* NewBank = Cast<UConversationDatabase>(NewAsset);
-	check(NewBank != nullptr);
+	check(NewBank != nullptr || OldBank != nullptr);
 
 	// sometimes we're comparing different revisions of one single asset (other 
 	// times we're comparing two completely separate assets altogether)
-	const bool bIsSingleAsset = (NewBank->GetName() == OldBank->GetName());
+	const bool bIsSingleAsset = !NewBank || !OldBank || (NewBank->GetName() == OldBank->GetName());
 
+	const FString BankName = NewBank? NewBank->GetName() : OldBank->GetName();
 	const FText WindowTitle = bIsSingleAsset ?
-		FText::Format(LOCTEXT("Conversation Diff", "{0} - Conversation Diff"), FText::FromString(NewBank->GetName())) :
+		FText::Format(LOCTEXT("Conversation Diff", "{0} - Conversation Diff"), FText::FromString(BankName)) :
 		LOCTEXT("NamelessConversationDiff", "Conversation Diff");
 
 	const TSharedPtr<SWindow> Window = SNew(SWindow)
diff --git a/Engine/Plugins/Experimental/CommonConversation/Source/CommonConversationEditor/Private/SConversationDiff.cpp b/Engine/Plugins/Experimental/CommonConversation/Source/CommonConversationEditor/Private/SConversationDiff.cpp
--- a/Engine/Plugins/Experimental/CommonConversation/Source/CommonConversationEditor/Private/SConversationDiff.cpp
+++ b/Engine/Plugins/Experimental/CommonConversation/Source/CommonConversationEditor/Private/SConversationDiff.cpp
@@ -179,8 +179,8 @@ void SConversationDiff::Construct( const FArguments& InArgs )
 	];
 
 	//@TODO: CONVERSATION: Diff tool doesn't support more than one graph yet
-	UEdGraph* OldGraph = FConversationCompiler::GetGraphFromBank(PanelOld.ConversationBank, 0);
-	UEdGraph* NewGraph = FConversationCompiler::GetGraphFromBank(PanelNew.ConversationBank, 0);
+	UEdGraph* OldGraph = PanelOld.ConversationBank? FConversationCompiler::GetGraphFromBank(PanelOld.ConversationBank, 0) : nullptr;
+	UEdGraph* NewGraph = PanelNew.ConversationBank? FConversationCompiler::GetGraphFromBank(PanelNew.ConversationBank, 0) : nullptr;
 	PanelOld.GeneratePanel(OldGraph, NewGraph);
 	PanelNew.GeneratePanel(NewGraph, OldGraph);
 }
@@ -268,8 +268,10 @@ void SConversationDiff::BuildDiffSourceArray()
 {
 	TArray<FDiffSingleResult> FoundDiffs;
 	//@TODO: CONVERSATION: Support diffing multiple graphs
-	FGraphDiffControl::DiffGraphs(FConversationCompiler::GetGraphFromBank(PanelOld.ConversationBank, 0), FConversationCompiler::GetGraphFromBank(PanelNew.ConversationBank, 0), FoundDiffs);
-
+	if (PanelOld.ConversationBank && PanelNew.ConversationBank)
+	{
+		FGraphDiffControl::DiffGraphs(FConversationCompiler::GetGraphFromBank(PanelOld.ConversationBank, 0), FConversationCompiler::GetGraphFromBank(PanelNew.ConversationBank, 0), FoundDiffs);
+	}
 	DiffListSource.Empty();
 	for (auto DiffIt(FoundDiffs.CreateConstIterator()); DiffIt; ++DiffIt)
 	{
@@ -502,7 +504,7 @@ FText SConversationDiff::FConversationDiffPanel::GetTitle() const
 		const FText RevisionText = FText::FromString(RevisionInfo.Revision);
 		const FText ChangelistText = FText::AsNumber(RevisionInfo.Changelist, &FNumberFormattingOptions::DefaultNoGrouping());
 
-		if (bShowAssetName)
+		if (bShowAssetName && ConversationBank)
 		{
 			FString AssetName = ConversationBank->GetName();
 			if(ISourceControlModule::Get().GetProvider().UsesChangelists())
@@ -530,7 +532,7 @@ FText SConversationDiff::FConversationDiffPanel::GetTitle() const
 			}
 		}
 	}
-	else if (bShowAssetName)
+	else if (bShowAssetName && ConversationBank)
 	{
 		FString AssetName = ConversationBank->GetName();
 		FText LocalizedFormat = LOCTEXT("NamedCurrentRevisionFmt", "{0} - Current Revision");
diff --git a/Engine/Plugins/Runtime/ModelViewViewModel/Source/ModelViewViewModelEditor/Private/ViewModel/AssetTypeActions_ViewModelBlueprint.cpp b/Engine/Plugins/Runtime/ModelViewViewModel/Source/ModelViewViewModelEditor/Private/ViewModel/AssetTypeActions_ViewModelBlueprint.cpp
--- a/Engine/Plugins/Runtime/ModelViewViewModel/Source/ModelViewViewModelEditor/Private/ViewModel/AssetTypeActions_ViewModelBlueprint.cpp
+++ b/Engine/Plugins/Runtime/ModelViewViewModel/Source/ModelViewViewModelEditor/Private/ViewModel/AssetTypeActions_ViewModelBlueprint.cpp
@@ -60,21 +60,9 @@ FText FAssetTypeActions_ViewModelBlueprint::GetAssetDescription( const FAssetDat
 
 void FAssetTypeActions_ViewModelBlueprint::PerformAssetDiff(UObject* Asset1, UObject* Asset2, const struct FRevisionInfo& OldRevision, const struct FRevisionInfo& NewRevision) const
 {
-	UBlueprint* OldBlueprint = CastChecked<UBlueprint>(Asset1);
-	UBlueprint* NewBlueprint = CastChecked<UBlueprint>(Asset2);
-
-	// sometimes we're comparing different revisions of one single asset (other 
-	//times we're comparing two completely separate assets altogether)
-	bool bIsSingleAsset = (NewBlueprint->GetName() == OldBlueprint->GetName());
-
-	FText WindowTitle = LOCTEXT("NamelessViewmodelBlueprintDiff", "Viewmodel Blueprint Diff");
-	if (bIsSingleAsset)
-	{
-		// identify the assumed single asset in the window's title
-		WindowTitle = FText::Format(LOCTEXT("WidgetBlueprintDiff", "{0} - Viewmodel Blueprint Diff"), FText::FromString(NewBlueprint->GetName()));
-	}
-
-	SBlueprintDiff::CreateDiffWindow(WindowTitle, OldBlueprint, NewBlueprint, OldRevision, NewRevision);
+	UBlueprint* OldBlueprint = Cast<UBlueprint>(Asset1);
+	UBlueprint* NewBlueprint = Cast<UBlueprint>(Asset2);
+	SBlueprintDiff::CreateDiffWindow(OldBlueprint, NewBlueprint, OldRevision, NewRevision, GetSupportedClass());
 }
 
 } //namespace
diff --git a/Engine/Source/Developer/AssetTools/Private/AssetTools.cpp b/Engine/Source/Developer/AssetTools/Private/AssetTools.cpp
--- a/Engine/Source/Developer/AssetTools/Private/AssetTools.cpp
+++ b/Engine/Source/Developer/AssetTools/Private/AssetTools.cpp
@@ -2717,15 +2717,15 @@ void UAssetToolsImpl::DiffAgainstDepot( UObject* InObject, const FString& InPack
 
 void UAssetToolsImpl::DiffAssets(UObject* OldAsset, UObject* NewAsset, const struct FRevisionInfo& OldRevision, const struct FRevisionInfo& NewRevision) const
 {
-	if(OldAsset == nullptr || NewAsset == nullptr)
+	if (OldAsset == nullptr && NewAsset == nullptr)
 	{
-		UE_LOG(LogAssetTools, Warning, TEXT("DiffAssets: One of the supplied assets was nullptr."));
+		UE_LOG(LogAssetTools, Warning, TEXT("DiffAssets: Both of the supplied assets were nullptr."));
 		return;
 	}
 
 	// Get class of both assets 
-	UClass* OldClass = OldAsset->GetClass();
-	UClass* NewClass = NewAsset->GetClass();
+	const UClass* OldClass = OldAsset ? OldAsset->GetClass() : NewAsset->GetClass();
+	const UClass* NewClass = NewAsset ? NewAsset->GetClass() : OldAsset->GetClass();
 	// If same class..
 	if(OldClass == NewClass)
 	{
@@ -2745,8 +2745,6 @@ void UAssetToolsImpl::DiffAssets(UObject* OldAsset, UObject* NewAsset, const str
 
 FString UAssetToolsImpl::DumpAssetToTempFile(UObject* Asset) const
 {
-	check(Asset);
-
 	// Clear the mark state for saving.
 	UnMarkAllObjects(EObjectMark(OBJECTMARK_TagExp | OBJECTMARK_TagImp));
 
@@ -2754,13 +2752,16 @@ FString UAssetToolsImpl::DumpAssetToTempFile(UObject* Asset) const
 	const FExportObjectInnerContext Context;
 
 	// Export asset to archive
-	UExporter::ExportToOutputDevice(&Context, Asset, nullptr, Archive, TEXT("copy"), 0, PPF_ExportsNotFullyQualified|PPF_Copy|PPF_Delimited, false, Asset->GetOuter());
-
+	if (Asset)
+	{
+		UExporter::ExportToOutputDevice(&Context, Asset, nullptr, Archive, TEXT("copy"), 0, PPF_ExportsNotFullyQualified|PPF_Copy|PPF_Delimited, false, Asset->GetOuter());
+	}
 	// Used to generate unique file names during a run
 	static int TempFileNum = 0;
 
 	// Build name for temp text file
-	FString RelTempFileName = FString::Printf(TEXT("%sText%s-%d.txt"), *FPaths::DiffDir(), *Asset->GetName(), TempFileNum++);
+	FString AssetName = Asset? Asset->GetName() : TEXT("empty");
+	FString RelTempFileName = FString::Printf(TEXT("%sText%s-%d.txt"), *FPaths::DiffDir(), *AssetName, TempFileNum++);
 	FString AbsoluteTempFileName = FPaths::ConvertRelativePathToFull(RelTempFileName);
 
 	// Save text into temp file
diff --git a/Engine/Source/Developer/AssetTools/Private/AssetTypeActions/AssetTypeActions_Blueprint.cpp b/Engine/Source/Developer/AssetTools/Private/AssetTypeActions/AssetTypeActions_Blueprint.cpp
--- a/Engine/Source/Developer/AssetTools/Private/AssetTypeActions/AssetTypeActions_Blueprint.cpp
+++ b/Engine/Source/Developer/AssetTools/Private/AssetTypeActions/AssetTypeActions_Blueprint.cpp
@@ -240,22 +240,10 @@ bool FAssetTypeActions_Blueprint::ShouldUseDataOnlyEditor( const UBlueprint* Blu
 
 void FAssetTypeActions_Blueprint::PerformAssetDiff(UObject* OldAsset, UObject* NewAsset, const FRevisionInfo& OldRevision, const FRevisionInfo& NewRevision) const
 {
-	UBlueprint* OldBlueprint = CastChecked<UBlueprint>(OldAsset);
-	UBlueprint* NewBlueprint = CastChecked<UBlueprint>(NewAsset);
+	UBlueprint* OldBlueprint = Cast<UBlueprint>(OldAsset);
+	UBlueprint* NewBlueprint = Cast<UBlueprint>(NewAsset);
 
-	// sometimes we're comparing different revisions of one single asset (other 
-	// times we're comparing two completely separate assets altogether)
-	bool bIsSingleAsset = (NewBlueprint->GetName() == OldBlueprint->GetName());
-
-	FText WindowTitle = LOCTEXT("NamelessBlueprintDiff", "Blueprint Diff");
-	// if we're diffing one asset against itself 
-	if (bIsSingleAsset)
-	{
-		// identify the assumed single asset in the window's title
-		WindowTitle = FText::Format(LOCTEXT("Blueprint Diff", "{0} - Blueprint Diff"), FText::FromString(NewBlueprint->GetName()));
-	}
-
-	SBlueprintDiff::CreateDiffWindow(WindowTitle, OldBlueprint, NewBlueprint, OldRevision, NewRevision);
+	SBlueprintDiff::CreateDiffWindow(OldBlueprint, NewBlueprint, OldRevision, NewRevision, GetSupportedClass());
 }
 
 UThumbnailInfo* FAssetTypeActions_Blueprint::GetThumbnailInfo(UObject* Asset) const
diff --git a/Engine/Source/Developer/AssetTools/Private/AssetTypeActions/AssetTypeActions_DataAsset.cpp b/Engine/Source/Developer/AssetTools/Private/AssetTypeActions/AssetTypeActions_DataAsset.cpp
--- a/Engine/Source/Developer/AssetTools/Private/AssetTypeActions/AssetTypeActions_DataAsset.cpp
+++ b/Engine/Source/Developer/AssetTools/Private/AssetTypeActions/AssetTypeActions_DataAsset.cpp
@@ -73,14 +73,7 @@ void FAssetTypeActions_DataAsset::GetActions(const TArray<UObject*>& InObjects,
 
 void FAssetTypeActions_DataAsset::PerformAssetDiff(UObject* OldAsset, UObject* NewAsset, const FRevisionInfo& OldRevision, const FRevisionInfo& NewRevision) const
 {
-	// sometimes we're comparing different revisions of one single asset (other 
-	// times we're comparing two completely separate assets altogether)
-	const bool bIsSingleAsset = (OldAsset->GetName() == NewAsset->GetName());
-	static const FText BasicWindowTitle = LOCTEXT("NamelessDataAssetDiff", "DataAsset Diff");
-
-	const FText WindowTitle = !bIsSingleAsset ? BasicWindowTitle : FText::Format(LOCTEXT("DataAsset Diff", "{0} - DataAsset Diff"), FText::FromString(NewAsset->GetName()));
-
-	SDetailsDiff::CreateDiffWindow(WindowTitle, OldAsset, NewAsset, OldRevision, NewRevision);
+	SDetailsDiff::CreateDiffWindow(OldAsset, NewAsset, OldRevision, NewRevision, GetSupportedClass());
 }
 
 void FAssetTypeActions_DataAsset::ExecuteChangeDataAssetClass(TArray<TWeakObjectPtr<UDataAsset>> InDataAssets)
diff --git a/Engine/Source/Developer/Merge/Private/SMergeGraphView.cpp b/Engine/Source/Developer/Merge/Private/SMergeGraphView.cpp
--- a/Engine/Source/Developer/Merge/Private/SMergeGraphView.cpp
+++ b/Engine/Source/Developer/Merge/Private/SMergeGraphView.cpp
@@ -383,7 +383,7 @@ void SMergeGraphView::Construct(const FArguments InArgs
 	{
 		DetailsPanelContainer->AddSlot()
 		[
-			Panel.DetailsView.ToSharedRef()
+			Panel.GetDetailsWidget()
 		];
 	}
 
diff --git a/Engine/Source/Editor/BehaviorTreeEditor/Private/AssetTypeActions_BehaviorTree.cpp b/Engine/Source/Editor/BehaviorTreeEditor/Private/AssetTypeActions_BehaviorTree.cpp
--- a/Engine/Source/Editor/BehaviorTreeEditor/Private/AssetTypeActions_BehaviorTree.cpp
+++ b/Engine/Source/Editor/BehaviorTreeEditor/Private/AssetTypeActions_BehaviorTree.cpp
@@ -81,21 +81,20 @@ UClass* FAssetTypeActions_BehaviorTree::GetSupportedClass() const
 void FAssetTypeActions_BehaviorTree::PerformAssetDiff(UObject* OldAsset, UObject* NewAsset, const struct FRevisionInfo& OldRevision, const struct FRevisionInfo& NewRevision) const
 {
 	UBehaviorTree* OldBehaviorTree = Cast<UBehaviorTree>(OldAsset);
-	check(OldBehaviorTree != NULL);
-
 	UBehaviorTree* NewBehaviorTree = Cast<UBehaviorTree>(NewAsset);
-	check(NewBehaviorTree != NULL);
+	check(NewBehaviorTree != nullptr || OldBehaviorTree != nullptr);
 
 	// sometimes we're comparing different revisions of one single asset (other 
 	// times we're comparing two completely separate assets altogether)
-	bool bIsSingleAsset = (NewBehaviorTree->GetName() == OldBehaviorTree->GetName());
+	const bool bIsSingleAsset = !NewBehaviorTree || !OldBehaviorTree || (NewBehaviorTree->GetName() == OldBehaviorTree->GetName());
 
 	FText WindowTitle = LOCTEXT("NamelessBehaviorTreeDiff", "Behavior Tree Diff");
 	// if we're diffing one asset against itself 
 	if (bIsSingleAsset)
 	{
 		// identify the assumed single asset in the window's title
-		WindowTitle = FText::Format(LOCTEXT("Behavior Tree Diff", "{0} - Behavior Tree Diff"), FText::FromString(NewBehaviorTree->GetName()));
+		const FString TreeName = NewBehaviorTree? NewBehaviorTree->GetName() : OldBehaviorTree->GetName();
+		WindowTitle = FText::Format(LOCTEXT("Behavior Tree Diff", "{0} - Behavior Tree Diff"), FText::FromString(TreeName));
 	}
 
 	const TSharedPtr<SWindow> Window = SNew(SWindow)
diff --git a/Engine/Source/Editor/BehaviorTreeEditor/Private/SBehaviorTreeDiff.cpp b/Engine/Source/Editor/BehaviorTreeEditor/Private/SBehaviorTreeDiff.cpp
--- a/Engine/Source/Editor/BehaviorTreeEditor/Private/SBehaviorTreeDiff.cpp
+++ b/Engine/Source/Editor/BehaviorTreeEditor/Private/SBehaviorTreeDiff.cpp
@@ -223,8 +223,8 @@ void SBehaviorTreeDiff::Construct( const FArguments& InArgs )
 		]
 	];
 
-	PanelOld.GeneratePanel(PanelOld.BehaviorTree->BTGraph, FoundDiffs);
-	PanelNew.GeneratePanel(PanelNew.BehaviorTree->BTGraph, FoundDiffs);
+	PanelOld.GeneratePanel(PanelOld.BehaviorTree ? PanelOld.BehaviorTree->BTGraph : nullptr, FoundDiffs);
+	PanelNew.GeneratePanel(PanelNew.BehaviorTree ? PanelNew.BehaviorTree->BTGraph : nullptr, FoundDiffs);
 }
 END_SLATE_FUNCTION_BUILD_OPTIMIZATION
 
@@ -309,9 +309,14 @@ TSharedRef<SWidget> SBehaviorTreeDiff::GenerateDiffListWidget()
 void SBehaviorTreeDiff::BuildDiffSourceArray()
 {
 	FoundDiffs->Empty();
+	DiffListSource.Empty();
+	if (!PanelOld.BehaviorTree || !PanelNew.BehaviorTree)
+	{
+		return;
+	}
+	
 	FGraphDiffControl::DiffGraphs(PanelOld.BehaviorTree->BTGraph, PanelNew.BehaviorTree->BTGraph, *FoundDiffs);
 
-	DiffListSource.Empty();
 	for (auto DiffIt(FoundDiffs->CreateConstIterator()); DiffIt; ++DiffIt)
 	{
 		DiffListSource.Add(FSharedDiffOnGraph(new FTreeDiffResultItem(*DiffIt)));
@@ -546,7 +551,7 @@ FText SBehaviorTreeDiff::FBehaviorTreeDiffPanel::GetTitle() const
 		const FText RevisionText = FText::FromString(RevisionInfo.Revision);
 		const FText ChangelistText = FText::AsNumber(RevisionInfo.Changelist, &FNumberFormattingOptions::DefaultNoGrouping());
 
-		if (bShowAssetName)
+		if (bShowAssetName && BehaviorTree)
 		{
 			FString AssetName = BehaviorTree->GetName();
 			if(ISourceControlModule::Get().GetProvider().UsesChangelists())
@@ -574,7 +579,7 @@ FText SBehaviorTreeDiff::FBehaviorTreeDiffPanel::GetTitle() const
 			}
 		}
 	}
-	else if (bShowAssetName)
+	else if (bShowAssetName && BehaviorTree)
 	{
 		FString AssetName = BehaviorTree->GetName();
 		FText LocalizedFormat = LOCTEXT("NamedCurrentRevisionFmt", "{0} - Current Revision");
diff --git a/Engine/Source/Editor/Blutility/Private/AssetTypeActions_EditorUtilityWidgetBlueprint.cpp b/Engine/Source/Editor/Blutility/Private/AssetTypeActions_EditorUtilityWidgetBlueprint.cpp
--- a/Engine/Source/Editor/Blutility/Private/AssetTypeActions_EditorUtilityWidgetBlueprint.cpp
+++ b/Engine/Source/Editor/Blutility/Private/AssetTypeActions_EditorUtilityWidgetBlueprint.cpp
@@ -98,22 +98,9 @@ uint32 FAssetTypeActions_EditorUtilityWidgetBlueprint::GetCategories()
 
 void FAssetTypeActions_EditorUtilityWidgetBlueprint::PerformAssetDiff(UObject* Asset1, UObject* Asset2, const struct FRevisionInfo& OldRevision, const struct FRevisionInfo& NewRevision) const
 {
-	UBlueprint* OldBlueprint = CastChecked<UBlueprint>(Asset1);
-	UBlueprint* NewBlueprint = CastChecked<UBlueprint>(Asset2);
-
-	// sometimes we're comparing different revisions of one single asset (other 
-	// times we're comparing two completely separate assets altogether)
-	bool bIsSingleAsset = (NewBlueprint->GetName() == OldBlueprint->GetName());
-
-	FText WindowTitle = LOCTEXT("NamelessEditorUtilityWidgetBlueprintDiff", "Editor Utility Widget Blueprint Diff");
-	// if we're diffing one asset against itself 
-	if (bIsSingleAsset)
-	{
-		// identify the assumed single asset in the window's title
-		WindowTitle = FText::Format(LOCTEXT("EditorUtilityWidgetBlueprintDiff", "{0} - Editor Utility Widget Blueprint Diff"), FText::FromString(NewBlueprint->GetName()));
-	}
-
-	SBlueprintDiff::CreateDiffWindow(WindowTitle, OldBlueprint, NewBlueprint, OldRevision, NewRevision);
+	UBlueprint* OldBlueprint = Cast<UBlueprint>(Asset1);
+	UBlueprint* NewBlueprint = Cast<UBlueprint>(Asset2);
+	SBlueprintDiff::CreateDiffWindow(OldBlueprint, NewBlueprint, OldRevision, NewRevision, GetSupportedClass());
 }
 
 void FAssetTypeActions_EditorUtilityWidgetBlueprint::ExecuteRun(FWeakBlueprintPointerArray InObjects)
diff --git a/Engine/Source/Editor/Kismet/Private/DetailsDiff.cpp b/Engine/Source/Editor/Kismet/Private/DetailsDiff.cpp
--- a/Engine/Source/Editor/Kismet/Private/DetailsDiff.cpp
+++ b/Engine/Source/Editor/Kismet/Private/DetailsDiff.cpp
@@ -34,7 +34,7 @@ FDetailsDiff::FDetailsDiff(const UObject* InObject, FOnDisplayedPropertiesChange
 	FPropertyEditorModule& EditModule = FModuleManager::Get().GetModuleChecked<FPropertyEditorModule>("PropertyEditor");
 	DetailsView = EditModule.CreateDetailView(DetailsViewArgs);
 	DetailsView->SetIsPropertyEditingEnabledDelegate(FIsPropertyEditingEnabled::CreateStatic([]{return false; }));
-	if (InObject->IsA<UBlueprint>())
+	if (InObject && InObject->IsA<UBlueprint>())
 	{
 		// create a custom property layout so that sections like interfaces will be included in the diff view
 		const FOnGetDetailCustomizationInstance LayoutOptionDetails = FOnGetDetailCustomizationInstance::CreateStatic(
@@ -89,6 +89,11 @@ TArray<FPropertySoftPath> FDetailsDiff::GetDisplayedProperties() const
 
 void FDetailsDiff::DiffAgainst(const FDetailsDiff& Newer, TArray< FSingleObjectDiffEntry > &OutDifferences, bool bSortByDisplayOrder) const
 {
+	if (!DisplayedObject || !Newer.DisplayedObject)
+	{
+		return;
+	}
+
 	TSharedPtr< class IDetailsView > OldDetailsView = DetailsView;
 	TSharedPtr< class IDetailsView > NewDetailsView = Newer.DetailsView;
 
diff --git a/Engine/Source/Editor/Kismet/Private/DiffControl.cpp b/Engine/Source/Editor/Kismet/Private/DiffControl.cpp
--- a/Engine/Source/Editor/Kismet/Private/DiffControl.cpp
+++ b/Engine/Source/Editor/Kismet/Private/DiffControl.cpp
@@ -164,37 +164,39 @@ FSCSDiffControl::FSCSDiffControl(const UBlueprint* InOldBlueprint, const UBluepr
 
 void FSCSDiffControl::GenerateTreeEntries(TArray<TSharedPtr<FBlueprintDifferenceTreeEntry>>& OutTreeEntries, TArray<TSharedPtr<FBlueprintDifferenceTreeEntry>>& OutRealDifferences)
 {
-	TArray< FSCSResolvedIdentifier > OldHierarchy = OldSCS.GetDisplayedHierarchy();
-	TArray< FSCSResolvedIdentifier > NewHierarchy = NewSCS.GetDisplayedHierarchy();
-	DiffUtils::CompareUnrelatedSCS(OldSCS.GetBlueprint(), OldHierarchy, NewSCS.GetBlueprint(), NewHierarchy, DifferingProperties);
-
-	const auto FocusSCSDifferenceEntry = [](FSCSDiffEntry Entry, FOnDiffEntryFocused InSelectionCallback, FSCSDiffControl* Owner)
+	TArray<TSharedPtr<FBlueprintDifferenceTreeEntry>> Children;
+	if (OldSCS.GetBlueprint() && NewSCS.GetBlueprint())
 	{
-		InSelectionCallback.ExecuteIfBound();
-		if (Entry.TreeIdentifier.Name != NAME_None)
+		const TArray< FSCSResolvedIdentifier > OldHierarchy = OldSCS.GetDisplayedHierarchy();
+		const TArray< FSCSResolvedIdentifier > NewHierarchy = NewSCS.GetDisplayedHierarchy();
+		DiffUtils::CompareUnrelatedSCS(OldSCS.GetBlueprint(), OldHierarchy, NewSCS.GetBlueprint(), NewHierarchy, DifferingProperties);
+
+		const auto FocusSCSDifferenceEntry = [](FSCSDiffEntry Entry, FOnDiffEntryFocused InSelectionCallback, FSCSDiffControl* Owner)
 		{
-			Owner->OldSCS.HighlightProperty(Entry.TreeIdentifier.Name, FPropertyPath());
-			Owner->NewSCS.HighlightProperty(Entry.TreeIdentifier.Name, FPropertyPath());
-		}
-	};
+			InSelectionCallback.ExecuteIfBound();
+			if (Entry.TreeIdentifier.Name != NAME_None)
+			{
+				Owner->OldSCS.HighlightProperty(Entry.TreeIdentifier.Name, FPropertyPath());
+				Owner->NewSCS.HighlightProperty(Entry.TreeIdentifier.Name, FPropertyPath());
+			}
+		};
 
-	const auto CreateSCSDifferenceWidget = [](FSCSDiffEntry Entry, FText ObjectName) -> TSharedRef<SWidget>
-	{
-		return SNew(STextBlock)
-				.Text(DiffViewUtils::SCSDiffMessage(Entry, ObjectName))
-				.ColorAndOpacity(DiffViewUtils::Differs());
-	};
+		const auto CreateSCSDifferenceWidget = [](FSCSDiffEntry Entry, FText ObjectName) -> TSharedRef<SWidget>
+		{
+			return SNew(STextBlock)
+					.Text(DiffViewUtils::SCSDiffMessage(Entry, ObjectName))
+					.ColorAndOpacity(DiffViewUtils::Differs());
+		};
 
-	TArray<TSharedPtr<FBlueprintDifferenceTreeEntry>> Children;
-	for (const FSCSDiffEntry& Difference : DifferingProperties.Entries)
-	{
-		TSharedPtr<FBlueprintDifferenceTreeEntry> Entry = MakeShared<FBlueprintDifferenceTreeEntry>(
-			FOnDiffEntryFocused::CreateStatic(FocusSCSDifferenceEntry, Difference, SelectionCallback, this),
-			FGenerateDiffEntryWidget::CreateStatic(CreateSCSDifferenceWidget, Difference, RightRevision));
-		Children.Push(Entry);
-		OutRealDifferences.Push(Entry);
+		for (const FSCSDiffEntry& Difference : DifferingProperties.Entries)
+		{
+			TSharedPtr<FBlueprintDifferenceTreeEntry> Entry = MakeShared<FBlueprintDifferenceTreeEntry>(
+				FOnDiffEntryFocused::CreateStatic(FocusSCSDifferenceEntry, Difference, SelectionCallback, this),
+				FGenerateDiffEntryWidget::CreateStatic(CreateSCSDifferenceWidget, Difference, RightRevision));
+			Children.Push(Entry);
+			OutRealDifferences.Push(Entry);
+		}
 	}
-
 	const bool bHasDifferences = Children.Num() != 0;
 	if (!bHasDifferences)
 	{
@@ -249,7 +251,6 @@ void FDetailsDiffControl::GenerateTreeEntries(TArray<TSharedPtr<FBlueprintDiffer
 {
 	for (const FSingleObjectDiffEntry& Difference : DifferingProperties)
 	{
-			
 		TSharedPtr<FBlueprintDifferenceTreeEntry> Entry = MakeShared<FBlueprintDifferenceTreeEntry>(
 			FOnDiffEntryFocused::CreateSP(TSharedFromThis<FDetailsDiffControl>::AsShared(), &FDetailsDiffControl::OnSelectDiffEntry, Difference.Identifier),
 			FGenerateDiffEntryWidget::CreateStatic(&GenerateObjectDiffWidget, Difference, RightRevision));
@@ -397,9 +398,13 @@ void FClassSettingsDiffControl::GenerateTreeEntries(TArray<TSharedPtr<FBlueprint
 /////////////////////////////////////////////////////////////////////////////
 /// FBlueprintTypeDiffControl
 
-FBlueprintTypeDiffControl::FBlueprintTypeDiffControl(const UBlueprint* InBlueprintOld, const UBlueprint* InBlueprintNew, FOnDiffEntryFocused InSelectionCallback): BlueprintOld(InBlueprintOld), BlueprintNew(InBlueprintNew), SelectionCallback(InSelectionCallback), bDiffSucceeded(false)
+FBlueprintTypeDiffControl::FBlueprintTypeDiffControl(const UBlueprint* InBlueprintOld, const UBlueprint* InBlueprintNew, FOnDiffEntryFocused InSelectionCallback)
+	: BlueprintOld(InBlueprintOld)
+	, BlueprintNew(InBlueprintNew)
+	, SelectionCallback(InSelectionCallback)
+	, bDiffSucceeded(false)
 {
-	check(InBlueprintNew && InBlueprintOld);
+	check(InBlueprintNew || InBlueprintOld);
 }
 
 void FBlueprintTypeDiffControl::GenerateTreeEntries(TArray< TSharedPtr<FBlueprintDifferenceTreeEntry> >& OutTreeEntries, TArray< TSharedPtr<FBlueprintDifferenceTreeEntry> >& OutRealDifferences)
@@ -472,7 +477,7 @@ TSharedRef<SWidget> FBlueprintTypeDiffControl::GenerateCategoryWidget(bool bHasR
 		Color = DiffViewUtils::Differs();
 	}
 
-	FText Label = BlueprintNew->GetClass()->GetDisplayNameText();
+	const FText Label = (BlueprintNew ? BlueprintNew : BlueprintOld)->GetClass()->GetDisplayNameText();
 
 	return SNew(SHorizontalBox)
 		+ SHorizontalBox::Slot()
@@ -485,6 +490,11 @@ TSharedRef<SWidget> FBlueprintTypeDiffControl::GenerateCategoryWidget(bool bHasR
 
 void FBlueprintTypeDiffControl::BuildDiffSourceArray()
 {
+	if (!BlueprintNew || !BlueprintOld)
+	{
+		return;
+	}
+	
 	TArray<FDiffSingleResult> BlueprintDiffResults;
 	FDiffResults BlueprintDiffs(&BlueprintDiffResults);
 	if (BlueprintNew->FindDiffs(BlueprintOld, BlueprintDiffs))
diff --git a/Engine/Source/Editor/Kismet/Private/DiffUtils.cpp b/Engine/Source/Editor/Kismet/Private/DiffUtils.cpp
--- a/Engine/Source/Editor/Kismet/Private/DiffUtils.cpp
+++ b/Engine/Source/Editor/Kismet/Private/DiffUtils.cpp
@@ -1043,6 +1043,11 @@ FText DiffViewUtils::SCSDiffMessage(const FSCSDiffEntry& Difference, FText Objec
 
 FText DiffViewUtils::GetPanelLabel(const UObject* Asset, const FRevisionInfo& Revision, FText Label)
 {
+	if( !Asset )
+	{
+		return NSLOCTEXT("DiffViewUtils", "NoBlueprint", "None" );
+	}
+	
 	if( !Revision.Revision.IsEmpty() )
 	{
 		FText RevisionData;
@@ -1077,24 +1082,19 @@ FText DiffViewUtils::GetPanelLabel(const UObject* Asset, const FRevisionInfo& Re
 	}
 	else
 	{
-		if( Asset )
+		if (Label.IsEmpty())
 		{
-			if (Label.IsEmpty())
-			{
-				return FText::Format(NSLOCTEXT("DiffViewUtils", "RevisionLabelTwoLines", "{0}\n{1}")
-					, FText::FromString(Asset->GetName())
-					, NSLOCTEXT("DiffViewUtils", "LocalRevisionLabel", "Local Revision"));
-			}
-			else
-			{
-				return FText::Format(NSLOCTEXT("DiffViewUtils", "RevisionLabel", "{0}\n{1}\n{2}")
-					, Label
-					, FText::FromString(Asset->GetName())
-					, NSLOCTEXT("DiffViewUtils", "LocalRevisionLabel", "Local Revision"));
-			}
+			return FText::Format(NSLOCTEXT("DiffViewUtils", "RevisionLabelTwoLines", "{0}\n{1}")
+				, FText::FromString(Asset->GetName())
+				, NSLOCTEXT("DiffViewUtils", "LocalRevisionLabel", "Local Revision"));
+		}
+		else
+		{
+			return FText::Format(NSLOCTEXT("DiffViewUtils", "RevisionLabel", "{0}\n{1}\n{2}")
+				, Label
+				, FText::FromString(Asset->GetName())
+				, NSLOCTEXT("DiffViewUtils", "LocalRevisionLabel", "Local Revision"));
 		}
-
-		return NSLOCTEXT("DiffViewUtils", "NoBlueprint", "None" );
 	}
 }
 
diff --git a/Engine/Source/Editor/Kismet/Private/SBlueprintDiff.cpp b/Engine/Source/Editor/Kismet/Private/SBlueprintDiff.cpp
--- a/Engine/Source/Editor/Kismet/Private/SBlueprintDiff.cpp
+++ b/Engine/Source/Editor/Kismet/Private/SBlueprintDiff.cpp
@@ -157,19 +157,19 @@ bool DiffWidgetUtils::HasPrevDifference(SListView< TSharedPtr< FDiffSingleResult
 BEGIN_SLATE_FUNCTION_BUILD_OPTIMIZATION
 void SBlueprintDiff::Construct( const FArguments& InArgs)
 {
-	check(InArgs._BlueprintOld && InArgs._BlueprintNew);
+	check(InArgs._BlueprintOld || InArgs._BlueprintNew);
 	PanelOld.Blueprint = InArgs._BlueprintOld;
 	PanelNew.Blueprint = InArgs._BlueprintNew;
 	PanelOld.RevisionInfo = InArgs._OldRevision;
 	PanelNew.RevisionInfo = InArgs._NewRevision;
 
 	// Create a skeleton if we don't have one, this is true for revision history diffs
-	if (!PanelOld.Blueprint->SkeletonGeneratedClass)
+	if (PanelOld.Blueprint && !PanelOld.Blueprint->SkeletonGeneratedClass)
 	{
 		FKismetEditorUtilities::GenerateBlueprintSkeleton(const_cast<UBlueprint*>(PanelOld.Blueprint));
 	}
 	
-	if (!PanelNew.Blueprint->SkeletonGeneratedClass)
+	if (PanelNew.Blueprint && !PanelNew.Blueprint->SkeletonGeneratedClass)
 	{
 		FKismetEditorUtilities::GenerateBlueprintSkeleton(const_cast<UBlueprint*>(PanelNew.Blueprint));
 	}
@@ -349,8 +349,14 @@ void SBlueprintDiff::Construct( const FArguments& InArgs)
 	SetCurrentMode(MyBlueprintMode);
 
 	// Bind to blueprint changed events as they may be real in memory blueprints that will be modified
-	const_cast<UBlueprint*>(PanelNew.Blueprint)->OnChanged().AddSP(this, &SBlueprintDiff::OnBlueprintChanged);
-	const_cast<UBlueprint*>(PanelOld.Blueprint)->OnChanged().AddSP(this, &SBlueprintDiff::OnBlueprintChanged);
+	if (PanelNew.Blueprint)
+	{
+		const_cast<UBlueprint*>(PanelNew.Blueprint)->OnChanged().AddSP(this, &SBlueprintDiff::OnBlueprintChanged);
+	}
+	if (PanelOld.Blueprint)
+	{
+		const_cast<UBlueprint*>(PanelOld.Blueprint)->OnChanged().AddSP(this, &SBlueprintDiff::OnBlueprintChanged);
+	}
 }
 END_SLATE_FUNCTION_BUILD_OPTIMIZATION
 
@@ -427,11 +433,11 @@ TSharedRef<SWidget> SBlueprintDiff::DefaultEmptyPanel()
 		];
 }
 
-TSharedPtr<SWindow> SBlueprintDiff::CreateDiffWindow(FText WindowTitle, UBlueprint* OldBlueprint, UBlueprint* NewBlueprint, const FRevisionInfo& OldRevision, const FRevisionInfo& NewRevision)
+TSharedPtr<SWindow> SBlueprintDiff::CreateDiffWindow(FText WindowTitle, const UBlueprint* OldBlueprint, const UBlueprint* NewBlueprint, const FRevisionInfo& OldRevision, const FRevisionInfo& NewRevision)
 {
 	// sometimes we're comparing different revisions of one single asset (other 
 	// times we're comparing two completely separate assets altogether)
-	bool bIsSingleAsset = (NewBlueprint->GetName() == OldBlueprint->GetName());
+	bool bIsSingleAsset = !OldBlueprint || !NewBlueprint || (NewBlueprint->GetName() == OldBlueprint->GetName());
 
 	TSharedPtr<SWindow> Window = SNew(SWindow)
 		.Title(WindowTitle)
@@ -446,7 +452,7 @@ TSharedPtr<SWindow> SBlueprintDiff::CreateDiffWindow(FText WindowTitle, UBluepri
 		.ParentWindow(Window));
 
 	// Make this window a child of the modal window if we've been spawned while one is active.
-	TSharedPtr<SWindow> ActiveModal = FSlateApplication::Get().GetActiveModalWindow();
+	const TSharedPtr<SWindow> ActiveModal = FSlateApplication::Get().GetActiveModalWindow();
 	if (ActiveModal.IsValid())
 	{
 		FSlateApplication::Get().AddWindowAsNativeChild(Window.ToSharedRef(), ActiveModal.ToSharedRef());
@@ -459,6 +465,27 @@ TSharedPtr<SWindow> SBlueprintDiff::CreateDiffWindow(FText WindowTitle, UBluepri
 	return Window;
 }
 
+TSharedPtr<SWindow> SBlueprintDiff::CreateDiffWindow(const UBlueprint* OldBlueprint, const UBlueprint* NewBlueprint,
+                                                     const FRevisionInfo& OldRevision, const FRevisionInfo& NewRevision, const UClass* BlueprintClass)
+{
+	check(OldBlueprint || NewBlueprint);
+
+	// sometimes we're comparing different revisions of one single asset (other 
+	// times we're comparing two completely separate assets altogether)
+	const bool bIsSingleAsset = !OldBlueprint || !NewBlueprint || (NewBlueprint->GetName() == OldBlueprint->GetName());
+
+	FText WindowTitle = FText::Format(LOCTEXT("NamelessBlueprintDiff", "{0} Diff"), BlueprintClass->GetDisplayNameText());
+	// if we're diffing one asset against itself 
+	if (bIsSingleAsset)
+	{
+		// identify the assumed single asset in the window's title
+		const FString BPName = NewBlueprint? NewBlueprint->GetName() : OldBlueprint->GetName();
+		WindowTitle = FText::Format(LOCTEXT("NamedBlueprintDiff", "{0} - {1} Diff"), FText::FromString(BPName), BlueprintClass->GetDisplayNameText());
+	}
+
+	return CreateDiffWindow(WindowTitle, OldBlueprint, NewBlueprint, OldRevision, NewRevision);
+}
+
 void SBlueprintDiff::NextDiff()
 {
 	DiffTreeView::HighlightNextDifference(DifferencesTreeView.ToSharedRef(), RealDifferences, PrimaryDifferencesList);
@@ -650,9 +677,13 @@ void FDiffPanel::GeneratePanel(UEdGraph* Graph, TSharedPtr<TArray<FDiffSingleRes
 	GraphEditorBox->SetContent(Widget.ToSharedRef());
 }
 
-TSharedRef<SWidget> FDiffPanel::GenerateMyBlueprintWidget()
+TSharedRef<class SWidget> FDiffPanel::GenerateMyBlueprintWidget()
 {
-	return SAssignNew(MyBlueprint, SMyBlueprint, TWeakPtr<FBlueprintEditor>(), Blueprint);
+	if (Blueprint)
+	{
+		return SAssignNew(MyBlueprint, SMyBlueprint, TWeakPtr<FBlueprintEditor>(), Blueprint);
+	}
+	return SNew(STextBlock).Text(LOCTEXT("NullBlueprint","Missing Blueprint"));
 }
 
 FGraphPanelSelectionSet FDiffPanel::GetSelectedNodes() const
@@ -704,6 +735,18 @@ void FDiffPanel::FocusDiff(UEdGraphNode& Node)
 	}
 }
 
+TSharedRef<SWidget> FDiffPanel::GetMyBlueprintWidget() const
+{
+	// if this panel is displaying a null object, return an empty boarder instead of a MyBlueprint
+	return MyBlueprint ? (TSharedRef<SWidget>)MyBlueprint.ToSharedRef() : TSharedRef<SWidget>(SNew(SBorder));
+}
+
+TSharedRef<SWidget> FDiffPanel::GetDetailsWidget() const
+{
+	// if this panel is displaying a null object, return an empty boarder instead of a details view
+	return DetailsView ? (TSharedRef<SWidget>)DetailsView.ToSharedRef() : TSharedRef<SWidget>(SNew(SBorder));
+}
+
 FDiffPanel& SBlueprintDiff::GetDiffPanelForNode(UEdGraphNode& Node)
 {
 	TSharedPtr<SGraphEditor> OldGraphEditorPtr = PanelOld.GraphEditor.Pin();
@@ -788,16 +831,21 @@ void SBlueprintDiff::GenerateDifferencesList()
 			.ShowLocalVariables(true);
 	};
 
-	PanelOld.GenerateMyBlueprintWidget();
-	PanelOld.DetailsView = CreateInspector(PanelOld.MyBlueprint);
-	PanelOld.MyBlueprint->SetInspector(PanelOld.DetailsView);
-	PanelNew.GenerateMyBlueprintWidget();
-	PanelNew.DetailsView = CreateInspector(PanelNew.MyBlueprint);
-	PanelNew.MyBlueprint->SetInspector(PanelNew.DetailsView);
-
 	TArray<UEdGraph*> GraphsOld, GraphsNew;
-	PanelOld.Blueprint->GetAllGraphs(GraphsOld);
-	PanelNew.Blueprint->GetAllGraphs(GraphsNew);
+	if (PanelOld.Blueprint)
+	{
+		PanelOld.GenerateMyBlueprintWidget();
+		PanelOld.DetailsView = CreateInspector(PanelOld.MyBlueprint);
+		PanelOld.MyBlueprint->SetInspector(PanelOld.DetailsView);
+		PanelOld.Blueprint->GetAllGraphs(GraphsOld);
+	}
+	if (PanelNew.Blueprint)
+	{
+		PanelNew.GenerateMyBlueprintWidget();
+		PanelNew.DetailsView = CreateInspector(PanelNew.MyBlueprint);
+		PanelNew.MyBlueprint->SetInspector(PanelNew.DetailsView);
+		PanelNew.Blueprint->GetAllGraphs(GraphsNew);
+	}
 
 	//Add Graphs that exist in both blueprints, or in blueprint 1 only
 	for (UEdGraph* GraphOld : GraphsOld)
@@ -831,8 +879,8 @@ void SBlueprintDiff::GenerateDifferencesList()
 	}
 
 	bool bHasComponents = false;
-	UClass* BlueprintClassOld = PanelOld.Blueprint->GeneratedClass;
-	UClass* BlueprintClassNew = PanelNew.Blueprint->GeneratedClass;
+	UClass* BlueprintClassOld = PanelOld.Blueprint ? PanelOld.Blueprint->GeneratedClass : PanelNew.Blueprint->GeneratedClass;
+	UClass* BlueprintClassNew = PanelNew.Blueprint ? PanelNew.Blueprint->GeneratedClass : PanelOld.Blueprint->GeneratedClass;
 	const bool bIsOldClassActor = BlueprintClassOld && BlueprintClassOld->IsChildOf<AActor>();
 	const bool bIsNewClassActor = BlueprintClassNew && BlueprintClassNew->IsChildOf<AActor>();
 	if (bIsOldClassActor || bIsNewClassActor)
@@ -841,7 +889,8 @@ void SBlueprintDiff::GenerateDifferencesList()
 	}
 
 	// If this isn't a normal blueprint type, add the type panel
-	if (PanelOld.Blueprint->GetClass() != UBlueprint::StaticClass())
+	const UClass* RepresentativeClass = PanelOld.Blueprint ? PanelOld.Blueprint->GetClass() : PanelNew.Blueprint->GetClass();
+	if (RepresentativeClass != UBlueprint::StaticClass())
 	{
 		ModePanels.Add(BlueprintTypeMode, GenerateBlueprintTypePanel());
 	}
@@ -849,7 +898,7 @@ void SBlueprintDiff::GenerateDifferencesList()
 	// Now that we have done the diffs, create the panel widgets
 	ModePanels.Add(MyBlueprintMode, GenerateMyBlueprintPanel());
 	ModePanels.Add(GraphMode, GenerateGraphPanel());
-	if (PanelOld.Blueprint->ParentClass && PanelNew.Blueprint->ParentClass)
+	if ((PanelOld.Blueprint && PanelOld.Blueprint->ParentClass) || (PanelNew.Blueprint && PanelNew.Blueprint->ParentClass))
 	{
 		ModePanels.Add(DefaultsMode, GenerateDefaultsPanel());
 	}
@@ -925,11 +974,11 @@ SBlueprintDiff::FDiffControl SBlueprintDiff::GenerateMyBlueprintPanel()
 				.PhysicalSplitterHandleSize(10.0f)
 				+ SSplitter::Slot()
 				[
-					PanelOld.MyBlueprint.ToSharedRef()
+					PanelOld.GetMyBlueprintWidget()
 				]
 				+ SSplitter::Slot()
 				[
-					PanelNew.MyBlueprint.ToSharedRef()
+					PanelNew.GetMyBlueprintWidget()
 				]
 			]
 			+ SSplitter::Slot()
@@ -939,11 +988,11 @@ SBlueprintDiff::FDiffControl SBlueprintDiff::GenerateMyBlueprintPanel()
 				.PhysicalSplitterHandleSize(10.0f)
 				+SSplitter::Slot()
 				[
-					PanelOld.DetailsView.ToSharedRef()
+					PanelOld.GetDetailsWidget()
 				]
 				+ SSplitter::Slot()
 				[
-					PanelNew.DetailsView.ToSharedRef()
+					PanelNew.GetDetailsWidget()
 				]
 			]
 		]
@@ -989,11 +1038,11 @@ SBlueprintDiff::FDiffControl SBlueprintDiff::GenerateGraphPanel()
 				.PhysicalSplitterHandleSize(10.0f)
 				+SSplitter::Slot()
 				[
-					PanelOld.DetailsView.ToSharedRef()
+					PanelOld.GetDetailsWidget()
 				]
 				+ SSplitter::Slot()
 				[
-					PanelNew.DetailsView.ToSharedRef()
+					PanelNew.GetDetailsWidget()
 				]
 			]
 		]
@@ -1125,8 +1174,14 @@ void SBlueprintDiff::SetCurrentMode(FName NewMode)
 	if (FoundControl)
 	{
 		// Reset inspector view
-		PanelOld.DetailsView->ShowDetailsForObjects(TArray<UObject*>());
-		PanelNew.DetailsView->ShowDetailsForObjects(TArray<UObject*>());
+		if (PanelOld.DetailsView)
+		{
+			PanelOld.DetailsView->ShowDetailsForObjects(TArray<UObject*>());
+		}
+		if (PanelNew.DetailsView)
+		{
+			PanelNew.DetailsView->ShowDetailsForObjects(TArray<UObject*>());
+		}
 
 		ModeContents->SetContent(FoundControl->Widget.ToSharedRef());
 	}
diff --git a/Engine/Source/Editor/Kismet/Private/SDetailsDiff.cpp b/Engine/Source/Editor/Kismet/Private/SDetailsDiff.cpp
--- a/Engine/Source/Editor/Kismet/Private/SDetailsDiff.cpp
+++ b/Engine/Source/Editor/Kismet/Private/SDetailsDiff.cpp
@@ -33,7 +33,7 @@ static const FName DetailsMode = FName(TEXT("DetailsMode"));
 BEGIN_SLATE_FUNCTION_BUILD_OPTIMIZATION
 void SDetailsDiff::Construct( const FArguments& InArgs)
 {
-	check(InArgs._AssetOld && InArgs._AssetNew);
+	check(InArgs._AssetOld || InArgs._AssetNew);
 	PanelOld.Object = InArgs._AssetOld;
 	PanelNew.Object = InArgs._AssetNew;
 	PanelOld.RevisionInfo = InArgs._OldRevision;
@@ -213,11 +213,11 @@ TSharedRef<SWidget> SDetailsDiff::DefaultEmptyPanel()
 		];
 }
 
-TSharedPtr<SWindow> SDetailsDiff::CreateDiffWindow(FText WindowTitle, UObject* OldObject, UObject* NewObject, const FRevisionInfo& OldRevision, const FRevisionInfo& NewRevision)
+TSharedPtr<SWindow> SDetailsDiff::CreateDiffWindow(FText WindowTitle, const UObject* OldObject, const UObject* NewObject, const FRevisionInfo& OldRevision, const FRevisionInfo& NewRevision)
 {
 	// sometimes we're comparing different revisions of one single asset (other 
 	// times we're comparing two completely separate assets altogether)
-	bool bIsSingleAsset = (NewObject->GetName() == OldObject->GetName());
+	bool bIsSingleAsset = !NewObject || !OldObject || (NewObject->GetName() == OldObject->GetName());
 
 	TSharedPtr<SWindow> Window = SNew(SWindow)
 		.Title(WindowTitle)
@@ -232,7 +232,7 @@ TSharedPtr<SWindow> SDetailsDiff::CreateDiffWindow(FText WindowTitle, UObject* O
 		.ParentWindow(Window));
 
 	// Make this window a child of the modal window if we've been spawned while one is active.
-	TSharedPtr<SWindow> ActiveModal = FSlateApplication::Get().GetActiveModalWindow();
+	const TSharedPtr<SWindow> ActiveModal = FSlateApplication::Get().GetActiveModalWindow();
 	if (ActiveModal.IsValid())
 	{
 		FSlateApplication::Get().AddWindowAsNativeChild(Window.ToSharedRef(), ActiveModal.ToSharedRef());
@@ -245,6 +245,27 @@ TSharedPtr<SWindow> SDetailsDiff::CreateDiffWindow(FText WindowTitle, UObject* O
 	return Window;
 }
 
+TSharedPtr<SWindow> SDetailsDiff::CreateDiffWindow(const UObject* OldObject, const UObject* NewObject,
+                                                   const FRevisionInfo& OldRevision, const FRevisionInfo& NewRevision, const UClass* ObjectClass)
+{
+	check(OldObject || NewObject);
+
+	// sometimes we're comparing different revisions of one single asset (other 
+	// times we're comparing two completely separate assets altogether)
+	const bool bIsSingleAsset = !OldObject || !NewObject || (NewObject->GetName() == OldObject->GetName());
+
+	FText WindowTitle = FText::Format(LOCTEXT("NamelessBlueprintDiff", "{0} Diff"), ObjectClass->GetDisplayNameText());
+	// if we're diffing one asset against itself 
+	if (bIsSingleAsset)
+	{
+		// identify the assumed single asset in the window's title
+		const FString BPName = NewObject? NewObject->GetName() : OldObject->GetName();
+		WindowTitle = FText::Format(LOCTEXT("NamedBlueprintDiff", "{0} - {1} Diff"), FText::FromString(BPName), ObjectClass->GetDisplayNameText());
+	}
+
+	return CreateDiffWindow(WindowTitle, OldObject, NewObject, OldRevision, NewRevision);
+}
+
 void SDetailsDiff::NextDiff()
 {
 	DiffTreeView::HighlightNextDifference(DifferencesTreeView.ToSharedRef(), RealDifferences, PrimaryDifferencesList);
diff --git a/Engine/Source/Editor/Kismet/Public/SBlueprintDiff.h b/Engine/Source/Editor/Kismet/Public/SBlueprintDiff.h
--- a/Engine/Source/Editor/Kismet/Public/SBlueprintDiff.h
+++ b/Engine/Source/Editor/Kismet/Public/SBlueprintDiff.h
@@ -100,6 +100,9 @@ struct KISMET_API FDiffPanel
 	void FocusDiff(UEdGraphPin& Pin);
 	void FocusDiff(UEdGraphNode& Node);
 
+	TSharedRef<SWidget> GetMyBlueprintWidget() const;
+	TSharedRef<SWidget> GetDetailsWidget() const;
+
 	/** The blueprint that owns the graph we are showing */
 	const UBlueprint*				Blueprint;
 
@@ -162,7 +165,10 @@ public:
 	static TSharedRef<SWidget> DefaultEmptyPanel();
 
 	/** Helper function to create a window that holds a diff widget */
-	static TSharedPtr<SWindow> CreateDiffWindow(FText WindowTitle, UBlueprint* OldBlueprint, UBlueprint* NewBlueprint, const struct FRevisionInfo& OldRevision, const struct FRevisionInfo& NewRevision);
+	static TSharedPtr<SWindow> CreateDiffWindow(FText WindowTitle, const UBlueprint* OldBlueprint, const UBlueprint* NewBlueprint, const struct FRevisionInfo& OldRevision, const struct FRevisionInfo& NewRevision);
+
+	/** Helper function to create a window that holds a diff widget that defaults a window title*/
+	static TSharedPtr<SWindow> CreateDiffWindow(const UBlueprint* OldBlueprint, const UBlueprint* NewBlueprint, const FRevisionInfo& OldRevision, const FRevisionInfo& NewRevision, const UClass* BlueprintClass);
 
 protected:
 	/** Called when user clicks button to go to next difference */
diff --git a/Engine/Source/Editor/Kismet/Public/SDetailsDiff.h b/Engine/Source/Editor/Kismet/Public/SDetailsDiff.h
--- a/Engine/Source/Editor/Kismet/Public/SDetailsDiff.h
+++ b/Engine/Source/Editor/Kismet/Public/SDetailsDiff.h
@@ -59,7 +59,10 @@ public:
 	static TSharedRef<SWidget> DefaultEmptyPanel();
 
 	/** Helper function to create a window that holds a diff widget */
-	static TSharedPtr<SWindow> CreateDiffWindow(FText WindowTitle, UObject* OldObject, UObject* NewObject, const struct FRevisionInfo& OldRevision, const struct FRevisionInfo& NewRevision);
+	static TSharedPtr<SWindow> CreateDiffWindow(FText WindowTitle, const UObject* OldObject, const UObject* NewObject, const struct FRevisionInfo& OldRevision, const struct FRevisionInfo& NewRevision);
+
+	/** Helper function to create a window that holds a diff widget (default window title) */
+	static TSharedPtr<SWindow> CreateDiffWindow(const UObject* OldObject, const UObject* NewObject, const FRevisionInfo& OldRevision, const FRevisionInfo& NewRevision, const UClass* ObjectClass);
 
 protected:
 	/** Called when user clicks button to go to next difference */
diff --git a/Engine/Source/Editor/UMGEditor/Private/AssetTypeActions_WidgetBlueprint.cpp b/Engine/Source/Editor/UMGEditor/Private/AssetTypeActions_WidgetBlueprint.cpp
--- a/Engine/Source/Editor/UMGEditor/Private/AssetTypeActions_WidgetBlueprint.cpp
+++ b/Engine/Source/Editor/UMGEditor/Private/AssetTypeActions_WidgetBlueprint.cpp
@@ -50,22 +50,9 @@ FText FAssetTypeActions_WidgetBlueprint::GetAssetDescription( const FAssetData&
 
 void FAssetTypeActions_WidgetBlueprint::PerformAssetDiff(UObject* Asset1, UObject* Asset2, const struct FRevisionInfo& OldRevision, const struct FRevisionInfo& NewRevision) const
 {
-	UBlueprint* OldBlueprint = CastChecked<UBlueprint>(Asset1);
-	UBlueprint* NewBlueprint = CastChecked<UBlueprint>(Asset2);
-
-	// sometimes we're comparing different revisions of one single asset (other 
-	// times we're comparing two completely separate assets altogether)
-	bool bIsSingleAsset = (NewBlueprint->GetName() == OldBlueprint->GetName());
-
-	FText WindowTitle = LOCTEXT("NamelessWidgetBlueprintDiff", "Widget Blueprint Diff");
-	// if we're diffing one asset against itself 
-	if (bIsSingleAsset)
-	{
-		// identify the assumed single asset in the window's title
-		WindowTitle = FText::Format(LOCTEXT("WidgetBlueprintDiffWindowTitle", "{0} - Widget Blueprint Diff"), FText::FromString(NewBlueprint->GetName()));
-	}
-
-	SBlueprintDiff::CreateDiffWindow(WindowTitle, OldBlueprint, NewBlueprint, OldRevision, NewRevision);
+	const UBlueprint* OldBlueprint = Cast<UBlueprint>(Asset1);
+	const UBlueprint* NewBlueprint = Cast<UBlueprint>(Asset2);
+	SBlueprintDiff::CreateDiffWindow(OldBlueprint, NewBlueprint, OldRevision, NewRevision, GetSupportedClass());
 }
 
 #undef LOCTEXT_NAMESPACE
diff --git a/Engine/Source/Editor/UnrealEd/Private/AssetDefinitionDefault.cpp b/Engine/Source/Editor/UnrealEd/Private/AssetDefinitionDefault.cpp
--- a/Engine/Source/Editor/UnrealEd/Private/AssetDefinitionDefault.cpp
+++ b/Engine/Source/Editor/UnrealEd/Private/AssetDefinitionDefault.cpp
@@ -23,9 +23,11 @@ EAssetCommandResult UAssetDefinitionDefault::OpenAssets(const FAssetOpenArgs& Op
 
 EAssetCommandResult UAssetDefinitionDefault::PerformAssetDiff(const FAssetDiffArgs& DiffArgs) const
 {
-	check(DiffArgs.OldAsset != nullptr);
-	check(DiffArgs.NewAsset != nullptr);
-
+	if (DiffArgs.OldAsset == nullptr && DiffArgs.NewAsset == nullptr)
+	{
+		return EAssetCommandResult::Unhandled;
+	}
+	
 	IAssetTools& AssetTools = FModuleManager::Get().LoadModuleChecked<FAssetToolsModule>("AssetTools").Get();
 
 	// Dump assets to temp text files
