[BugFix] Static Analysis warnings

diff --git a/Engine/Plugins/Editor/EngineAssetDefinitions/Source/Private/AssetDefinition_Curve.cpp b/Engine/Plugins/Editor/EngineAssetDefinitions/Source/Private/AssetDefinition_Curve.cpp
--- a/Engine/Plugins/Editor/EngineAssetDefinitions/Source/Private/AssetDefinition_Curve.cpp
+++ b/Engine/Plugins/Editor/EngineAssetDefinitions/Source/Private/AssetDefinition_Curve.cpp
@@ -33,8 +33,8 @@ EAssetCommandResult UAssetDefinition_Curve::OpenAssets(const FAssetOpenArgs& Ope
 
 EAssetCommandResult UAssetDefinition_Curve::PerformAssetDiff(const FAssetDiffArgs& DiffArgs) const
 {
-	UCurveBase* OldCurve = Cast<UCurveBase>(DiffArgs.OldAsset);
-	UCurveBase* NewCurve = Cast<UCurveBase>(DiffArgs.NewAsset);
+	const UCurveBase* OldCurve = Cast<UCurveBase>(DiffArgs.OldAsset);
+	const UCurveBase* NewCurve = Cast<UCurveBase>(DiffArgs.NewAsset);
 
 	if (NewCurve == nullptr && OldCurve == nullptr)
 	{
@@ -42,11 +42,24 @@ EAssetCommandResult UAssetDefinition_Curve::PerformAssetDiff(const FAssetDiffArg
 	}
 
 	// Build names for temp json files
-	const FString OldAssetName = DiffArgs.OldAsset ? DiffArgs.OldAsset->GetName() : DiffArgs.NewAsset->GetName();
+	const auto AssetNameFallback = [OldCurve, NewCurve]()
+	{
+		if (OldCurve)
+		{
+			return OldCurve->GetName();
+		}
+		if (NewCurve)
+		{
+			return NewCurve->GetName();
+		}
+		return FString();
+	};
+	
+	const FString OldAssetName = OldCurve ? OldCurve->GetName() : AssetNameFallback();
 	const FString RelOldTempFileName = FString::Printf(TEXT("%sTemp%s-%s.json"), *FPaths::DiffDir(), *OldAssetName, *DiffArgs.OldRevision.Revision);
 	const FString AbsoluteOldTempFileName = FPaths::ConvertRelativePathToFull(RelOldTempFileName);
 	
-	const FString NewAssetName = DiffArgs.NewAsset ? DiffArgs.NewAsset->GetName() : DiffArgs.OldAsset->GetName();
+	const FString NewAssetName = NewCurve ? NewCurve->GetName() : AssetNameFallback();
 	const FString RelNewTempFileName = FString::Printf(TEXT("%sTemp%s-%s.json"), *FPaths::DiffDir(), *NewAssetName, *DiffArgs.NewRevision.Revision);
 	const FString AbsoluteNewTempFileName = FPaths::ConvertRelativePathToFull(RelNewTempFileName);
 
diff --git a/Engine/Plugins/Editor/EngineAssetDefinitions/Source/Private/Table/AssetDefinition_CurveTable.cpp b/Engine/Plugins/Editor/EngineAssetDefinitions/Source/Private/Table/AssetDefinition_CurveTable.cpp
--- a/Engine/Plugins/Editor/EngineAssetDefinitions/Source/Private/Table/AssetDefinition_CurveTable.cpp
+++ b/Engine/Plugins/Editor/EngineAssetDefinitions/Source/Private/Table/AssetDefinition_CurveTable.cpp
@@ -31,13 +31,26 @@ EAssetCommandResult UAssetDefinition_CurveTable::PerformAssetDiff(const FAssetDi
 		return EAssetCommandResult::Unhandled;
 	}
 
+	const auto AssetNameFallback = [OldCurveTable, NewCurveTable]()
+	{
+		if (OldCurveTable)
+		{
+			return OldCurveTable->GetName();
+		}
+		if (NewCurveTable)
+		{
+			return NewCurveTable->GetName();
+		}
+		return FString();
+	};
 
 	// Build names for temp csv files
-	const FString OldAssetName = DiffArgs.OldAsset ? DiffArgs.OldAsset->GetName() : DiffArgs.NewAsset->GetName();
+	const FString OldAssetName = OldCurveTable ? OldCurveTable->GetName() : AssetNameFallback();
+	const FString NewAssetName = NewCurveTable ? NewCurveTable->GetName() : AssetNameFallback();
+	
 	const FString RelOldTempFileName = FString::Printf(TEXT("%sTemp%s-%s.csv"), *FPaths::DiffDir(), *OldAssetName, *DiffArgs.OldRevision.Revision);
 	const FString AbsoluteOldTempFileName = FPaths::ConvertRelativePathToFull(RelOldTempFileName);
 	
-	const FString NewAssetName = DiffArgs.NewAsset ? DiffArgs.NewAsset->GetName() : DiffArgs.OldAsset->GetName();
 	const FString RelNewTempFileName = FString::Printf(TEXT("%sTemp%s-%s.csv"), *FPaths::DiffDir(), *NewAssetName, *DiffArgs.NewRevision.Revision);
 	const FString AbsoluteNewTempFileName = FPaths::ConvertRelativePathToFull(RelNewTempFileName);
 
diff --git a/Engine/Plugins/Editor/EngineAssetDefinitions/Source/Private/Table/AssetDefinition_DataTable.cpp b/Engine/Plugins/Editor/EngineAssetDefinitions/Source/Private/Table/AssetDefinition_DataTable.cpp
--- a/Engine/Plugins/Editor/EngineAssetDefinitions/Source/Private/Table/AssetDefinition_DataTable.cpp
+++ b/Engine/Plugins/Editor/EngineAssetDefinitions/Source/Private/Table/AssetDefinition_DataTable.cpp
@@ -55,11 +55,25 @@ EAssetCommandResult UAssetDefinition_DataTable::PerformAssetDiff(const FAssetDif
 	}
 	
 	// Build names for temp csv files
-	const FString OldAssetName = DiffArgs.OldAsset ? DiffArgs.OldAsset->GetName() : DiffArgs.NewAsset->GetName();
+	const auto AssetNameFallback = [NewDataTable, OldDataTable]()
+	{
+		if (OldDataTable)
+		{
+			return OldDataTable->GetName();
+		}
+		if (NewDataTable)
+		{
+			return NewDataTable->GetName();
+		}
+		return FString();
+	};
+
+	const FString OldAssetName = OldDataTable ? OldDataTable->GetName() : AssetNameFallback();
+	const FString NewAssetName = NewDataTable ? NewDataTable->GetName() : AssetNameFallback();
+	
 	const FString RelOldTempFileName = FString::Printf(TEXT("%sTemp%s-%s.csv"), *FPaths::DiffDir(), *OldAssetName, *DiffArgs.OldRevision.Revision);
 	const FString AbsoluteOldTempFileName = FPaths::ConvertRelativePathToFull(RelOldTempFileName);
 	
-	const FString NewAssetName = DiffArgs.NewAsset ? DiffArgs.NewAsset->GetName() : DiffArgs.OldAsset->GetName();
 	const FString RelNewTempFileName = FString::Printf(TEXT("%sTemp%s-%s.csv"), *FPaths::DiffDir(), *NewAssetName, *DiffArgs.NewRevision.Revision);
 	const FString AbsoluteNewTempFileName = FPaths::ConvertRelativePathToFull(RelNewTempFileName);
 
diff --git a/Engine/Plugins/Experimental/CommonConversation/Source/CommonConversationEditor/Private/AssetTypeActions_ConversationDatabase.cpp b/Engine/Plugins/Experimental/CommonConversation/Source/CommonConversationEditor/Private/AssetTypeActions_ConversationDatabase.cpp
--- a/Engine/Plugins/Experimental/CommonConversation/Source/CommonConversationEditor/Private/AssetTypeActions_ConversationDatabase.cpp
+++ b/Engine/Plugins/Experimental/CommonConversation/Source/CommonConversationEditor/Private/AssetTypeActions_ConversationDatabase.cpp
@@ -42,7 +42,15 @@ void FAssetTypeActions_ConversationDatabase::PerformAssetDiff(UObject* OldAsset,
 	// times we're comparing two completely separate assets altogether)
 	const bool bIsSingleAsset = !NewBank || !OldBank || (NewBank->GetName() == OldBank->GetName());
 
-	const FString BankName = NewBank? NewBank->GetName() : OldBank->GetName();
+	FString BankName;
+	if (NewBank)
+	{
+		BankName = NewBank->GetName();
+	}
+	else if (OldBank)
+	{
+		BankName = OldBank->GetName();
+	}
 	const FText WindowTitle = bIsSingleAsset ?
 		FText::Format(LOCTEXT("Conversation Diff", "{0} - Conversation Diff"), FText::FromString(BankName)) :
 		LOCTEXT("NamelessConversationDiff", "Conversation Diff");
diff --git a/Engine/Source/Developer/AssetTools/Private/AssetTools.cpp b/Engine/Source/Developer/AssetTools/Private/AssetTools.cpp
--- a/Engine/Source/Developer/AssetTools/Private/AssetTools.cpp
+++ b/Engine/Source/Developer/AssetTools/Private/AssetTools.cpp
@@ -2723,9 +2723,24 @@ void UAssetToolsImpl::DiffAssets(UObject* OldAsset, UObject* NewAsset, const str
 		return;
 	}
 
-	// Get class of both assets 
-	const UClass* OldClass = OldAsset ? OldAsset->GetClass() : NewAsset->GetClass();
-	const UClass* NewClass = NewAsset ? NewAsset->GetClass() : OldAsset->GetClass();
+	// Get class of both assets
+	const auto AssetClassFallback = [NewAsset, OldAsset]()->UClass*
+	{
+		if (OldAsset)
+		{
+			return OldAsset->GetClass();
+		}
+		if (NewAsset)
+		{
+			return NewAsset->GetClass();
+		}
+		check(false); // this should never happen
+		return nullptr;
+	};
+	
+	const UClass* OldClass = OldAsset ? OldAsset->GetClass() : AssetClassFallback();
+	const UClass* NewClass = NewAsset ? NewAsset->GetClass() : AssetClassFallback();
+	
 	// If same class..
 	if(OldClass == NewClass)
 	{
diff --git a/Engine/Source/Editor/BehaviorTreeEditor/Private/AssetTypeActions_BehaviorTree.cpp b/Engine/Source/Editor/BehaviorTreeEditor/Private/AssetTypeActions_BehaviorTree.cpp
--- a/Engine/Source/Editor/BehaviorTreeEditor/Private/AssetTypeActions_BehaviorTree.cpp
+++ b/Engine/Source/Editor/BehaviorTreeEditor/Private/AssetTypeActions_BehaviorTree.cpp
@@ -93,7 +93,15 @@ void FAssetTypeActions_BehaviorTree::PerformAssetDiff(UObject* OldAsset, UObject
 	if (bIsSingleAsset)
 	{
 		// identify the assumed single asset in the window's title
-		const FString TreeName = NewBehaviorTree? NewBehaviorTree->GetName() : OldBehaviorTree->GetName();
+		FString TreeName;
+		if (NewBehaviorTree)
+		{
+			TreeName = NewBehaviorTree->GetName();
+		}
+		else if (OldBehaviorTree)
+		{
+			TreeName = OldBehaviorTree->GetName();
+		}
 		WindowTitle = FText::Format(LOCTEXT("Behavior Tree Diff", "{0} - Behavior Tree Diff"), FText::FromString(TreeName));
 	}
 
diff --git a/Engine/Source/Editor/Kismet/Private/SBlueprintDiff.cpp b/Engine/Source/Editor/Kismet/Private/SBlueprintDiff.cpp
--- a/Engine/Source/Editor/Kismet/Private/SBlueprintDiff.cpp
+++ b/Engine/Source/Editor/Kismet/Private/SBlueprintDiff.cpp
@@ -479,7 +479,15 @@ TSharedPtr<SWindow> SBlueprintDiff::CreateDiffWindow(const UBlueprint* OldBluepr
 	if (bIsSingleAsset)
 	{
 		// identify the assumed single asset in the window's title
-		const FString BPName = NewBlueprint? NewBlueprint->GetName() : OldBlueprint->GetName();
+		FString BPName;
+		if (NewBlueprint)
+		{
+			BPName = NewBlueprint->GetName();
+		}
+		else if (OldBlueprint)
+		{
+			BPName = OldBlueprint->GetName();
+		}
 		WindowTitle = FText::Format(LOCTEXT("NamedBlueprintDiff", "{0} - {1} Diff"), FText::FromString(BPName), BlueprintClass->GetDisplayNameText());
 	}
 
@@ -879,8 +887,23 @@ void SBlueprintDiff::GenerateDifferencesList()
 	}
 
 	bool bHasComponents = false;
-	UClass* BlueprintClassOld = PanelOld.Blueprint ? PanelOld.Blueprint->GeneratedClass : PanelNew.Blueprint->GeneratedClass;
-	UClass* BlueprintClassNew = PanelNew.Blueprint ? PanelNew.Blueprint->GeneratedClass : PanelOld.Blueprint->GeneratedClass;
+	
+	const auto BPClassFallback = [this]()->TSubclassOf<UObject>
+	{
+		if (PanelOld.Blueprint)
+		{
+			return PanelOld.Blueprint->GeneratedClass;
+		}
+		if (PanelNew.Blueprint)
+		{
+			return PanelNew.Blueprint->GeneratedClass;
+		}
+		check(false);
+		return TSubclassOf<UObject>(nullptr); // this should never happen
+	};
+	
+	const UClass* BlueprintClassOld = PanelOld.Blueprint ? PanelOld.Blueprint->GeneratedClass : BPClassFallback();
+	const UClass* BlueprintClassNew = PanelNew.Blueprint ? PanelNew.Blueprint->GeneratedClass : BPClassFallback();
 	const bool bIsOldClassActor = BlueprintClassOld && BlueprintClassOld->IsChildOf<AActor>();
 	const bool bIsNewClassActor = BlueprintClassNew && BlueprintClassNew->IsChildOf<AActor>();
 	if (bIsOldClassActor || bIsNewClassActor)
@@ -889,8 +912,18 @@ void SBlueprintDiff::GenerateDifferencesList()
 	}
 
 	// If this isn't a normal blueprint type, add the type panel
-	const UClass* RepresentativeClass = PanelOld.Blueprint ? PanelOld.Blueprint->GetClass() : PanelNew.Blueprint->GetClass();
-	if (RepresentativeClass != UBlueprint::StaticClass())
+
+	bool bIsSpecialized = false;
+	if (PanelOld.Blueprint)
+	{
+		bIsSpecialized = PanelOld.Blueprint->GetClass() == UBlueprint::StaticClass();
+	}
+	if (PanelNew.Blueprint)
+	{
+		bIsSpecialized |= PanelNew.Blueprint->GetClass() == UBlueprint::StaticClass();
+	}
+	
+	if (bIsSpecialized)
 	{
 		ModePanels.Add(BlueprintTypeMode, GenerateBlueprintTypePanel());
 	}
