Fixing issue where Find References on variables in Blueprint Editor fails to find any results

#rnx

diff --git a/Engine/Source/Editor/Kismet/Private/FindInBlueprintManager.cpp b/Engine/Source/Editor/Kismet/Private/FindInBlueprintManager.cpp
--- a/Engine/Source/Editor/Kismet/Private/FindInBlueprintManager.cpp
+++ b/Engine/Source/Editor/Kismet/Private/FindInBlueprintManager.cpp
@@ -558,11 +558,16 @@ namespace BlueprintSearchMetaDataHelpers
 			WriteTextValue(FText::FromString(String));
 		}
 
+		virtual void WriteStringValue(FStringView String) override
+		{
+			// We just want to make sure all strings are converted into FText hex strings, used by the FiB system
+			WriteTextValue(FText::FromStringView(String));
+		}
+		
 		virtual void WriteTextValue(const FText& Text) override
 		{
 			// Check to see if the value has already been added.
-			int32* TableLookupValuePtr = ReverseLookupTable.Find(FLookupTableItem(Text));
-			if (TableLookupValuePtr)
+			if (int32* TableLookupValuePtr = ReverseLookupTable.Find(FLookupTableItem(Text)))
 			{
 				TFindInBlueprintJsonStringWriter<TCondensedJsonPrintPolicy<TCHAR>>::WriteStringValue(FString::FromInt(*TableLookupValuePtr));
 			}
diff --git a/Engine/Source/Runtime/Json/Public/Serialization/JsonWriter.h b/Engine/Source/Runtime/Json/Public/Serialization/JsonWriter.h
--- a/Engine/Source/Runtime/Json/Public/Serialization/JsonWriter.h
+++ b/Engine/Source/Runtime/Json/Public/Serialization/JsonWriter.h
@@ -472,7 +472,7 @@ protected:
 
 	virtual void WriteStringValue(const FString& String)
 	{
-		WriteStringValue(FStringView(String));
+		TJsonWriter::WriteStringValue(FStringView(String));
 	}
 
 	template<typename InCharType>
