UE-19803 - Feature request: the ability to mark custom events as private

Added default FunctionFlags values that custom events would need for specifiers. Also modified the UI to show the access specifier drop down for both function entry nodes and custom events.

diff --git a/Engine/Source/Editor/BlueprintGraph/Private/K2Node_CustomEvent.cpp b/Engine/Source/Editor/BlueprintGraph/Private/K2Node_CustomEvent.cpp
--- a/Engine/Source/Editor/BlueprintGraph/Private/K2Node_CustomEvent.cpp
+++ b/Engine/Source/Editor/BlueprintGraph/Private/K2Node_CustomEvent.cpp
@@ -33,6 +33,7 @@
 #include "UObject/Object.h"
 #include "UObject/ObjectPtr.h"
 #include "UObject/Script.h"
+#include "UObject/UE5MainStreamObjectVersion.h"
 #include "UObject/UnrealNames.h"
 #include "UObject/UnrealType.h"
 
@@ -46,28 +47,29 @@ struct FLinearColor;
  * @param  CustomEventFunc	The function you want to find an associated node for.
  * @return A pointer to the found node (NULL if a corresponding node wasn't found)
  */
-static UK2Node_CustomEvent const* FindCustomEventNodeFromFunction(UFunction* CustomEventFunc)
+static const UK2Node_CustomEvent* FindCustomEventNodeFromFunction(UFunction* CustomEventFunc)
 {
-	UK2Node_CustomEvent const* FoundEventNode = NULL;
-	if (CustomEventFunc != NULL)
+	const UK2Node_CustomEvent* FoundEventNode = nullptr;
+	if (CustomEventFunc != nullptr)
 	{
-		UObject const* const FuncOwner = CustomEventFunc->GetOuter();
-		check(FuncOwner != NULL);
+		const UObject* const FuncOwner = CustomEventFunc->GetOuter();
+		check(FuncOwner != nullptr);
 
 		// if the found function is a NOT a native function (it's user generated)
 		if (FuncOwner->IsA(UBlueprintGeneratedClass::StaticClass()))
 		{
-			UBlueprintGeneratedClass* FuncClass = Cast<UBlueprintGeneratedClass>(CustomEventFunc->GetOuter());
-			check(FuncClass != NULL);
-			UBlueprint* FuncBlueprint = Cast<UBlueprint>(FuncClass->ClassGeneratedBy);
-			check(FuncBlueprint != NULL);
+			const UBlueprintGeneratedClass* FuncClass = Cast<UBlueprintGeneratedClass>(CustomEventFunc->GetOuter());
+			check(FuncClass != nullptr);
+			const UBlueprint* FuncBlueprint = Cast<UBlueprint>(FuncClass->ClassGeneratedBy);
+			check(FuncBlueprint != nullptr);
 
 			TArray<UK2Node_CustomEvent*> BpCustomEvents;
 			FBlueprintEditorUtils::GetAllNodesOfClass<UK2Node_CustomEvent>(FuncBlueprint, BpCustomEvents);
 
 			// look to see if the function that this is overriding is a custom-event
-			for (UK2Node_CustomEvent const* const UserEvent : BpCustomEvents)
+			for (const UK2Node_CustomEvent* const UserEvent : BpCustomEvents)
 			{
+				check(UserEvent);
 				if (UserEvent->CustomFunctionName == CustomEventFunc->GetFName())
 				{
 					FoundEventNode = UserEvent;
@@ -167,15 +169,24 @@ UK2Node_CustomEvent::UK2Node_CustomEvent(const FObjectInitializer& ObjectInitial
 	bCanRenameNode = true;
 	bIsDeprecated = false;
 	bCallInEditor = false;
+
+	FunctionFlags = (FUNC_BlueprintCallable | FUNC_BlueprintEvent | FUNC_Public);
 }
 
 void UK2Node_CustomEvent::Serialize(FArchive& Ar)
 {
+	Ar.UsingCustomVersion(FUE5MainStreamObjectVersion::GUID);
+
 	Super::Serialize(Ar);
 
 	if (Ar.IsLoading())
 	{
 		CachedNodeTitle.MarkDirty();
+
+		if (Ar.CustomVer(FUE5MainStreamObjectVersion::GUID) < FUE5MainStreamObjectVersion::AccessSpecifiersForCustomEvents)
+		{
+			FunctionFlags |= (FUNC_BlueprintCallable | FUNC_BlueprintEvent | FUNC_Public);
+		}
 	}
 }
 
diff --git a/Engine/Source/Editor/Kismet/Private/BlueprintDetailsCustomization.cpp b/Engine/Source/Editor/Kismet/Private/BlueprintDetailsCustomization.cpp
--- a/Engine/Source/Editor/Kismet/Private/BlueprintDetailsCustomization.cpp
+++ b/Engine/Source/Editor/Kismet/Private/BlueprintDetailsCustomization.cpp
@@ -3979,34 +3979,6 @@ void FBlueprintGraphActionDetails::CustomizeDetails( IDetailLayoutBuilder& Detai
 					];
 			}
 
-			if (IsAccessSpecifierVisible())
-			{
-				Category.AddCustomRow( LOCTEXT( "AccessSpecifier", "Access Specifier" ) )
-				.NameContent()
-				[
-					SNew(STextBlock)
-						.Text( LOCTEXT( "AccessSpecifier", "Access Specifier" ) )
-						.Font( IDetailLayoutBuilder::GetDetailFont() )
-				]
-				.ValueContent()
-				[
-					SAssignNew(AccessSpecifierComboButton, SComboButton)
-					.ContentPadding(0.0f)
-					.ButtonContent()
-					[
-						SNew(STextBlock)
-							.Text(this, &FBlueprintGraphActionDetails::GetCurrentAccessSpecifierName)
-							.Font( IDetailLayoutBuilder::GetDetailFont() )
-					]
-					.MenuContent()
-					[
-						SNew(SListView<TSharedPtr<FAccessSpecifierLabel> >)
-							.ListItemsSource( &AccessSpecifierLabels )
-							.OnGenerateRow(this, &FBlueprintGraphActionDetails::HandleGenerateRowAccessSpecifier)
-							.OnSelectionChanged(this, &FBlueprintGraphActionDetails::OnAccessSpecifierSelected)
-					]
-				];
-			}
 			if (GetInstanceColorVisibility())
 			{
 				Category.AddCustomRow( LOCTEXT( "InstanceColor", "Instance Color" ) )
@@ -4291,7 +4263,7 @@ void FBlueprintGraphActionDetails::CustomizeDetails( IDetailLayoutBuilder& Detai
 			];
 		}
 
-		const bool bShowDeprecated = bIsFunctionGraph || bIsCustomEvent;
+		const bool bShowDeprecated = bIsCustomEvent || bIsFunctionGraph;
 		if (bShowDeprecated)
 		{
 			FFormatNamedArguments DeprecationTooltipFormatArgs;
@@ -4352,6 +4324,36 @@ void FBlueprintGraphActionDetails::CustomizeDetails( IDetailLayoutBuilder& Detai
 			];
 		}
 
+		const bool bShowAccessSpecifiers = bIsCustomEvent || bIsFunctionGraph;
+		if (IsAccessSpecifierVisible())
+		{
+			Category.AddCustomRow(LOCTEXT("AccessSpecifier", "Access Specifier"))
+				.NameContent()
+				[
+					SNew(STextBlock)
+					.Text(LOCTEXT("AccessSpecifier", "Access Specifier"))
+					.Font(IDetailLayoutBuilder::GetDetailFont())
+				]
+				.ValueContent()
+				[
+					SAssignNew(AccessSpecifierComboButton, SComboButton)
+					.ContentPadding(0.0f)
+					.ButtonContent()
+					[
+						SNew(STextBlock)
+						.Text(this, &FBlueprintGraphActionDetails::GetCurrentAccessSpecifierName)
+						.Font(IDetailLayoutBuilder::GetDetailFont())
+					]
+					.MenuContent()
+					[
+						SNew(SListView<TSharedPtr<FAccessSpecifierLabel> >)
+						.ListItemsSource(&AccessSpecifierLabels)
+						.OnGenerateRow(this, &FBlueprintGraphActionDetails::HandleGenerateRowAccessSpecifier)
+						.OnSelectionChanged(this, &FBlueprintGraphActionDetails::OnAccessSpecifierSelected)
+					]
+				];
+		}
+
 		IDetailCategoryBuilder& InputsCategory = DetailLayout.EditCategory("Inputs", LOCTEXT("FunctionDetailsInputs", "Inputs"));
 		
 		TSharedRef<FBlueprintGraphArgumentGroupLayout> InputArgumentGroup =
diff --git a/Engine/Source/Runtime/Core/Public/UObject/UE5MainStreamObjectVersion.h b/Engine/Source/Runtime/Core/Public/UObject/UE5MainStreamObjectVersion.h
--- a/Engine/Source/Runtime/Core/Public/UObject/UE5MainStreamObjectVersion.h
+++ b/Engine/Source/Runtime/Core/Public/UObject/UE5MainStreamObjectVersion.h
@@ -331,6 +331,9 @@ struct CORE_API FUE5MainStreamObjectVersion
 		// In certain cases, Blueprint pins with a PC_Object category would serialize a null PinSubCategoryObject
 		NullPinSubCategoryObjectFix,
 
+		// Allow custom event nodes to use access specifiers
+		AccessSpecifiersForCustomEvents,
+
 		// -----<new versions can be added above this line>-------------------------------------------------
 		VersionPlusOne,
 		LatestVersion = VersionPlusOne - 1
