Fix spurious LoadPackage can't find package warnings when opening new ControlRigs.
#rnx

diff --git a/Engine/Source/Editor/Kismet/Private/BlueprintEditor.cpp b/Engine/Source/Editor/Kismet/Private/BlueprintEditor.cpp
--- a/Engine/Source/Editor/Kismet/Private/BlueprintEditor.cpp
+++ b/Engine/Source/Editor/Kismet/Private/BlueprintEditor.cpp
@@ -1921,18 +1921,24 @@ struct FLoadObjectsFromAssetRegistryHelper
 
 		const double CompileStartTime = FPlatformTime::Seconds();
 
-		TArray<FAssetData> AssetData;
-		AssetRegistryModule.Get().GetAssetsByClass(TObjectType::StaticClass()->GetClassPathName(), AssetData);
+		TArray<FAssetData> AssetDatas;
+		AssetRegistryModule.Get().GetAssetsByClass(TObjectType::StaticClass()->GetClassPathName(), AssetDatas);
 
-		for (int32 AssetIndex = 0; AssetIndex < AssetData.Num(); ++AssetIndex)
+		for (FAssetData& AssetData : AssetDatas)
 		{
-			if(AssetData[AssetIndex].IsValid())
+			if (AssetData.IsValid())
 			{
-				FString AssetPath = AssetData[AssetIndex].GetObjectPathString();
-				TObjectType* Object = LoadObject<TObjectType>(nullptr, *AssetPath, nullptr, 0, nullptr);
-				if (Object)
+				FString AssetPath = AssetData.GetObjectPathString();
+				// Workaround for UE-178174: AssetRegistry returning unmounted AssetPaths. Test the path for mountedness before loading.
+				FString PackagePathString = FPackageName::ObjectPathToPackageName(AssetPath);
+				FPackagePath UnusedPackagePath;
+				if (FPackagePath::TryFromMountedName(PackagePathString, UnusedPackagePath))
 				{
-					Collection.Add( MakeWeakObjectPtr(Object) );
+					TObjectType* Object = LoadObject<TObjectType>(nullptr, *AssetPath, nullptr, 0, nullptr);
+					if (Object)
+					{
+						Collection.Add(MakeWeakObjectPtr(Object));
+					}
 				}
 			}
 		}
