Moving SBlueprintEditorSelectedDebugObjectWidget.h to Public for external modules to use correctly.

diff --git a/Engine/Source/Editor/Kismet/Private/SBlueprintEditorSelectedDebugObjectWidget.h b/Engine/Source/Editor/Kismet/Private/SBlueprintEditorSelectedDebugObjectWidget.h
deleted file mode 100644
diff --git a/Engine/Source/Editor/Kismet/Public/SBlueprintEditorSelectedDebugObjectWidget.h b/Engine/Source/Editor/Kismet/Public/SBlueprintEditorSelectedDebugObjectWidget.h
new file mode 100644
--- /dev/null
+++ b/Engine/Source/Editor/Kismet/Public/SBlueprintEditorSelectedDebugObjectWidget.h
@@ -0,0 +1,155 @@
+// Copyright Epic Games, Inc. All Rights Reserved.
+
+#pragma once
+
+#include "BlueprintEditor.h"
+#include "Containers/Array.h"
+#include "Containers/BitArray.h"
+#include "Containers/Set.h"
+#include "Containers/SparseArray.h"
+#include "Containers/UnrealString.h"
+#include "Delegates/Delegate.h"
+#include "HAL/PlatformCrt.h"
+#include "Internationalization/Text.h"
+#include "Layout/Visibility.h"
+#include "Misc/Optional.h"
+#include "Serialization/Archive.h"
+#include "Templates/SharedPointer.h"
+#include "Templates/TypeHash.h"
+#include "Templates/UnrealTemplate.h"
+#include "Types/SlateEnums.h"
+#include "UObject/WeakObjectPtr.h"
+#include "UObject/WeakObjectPtrTemplates.h"
+#include "Widgets/DeclarativeSyntaxSupport.h"
+#include "Widgets/Input/SComboBox.h"
+#include "Widgets/SCompoundWidget.h"
+
+class STextComboBox;
+class SWidget;
+class UBlueprint;
+class UObject;
+class UWorld;
+struct FGeometry;
+
+struct FBlueprintDebugObjectInstance
+{
+	/** Actual object to debug, can be null */
+	TWeakObjectPtr<UObject> ObjectPtr;
+
+	/** Friendly label for object to debug */
+	FString ObjectLabel;
+
+	/** Raw object path of spawned PIE object, this is not a SoftObjectPath because we don't want it to get fixed up */
+	FString ObjectPath;
+
+	/** Object path to object in the editor, will only be set for static objects */
+	FString EditorObjectPath;
+
+	FBlueprintDebugObjectInstance(TWeakObjectPtr<UObject> InPtr, const FString& InLabel)
+		: ObjectPtr(InPtr)
+		, ObjectLabel(InLabel)
+	{
+	}
+
+	/** Returns true if this is the special entry for no specific object */
+	bool IsEmptyObject() const
+	{
+		return ObjectPath.IsEmpty();
+	}
+
+	/** If this has no editor path, it was spawned */
+	bool IsSpawnedObject() const 
+	{
+		return !ObjectPath.IsEmpty() && EditorObjectPath.IsEmpty();
+	}
+	
+	/** If editor and object path are the same length because there's no prefix, this is the editor object */
+	bool IsEditorObject() const 
+	{
+		return !ObjectPath.IsEmpty() && ObjectPath.Len() == EditorObjectPath.Len();
+	}
+};
+
+class KISMET_API SBlueprintEditorSelectedDebugObjectWidget : public SCompoundWidget
+{
+public:
+	SLATE_BEGIN_ARGS(SBlueprintEditorSelectedDebugObjectWidget){}
+	SLATE_END_ARGS()
+
+	void Construct(const FArguments& InArgs, TSharedPtr<FBlueprintEditor> InBlueprintEditor);
+
+	// SWidget interface
+	virtual void Tick( const FGeometry& AllottedGeometry, const double InCurrentTime, const float InDeltaTime ) override;
+	// End of SWidget interface
+
+	/** Adds an object to the list of debug choices */
+	void AddDebugObject(UObject* TestObject, const FString& TestObjectName = FString());
+
+private:
+	UBlueprint* GetBlueprintObj() const { return BlueprintEditor.Pin()->GetBlueprintObj(); }
+
+	/** Creates a list of all debug objects **/
+	void GenerateDebugObjectInstances(bool bRestoreSelection);
+
+	/** Generate list of active PIE worlds to debug **/
+	void GenerateDebugWorldNames(bool bRestoreSelection);
+
+	/** Refresh the widget. **/
+	void OnRefresh();
+
+	/** Returns the entry for the current debug actor */
+	TSharedPtr<FBlueprintDebugObjectInstance> GetDebugObjectInstance() const;
+
+	/** Returns the name of the current debug actor */
+	TSharedPtr<FString> GetDebugWorldName() const;
+
+	/** Handles the selection changed event for the debug actor combo box */
+	void DebugObjectSelectionChanged(TSharedPtr<FBlueprintDebugObjectInstance> NewSelection, ESelectInfo::Type SelectInfo);
+
+	void DebugWorldSelectionChanged(TSharedPtr<FString> NewSelection, ESelectInfo::Type SelectInfo);
+
+	/** Called when user clicks button to select the current object being debugged */
+	void SelectedDebugObject_OnClicked();
+
+	/** Returns true if a debug actor is currently selected */
+	bool IsDebugObjectSelected() const;
+
+	EVisibility IsDebugWorldComboVisible() const;
+
+	/* Returns the string to indicate no debug object is selected */
+	const FString& GetNoDebugString() const;
+
+	const FString& GetDebugAllWorldsString() const;
+
+	TSharedRef<SWidget> OnGetActiveDetailSlotContent(bool bChangedToHighDetail);
+
+	/** Helper method to construct a debug object label string */
+	FString MakeDebugObjectLabel(UObject* TestObject, bool bAddContextIfSelectedInEditor, bool bAddSpawnedContext) const;
+
+	/** Fills in data for a specific instance */
+	void FillDebugObjectInstance(TSharedPtr<FBlueprintDebugObjectInstance> Instance);
+
+	/** Called to create a widget for each debug object item */
+	TSharedRef<SWidget> CreateDebugObjectItemWidget(TSharedPtr<FBlueprintDebugObjectInstance> InItem);
+
+	/** Returns the combo button label to use for the currently-selected debug object item */
+	FText GetSelectedDebugObjectTextLabel() const;
+
+private:
+	/** Pointer back to the blueprint editor tool that owns us */
+	TWeakPtr<FBlueprintEditor> BlueprintEditor;
+
+	/** Lists of actors of a given blueprint type and their names */
+	TArray<TSharedPtr<FBlueprintDebugObjectInstance>> DebugObjects;
+
+	/** PIE worlds that we can debug */
+	TArray< TWeakObjectPtr<UWorld> > DebugWorlds;
+	TArray< TSharedPtr<FString> > DebugWorldNames;
+
+	/** Widget containing the names of all possible debug actors. This is a "generic" SComboBox rather than an STextComboBox so that we can customize the label on the combo button widget. */
+	TSharedPtr<SComboBox<TSharedPtr<FBlueprintDebugObjectInstance>>> DebugObjectsComboBox;
+
+	TSharedPtr<STextComboBox> DebugWorldsComboBox;
+
+	TWeakObjectPtr<UObject> LastObjectObserved;
+};
