PR #10192: Prevent SBlueprintActionMenu checkbox text being cutoff (Contributed by [REDACTED_PERSON])

diff --git a/Engine/Source/Editor/Kismet/Private/SBlueprintActionMenu.cpp b/Engine/Source/Editor/Kismet/Private/SBlueprintActionMenu.cpp
--- a/Engine/Source/Editor/Kismet/Private/SBlueprintActionMenu.cpp
+++ b/Engine/Source/Editor/Kismet/Private/SBlueprintActionMenu.cpp
@@ -149,7 +149,7 @@ public:
 				.FillWidth(1.0)
 			[
 				SNew( SCheckBox )
-					.Visibility(this, &SBlueprintActionFavoriteToggle::IsVisibile)
+					.Visibility(this, &SBlueprintActionFavoriteToggle::IsVisible)
 					.ToolTipText(this, &SBlueprintActionFavoriteToggle::GetToolTipText)
 					.IsChecked(this, &SBlueprintActionFavoriteToggle::GetFavoritedState)
 					.OnCheckStateChanged(this, &SBlueprintActionFavoriteToggle::OnFavoriteToggled)
@@ -165,7 +165,7 @@ private:
 	 *
 	 * @return True if this toggle switch should be showing, false if not.
 	 */
-	EVisibility IsVisibile() const
+	EVisibility IsVisible() const
 	{
 		bool bNoFavorites = false;
 		GConfig->GetBool(TEXT("BlueprintEditor.Palette"), TEXT("bUseLegacyLayout"), bNoFavorites, GEditorIni);
@@ -368,7 +368,7 @@ void SBlueprintActionMenu::Construct( const FArguments& InArgs, TSharedPtr<FBlue
 
 					// Search context description
 					+SHorizontalBox::Slot()
-					.AutoWidth()
+					.FillWidth(1.f)
 					.VAlign(VAlign_Center)
 					[
 						SNew(STextBlock)
@@ -379,13 +379,14 @@ void SBlueprintActionMenu::Construct( const FArguments& InArgs, TSharedPtr<FBlue
 							NULL,
 							TEXT("Shared/Editors/BlueprintEditor"),
 							TEXT("BlueprintActionMenuContextText")))
-						.WrapTextAt(280)
+						.AutoWrapText(true)
 					]
 
 					// Context Toggle
 					+SHorizontalBox::Slot()
 					.HAlign(HAlign_Right)
 					.VAlign(VAlign_Center)
+					.AutoWidth()
 					[
 						SNew(SCheckBox)
 						.OnCheckStateChanged(this, &SBlueprintActionMenu::OnContextToggleChanged)
