Add Details View row name to BP Namespaces

diff --git a/Engine/Source/Editor/Kismet/Public/BlueprintManagedListDetails.h b/Engine/Source/Editor/Kismet/Public/BlueprintManagedListDetails.h
--- a/Engine/Source/Editor/Kismet/Public/BlueprintManagedListDetails.h
+++ b/Engine/Source/Editor/Kismet/Public/BlueprintManagedListDetails.h
@@ -56,7 +56,7 @@ protected:
 	/** IDetailCustomNodeBuilder interface */
 	virtual void GenerateHeaderRowContent(FDetailWidgetRow& HeaderRow) override;
 	virtual void GenerateChildContent(IDetailChildrenBuilder& ChildrenBuilder) override;
-	virtual FName GetName() const override { return NAME_None; }
+	virtual FName GetName() const override { return *DisplayOptions.TitleText.ToString(); }
 	virtual bool InitiallyCollapsed() const override { return false; }
 	virtual void Tick(float DeltaTime) override {}
 	virtual bool RequiresTick() const override { return false; }
