Fix up redirectors to CDO and other archetypes when recreating CDOs and other archetypes
#rnx

diff --git a/Engine/Source/Editor/Kismet/Private/BlueprintCompilationManager.cpp b/Engine/Source/Editor/Kismet/Private/BlueprintCompilationManager.cpp
--- a/Engine/Source/Editor/Kismet/Private/BlueprintCompilationManager.cpp
+++ b/Engine/Source/Editor/Kismet/Private/BlueprintCompilationManager.cpp
@@ -2629,6 +2629,17 @@ void FBlueprintCompilationManagerImpl::ReinstanceBatch(TArray<FReinstancingJob>&
 		}
 	}
 
+	// Data loading may be in flight, lets immediately patch existing redirectors - 
+	// we may want to search the entire graph some day, but that would be expensive:
+	for (TObjectIterator<UObjectRedirector> Itr; Itr; ++Itr)
+	{
+		if (Itr->DestinationObject && 
+			Itr->DestinationObject->HasAnyFlags(RF_ClassDefaultObject|RF_ArchetypeObject))
+		{
+			ArchetypeReferencers.Add(*Itr);
+		}
+	}
+
 	for(UObject* ArchetypeReferencer : ArchetypeReferencers)
 	{
 		FArchiveReplaceObjectRef<UObject> ReplaceInCDOAr(ArchetypeReferencer, OldArchetypeToNewArchetype);
