Fix typo in tooltip
#rnx

diff --git a/Engine/Source/Editor/Kismet/Private/BlueprintDetailsCustomization.cpp b/Engine/Source/Editor/Kismet/Private/BlueprintDetailsCustomization.cpp
--- a/Engine/Source/Editor/Kismet/Private/BlueprintDetailsCustomization.cpp
+++ b/Engine/Source/Editor/Kismet/Private/BlueprintDetailsCustomization.cpp
@@ -4074,7 +4074,7 @@ void FBlueprintGraphActionDetails::CustomizeDetails( IDetailLayoutBuilder& Detai
 				[
 					SNew(STextBlock)
 					.Text(LOCTEXT("FunctionUnsafeDuringActorConstruction_Tooltip", "Unsafe During Actor Construction"))
-					.ToolTipText(LOCTEXT("FunctionIsUnsafeDuringActorConstruction_Tooltip", "Mark this function as unsafe during actor construction so that a warning is generated when it is called by a Constructin Script - useful when calling native functions that are also unsafe during construction"))
+					.ToolTipText(LOCTEXT("FunctionIsUnsafeDuringActorConstruction_Tooltip", "Mark this function as unsafe during actor construction so that a warning is generated when it is called by a Construction Script - useful when calling native functions that are also unsafe during construction"))
 					.Font(IDetailLayoutBuilder::GetDetailFont())
 				]
 				.ValueContent()
