#UE Disable expensive blueprint namespace registry asset registry OnAssetAdded function

diff --git a/Engine/Source/Editor/Kismet/Private/BlueprintNamespaceRegistry.cpp b/Engine/Source/Editor/Kismet/Private/BlueprintNamespaceRegistry.cpp
--- a/Engine/Source/Editor/Kismet/Private/BlueprintNamespaceRegistry.cpp
+++ b/Engine/Source/Editor/Kismet/Private/BlueprintNamespaceRegistry.cpp
@@ -76,7 +76,7 @@ void FBlueprintNamespaceRegistry::Initialize()
 	if(GIsEditor && !IsRunningCommandlet())
 	{
 		IAssetRegistry& AssetRegistry = FModuleManager::LoadModuleChecked<FAssetRegistryModule>(TEXT("AssetRegistry")).Get();
-		OnAssetAddedDelegateHandle = AssetRegistry.OnAssetAdded().AddRaw(this, &FBlueprintNamespaceRegistry::OnAssetAdded);
+		//OnAssetAddedDelegateHandle = AssetRegistry.OnAssetAdded().AddRaw(this, &FBlueprintNamespaceRegistry::OnAssetAdded);
 		OnAssetRemovedDelegateHandle = AssetRegistry.OnAssetRemoved().AddRaw(this, &FBlueprintNamespaceRegistry::OnAssetRemoved);
 		OnAssetRenamedDelegateHandle = AssetRegistry.OnAssetRenamed().AddRaw(this, &FBlueprintNamespaceRegistry::OnAssetRenamed);
 
