#UE Disable expensive find in blueprint manager asset registry OnAssetAdded function

diff --git a/Engine/Source/Editor/Kismet/Private/FindInBlueprintManager.cpp b/Engine/Source/Editor/Kismet/Private/FindInBlueprintManager.cpp
--- a/Engine/Source/Editor/Kismet/Private/FindInBlueprintManager.cpp
+++ b/Engine/Source/Editor/Kismet/Private/FindInBlueprintManager.cpp
@@ -2042,7 +2042,7 @@ void FFindInBlueprintSearchManager::Initialize()
 		IAssetRegistry* AssetRegistry = AssetRegistryModule->TryGet();
 		if (AssetRegistry)
 		{
-			AssetRegistry->OnAssetAdded().AddRaw(this, &FFindInBlueprintSearchManager::OnAssetAdded);
+			//AssetRegistry->OnAssetAdded().AddRaw(this, &FFindInBlueprintSearchManager::OnAssetAdded);
 			AssetRegistry->OnAssetRemoved().AddRaw(this, &FFindInBlueprintSearchManager::OnAssetRemoved);
 			AssetRegistry->OnAssetRenamed().AddRaw(this, &FFindInBlueprintSearchManager::OnAssetRenamed);
 			AssetRegistry->OnFilesLoaded().AddRaw(this, &FFindInBlueprintSearchManager::OnAssetRegistryFilesLoaded);
