Move Reinstancing Map before reinstancing, in case we load more data while reinstancing

diff --git a/Engine/Source/Editor/Kismet/Private/BlueprintCompilationManager.cpp b/Engine/Source/Editor/Kismet/Private/BlueprintCompilationManager.cpp
--- a/Engine/Source/Editor/Kismet/Private/BlueprintCompilationManager.cpp
+++ b/Engine/Source/Editor/Kismet/Private/BlueprintCompilationManager.cpp
@@ -1893,10 +1893,13 @@ void FBlueprintCompilationManagerImpl::FlushReinstancingQueueImpl(bool bFindAndR
 		
 		TGuardValue<bool> ReinstancingGuard(GIsReinstancing, true);
 		
+		TMap<UClass*, UClass*> ClassesToReinstanceOwned = MoveTemp(ClassesToReinstance);
+		ClassesToReinstance = TMap<UClass*, UClass*>();
+
 		FReplaceInstancesOfClassParameters Options;
 		Options.bArchetypesAreUpToDate = true;
 		Options.bReplaceReferencesToOldCDOs = bFindAndReplaceCDOReferences;
-		FBlueprintCompileReinstancer::BatchReplaceInstancesOfClass(ClassesToReinstance, Options);
+		FBlueprintCompileReinstancer::BatchReplaceInstancesOfClass(ClassesToReinstanceOwned, Options);
 
 		// Special case when we run on ALT, we want to cleanup all classes flagged for reinstanciation right away.
 		const bool bIsInActualAsyncLoadingThread = IsInAsyncLoadingThread() && !IsInGameThread();
@@ -1906,7 +1909,7 @@ void FBlueprintCompilationManagerImpl::FlushReinstancingQueueImpl(bool bFindAndR
 			// async loaded. Those instances will need to be reinstanced once they finish
 			// loading, there's no race here because if any instances are created after
 			// we check ClassHasInstancesAsyncLoading they will be created with the new class:
-			for( TMap<UClass*, UClass*>::TIterator It(ClassesToReinstance); It; ++It )
+			for( TMap<UClass*, UClass*>::TIterator It(ClassesToReinstanceOwned); It; ++It )
 			{
 				if (!ClassHasInstancesAsyncLoading(It->Key))
 				{
@@ -1920,12 +1923,11 @@ void FBlueprintCompilationManagerImpl::FlushReinstancingQueueImpl(bool bFindAndR
 		else
 		{
 			// Make sure to cleanup all properties that couldn't be destroyed in PurgeClass
-			for (TMap<UClass*, UClass*>::TIterator It(ClassesToReinstance); It; ++It)
+			for (TMap<UClass*, UClass*>::TIterator It(ClassesToReinstanceOwned); It; ++It)
 			{
 				It->Key->DestroyPropertiesPendingDestruction();
 				It->Value->DestroyPropertiesPendingDestruction();
 			}
-			ClassesToReinstance.Empty();
 		}
 	}
 	
