Avoid leaking any properties we can't destroy due to in flight async loads
#rnx

diff --git a/Engine/Source/Editor/Kismet/Private/BlueprintCompilationManager.cpp b/Engine/Source/Editor/Kismet/Private/BlueprintCompilationManager.cpp
--- a/Engine/Source/Editor/Kismet/Private/BlueprintCompilationManager.cpp
+++ b/Engine/Source/Editor/Kismet/Private/BlueprintCompilationManager.cpp
@@ -1919,6 +1919,8 @@ void FBlueprintCompilationManagerImpl::FlushReinstancingQueueImpl(bool bFindAndR
 					It.RemoveCurrent();
 				}
 			}
+			// preserve any pairs that are currently loading:
+			ClassesToReinstance = MoveTemp(ClassesToReinstanceOwned);
 		}
 		else
 		{
