Reintroduce null check that was removed because IsChildOf was relying on null 'this'

diff --git a/Engine/Source/Editor/Kismet/Private/SKismetInspector.cpp b/Engine/Source/Editor/Kismet/Private/SKismetInspector.cpp
--- a/Engine/Source/Editor/Kismet/Private/SKismetInspector.cpp
+++ b/Engine/Source/Editor/Kismet/Private/SKismetInspector.cpp
@@ -894,7 +894,7 @@ bool SKismetInspector::IsPropertyVisible( const FPropertyAndParent& PropertyAndP
 
 	const FProperty* TestProperty = ArrayProperty ? ArrayProperty->Inner : &Property;
 	const FObjectPropertyBase* ObjectProperty = CastField<const FObjectPropertyBase>(TestProperty);
-	bool bIsActorProperty = (ObjectProperty != nullptr && ObjectProperty->PropertyClass->IsChildOf(AActor::StaticClass()));
+	bool bIsActorProperty = (ObjectProperty != nullptr && ObjectProperty->PropertyClass && ObjectProperty->PropertyClass->IsChildOf(AActor::StaticClass()));
 
 	if (bIsEditingTemplate && Property.HasAnyPropertyFlags(CPF_DisableEditOnTemplate) && bIsActorProperty)
 	{
@@ -903,7 +903,7 @@ bool SKismetInspector::IsPropertyVisible( const FPropertyAndParent& PropertyAndP
 		return false;
 	}
 
-	bool bIsComponent = (ObjectProperty != nullptr && ObjectProperty->PropertyClass->IsChildOf(UActorComponent::StaticClass()));
+	bool bIsComponent = (ObjectProperty != nullptr && ObjectProperty->PropertyClass && ObjectProperty->PropertyClass->IsChildOf(UActorComponent::StaticClass()));
 	if (!bShowComponents && bIsComponent)
 	{
 		// Don't show sub components properties, thats what selecting components in the component tree is for.
