Add support for deferring Find-in-Blueprint and namespace asset discovery until after the initial registry file scan is complete.

Notes:
- For Find-in-Blueprints, this defaults to off (legacy behavior). It can be toggled on via a project's DefaultEditor.ini, trading a potentially faster asset scan for a possible single-frame hitch at the end of the scan (depending on total BP asset count, it may or may not be noticeable).
- For Blueprint namespaces, we always defer (the cost is relatively small enough for it to not be noticeable to just process on completion).

diff --git a/Engine/Config/BaseEditor.ini b/Engine/Config/BaseEditor.ini
--- a/Engine/Config/BaseEditor.ini
+++ b/Engine/Config/BaseEditor.ini
@@ -519,6 +519,7 @@ bDisableThreadedIndexing=false
 bEnableCsvStatsProfiling=false
 bEnableDeveloperMenuTools=false
 bDisableSearchResultTemplates=false
+bDisableImmediateAssetDiscovery=false
 
 +ClassesWithStaticSearchableValues=/Script/GameplayAbilities.GameplayCueNotify_Actor
 +ClassesWithStaticSearchableValues=/Script/GameplayAbilities.GameplayCueNotify_Static
diff --git a/Engine/Source/Editor/Kismet/Private/BlueprintNamespaceRegistry.cpp b/Engine/Source/Editor/Kismet/Private/BlueprintNamespaceRegistry.cpp
--- a/Engine/Source/Editor/Kismet/Private/BlueprintNamespaceRegistry.cpp
+++ b/Engine/Source/Editor/Kismet/Private/BlueprintNamespaceRegistry.cpp
@@ -76,11 +76,14 @@ void FBlueprintNamespaceRegistry::Initialize()
 	if(GIsEditor && !IsRunningCommandlet())
 	{
 		IAssetRegistry& AssetRegistry = FModuleManager::LoadModuleChecked<FAssetRegistryModule>(TEXT("AssetRegistry")).Get();
-		//OnAssetAddedDelegateHandle = AssetRegistry.OnAssetAdded().AddRaw(this, &FBlueprintNamespaceRegistry::OnAssetAdded);
-		OnAssetRemovedDelegateHandle = AssetRegistry.OnAssetRemoved().AddRaw(this, &FBlueprintNamespaceRegistry::OnAssetRemoved);
-		OnAssetRenamedDelegateHandle = AssetRegistry.OnAssetRenamed().AddRaw(this, &FBlueprintNamespaceRegistry::OnAssetRenamed);
-
-		FindAndRegisterAllNamespaces();
+		if (AssetRegistry.IsLoadingAssets())
+		{
+			OnFilesLoadedDelegateHandle = AssetRegistry.OnFilesLoaded().AddRaw(this, &FBlueprintNamespaceRegistry::OnAssetRegistryFilesLoaded);
+		}
+		else
+		{
+			OnAssetRegistryFilesLoaded();
+		}
 
 		OnReloadCompleteDelegateHandle = FCoreUObjectDelegates::ReloadCompleteDelegate.AddRaw(this, &FBlueprintNamespaceRegistry::OnReloadComplete);
 		OnDefaultNamespaceTypeChangedDelegateHandle = FBlueprintNamespaceUtilities::OnDefaultBlueprintNamespaceTypeChanged().AddRaw(this, &FBlueprintNamespaceRegistry::OnDefaultNamespaceTypeChanged);
@@ -107,6 +110,7 @@ void FBlueprintNamespaceRegistry::Shutdown()
 			AssetRegistry->OnAssetAdded().Remove(OnAssetAddedDelegateHandle);
 			AssetRegistry->OnAssetRemoved().Remove(OnAssetRemovedDelegateHandle);
 			AssetRegistry->OnAssetRenamed().Remove(OnAssetRenamedDelegateHandle);
+			AssetRegistry->OnFilesLoaded().Remove(OnFilesLoadedDelegateHandle);
 		}
 	}
 
@@ -169,6 +173,18 @@ void FBlueprintNamespaceRegistry::OnAssetRenamed(const FAssetData& AssetData, co
 	}
 }
 
+void FBlueprintNamespaceRegistry::OnAssetRegistryFilesLoaded()
+{
+	// Find available namespaces once all assets are registered.
+	FindAndRegisterAllNamespaces();
+
+	// Get notified for new asset events that occur going forward.
+	IAssetRegistry& AssetRegistry = FModuleManager::LoadModuleChecked<FAssetRegistryModule>(TEXT("AssetRegistry")).Get();
+	OnAssetAddedDelegateHandle = AssetRegistry.OnAssetAdded().AddRaw(this, &FBlueprintNamespaceRegistry::OnAssetAdded);
+	OnAssetRemovedDelegateHandle = AssetRegistry.OnAssetRemoved().AddRaw(this, &FBlueprintNamespaceRegistry::OnAssetRemoved);
+	OnAssetRenamedDelegateHandle = AssetRegistry.OnAssetRenamed().AddRaw(this, &FBlueprintNamespaceRegistry::OnAssetRenamed);
+}
+
 void FBlueprintNamespaceRegistry::OnReloadComplete(EReloadCompleteReason InReason)
 {
 	// Rebuild the registry to reflect the appropriate default namespace identifiers for any reloaded types.
@@ -226,6 +242,8 @@ void FBlueprintNamespaceRegistry::GetAllRegisteredPaths(TArray<FString>& OutPath
 
 void FBlueprintNamespaceRegistry::FindAndRegisterAllNamespaces()
 {
+	TRACE_CPUPROFILER_EVENT_SCOPE(FBlueprintNamespaceRegistry::FindAndRegisterAllNamespaces);
+
 	// Resolve the subset of excluded objects that are loaded.
 	TSet<const UObject*> ExcludedObjects;
 	for (const FSoftObjectPath& ExcludedObjectPath : ExcludedObjectPaths)
diff --git a/Engine/Source/Editor/Kismet/Private/FindInBlueprintManager.cpp b/Engine/Source/Editor/Kismet/Private/FindInBlueprintManager.cpp
--- a/Engine/Source/Editor/Kismet/Private/FindInBlueprintManager.cpp
+++ b/Engine/Source/Editor/Kismet/Private/FindInBlueprintManager.cpp
@@ -1981,6 +1981,7 @@ FFindInBlueprintSearchManager::FFindInBlueprintSearchManager()
 	, bEnableCSVStatsProfiling(false)
 	, bEnableDeveloperMenuTools(false)
 	, bDisableSearchResultTemplates(false)
+	, bDisableImmediateAssetDiscovery(false)
 {
 	for (int32 TabIdx = 0; TabIdx < UE_ARRAY_COUNT(GlobalFindResultsTabIDs); TabIdx++)
 	{
@@ -2021,6 +2022,7 @@ void FFindInBlueprintSearchManager::Initialize()
 	GConfig->GetBool(TEXT("BlueprintSearchSettings"), TEXT("bEnableCsvStatsProfiling"), bEnableCSVStatsProfiling, GEditorIni);
 	GConfig->GetBool(TEXT("BlueprintSearchSettings"), TEXT("bEnableDeveloperMenuTools"), bEnableDeveloperMenuTools, GEditorIni);
 	GConfig->GetBool(TEXT("BlueprintSearchSettings"), TEXT("bDisableSearchResultTemplates"), bDisableSearchResultTemplates, GEditorIni);
+	GConfig->GetBool(TEXT("BlueprintSearchSettings"), TEXT("bDisableImmediateAssetDiscovery"), bDisableImmediateAssetDiscovery, GEditorIni);
 
 #if CSV_PROFILER
 	// If profiling has been enabled, turn on the stat category and begin a capture.
@@ -2042,9 +2044,6 @@ void FFindInBlueprintSearchManager::Initialize()
 		IAssetRegistry* AssetRegistry = AssetRegistryModule->TryGet();
 		if (AssetRegistry)
 		{
-			//AssetRegistry->OnAssetAdded().AddRaw(this, &FFindInBlueprintSearchManager::OnAssetAdded);
-			AssetRegistry->OnAssetRemoved().AddRaw(this, &FFindInBlueprintSearchManager::OnAssetRemoved);
-			AssetRegistry->OnAssetRenamed().AddRaw(this, &FFindInBlueprintSearchManager::OnAssetRenamed);
 			AssetRegistry->OnFilesLoaded().AddRaw(this, &FFindInBlueprintSearchManager::OnAssetRegistryFilesLoaded);
 		}
 	}
@@ -2068,7 +2067,7 @@ void FFindInBlueprintSearchManager::Initialize()
 	// Register to be notified of reloads
 	FCoreUObjectDelegates::ReloadCompleteDelegate.AddRaw(this, &FFindInBlueprintSearchManager::OnReloadComplete);
 
-	if(!GIsSavingPackage && AssetRegistryModule)
+	if(!GIsSavingPackage && AssetRegistryModule && (!bDisableImmediateAssetDiscovery || !AssetRegistryModule->GetRegistry().IsLoadingAssets()))
 	{
 		// Do an immediate load of the cache to catch any Blueprints that were discovered by the asset registry before we initialized.
 		BuildCache();
@@ -2331,6 +2330,13 @@ void FFindInBlueprintSearchManager::OnAssetRegistryFilesLoaded()
 {
 	CSV_EVENT(FindInBlueprint, TEXT("OnAssetRegistryFilesLoaded"));
 
+	// If we've deferred asset discovery, scan all registered assets now to extract search metadata from the asset tags.
+	// Note: Depending on how many assets there are (loaded/unloaded), this may block the UI frame for an extended period.
+	if (bDisableImmediateAssetDiscovery)
+	{
+		BuildCache();
+	}
+
 	if (!IsCacheInProgress() && PendingAssets.Num() == 0)
 	{
 		// Invoke the completion callback on any active global FiB tabs that are currently open to signal that the discovery stage is complete.
@@ -2946,7 +2952,18 @@ void FFindInBlueprintSearchManager::CleanCache()
 
 void FFindInBlueprintSearchManager::BuildCache()
 {
-	AssetRegistryModule = &FModuleManager::LoadModuleChecked<FAssetRegistryModule>(TEXT("AssetRegistry"));
+	TRACE_CPUPROFILER_EVENT_SCOPE(FFindInBlueprintSearchManager::BuildCache);
+
+	if (!ensure(AssetRegistryModule))
+	{
+		return;
+	}
+
+	IAssetRegistry* AssetRegistry = AssetRegistryModule->TryGet();
+	if (!AssetRegistry)
+	{
+		return;
+	}
 
 	TArray< FAssetData > BlueprintAssets;
 	FARFilter ClassFilter;
@@ -2957,12 +2974,17 @@ void FFindInBlueprintSearchManager::BuildCache()
 		ClassFilter.ClassPaths.Add(ClassPathName);
 	}
 
-	AssetRegistryModule->Get().GetAssets(ClassFilter, BlueprintAssets);
+	AssetRegistry->GetAssets(ClassFilter, BlueprintAssets);
 	
 	for( FAssetData& Asset : BlueprintAssets )
 	{
 		OnAssetAdded(Asset);
 	}
+
+	// Register to be notified for future asset registry events.
+	AssetRegistry->OnAssetAdded().AddRaw(this, &FFindInBlueprintSearchManager::OnAssetAdded);
+	AssetRegistry->OnAssetRemoved().AddRaw(this, &FFindInBlueprintSearchManager::OnAssetRemoved);
+	AssetRegistry->OnAssetRenamed().AddRaw(this, &FFindInBlueprintSearchManager::OnAssetRenamed);
 }
 
 void FFindInBlueprintSearchManager::DumpCache(FArchive& Ar)
diff --git a/Engine/Source/Editor/Kismet/Public/BlueprintNamespaceRegistry.h b/Engine/Source/Editor/Kismet/Public/BlueprintNamespaceRegistry.h
--- a/Engine/Source/Editor/Kismet/Public/BlueprintNamespaceRegistry.h
+++ b/Engine/Source/Editor/Kismet/Public/BlueprintNamespaceRegistry.h
@@ -84,6 +84,7 @@ protected:
 	void OnAssetAdded(const FAssetData& AssetData);
 	void OnAssetRemoved(const FAssetData& AssetData);
 	void OnAssetRenamed(const FAssetData& AssetData, const FString& InOldName);
+	void OnAssetRegistryFilesLoaded();
 
 	/** Namespace identifier registration methods. */
 	void FindAndRegisterAllNamespaces();
@@ -106,6 +107,7 @@ private:
 	FDelegateHandle OnAssetAddedDelegateHandle;
 	FDelegateHandle OnAssetRemovedDelegateHandle;
 	FDelegateHandle OnAssetRenamedDelegateHandle;
+	FDelegateHandle OnFilesLoadedDelegateHandle;
 	FDelegateHandle OnReloadCompleteDelegateHandle;
 	FDelegateHandle OnDefaultNamespaceTypeChangedDelegateHandle;
 
diff --git a/Engine/Source/Editor/Kismet/Public/FindInBlueprintManager.h b/Engine/Source/Editor/Kismet/Public/FindInBlueprintManager.h
--- a/Engine/Source/Editor/Kismet/Public/FindInBlueprintManager.h
+++ b/Engine/Source/Editor/Kismet/Public/FindInBlueprintManager.h
@@ -867,6 +867,9 @@ private:
 
 	/** Disable the use of search result templates. Setting this to TRUE will slightly decrease overall memory usage, but will also increase global search times */
 	bool bDisableSearchResultTemplates;
+
+	/** Defers the cost to extract metadata for each discovered asset during the initial asset registry scan into a single pass over the full asset registry once the scan is complete. */
+	bool bDisableImmediateAssetDiscovery;
 };
 
 struct KISMET_API FDisableGatheringDataOnScope
