[Feature] Review tool Comments utilizing the swarm API

diff --git a/Engine/Plugins/Editor/ChangelistReview/Source/Private/ChangelistReviewModule.cpp b/Engine/Plugins/Editor/ChangelistReview/Source/Private/ChangelistReviewModule.cpp
--- a/Engine/Plugins/Editor/ChangelistReview/Source/Private/ChangelistReviewModule.cpp
+++ b/Engine/Plugins/Editor/ChangelistReview/Source/Private/ChangelistReviewModule.cpp
@@ -12,6 +12,63 @@
 
 static const FName SourceControlReviewTabName = FName(TEXT("SourceControlChangelistReview"));
 
+namespace UE::DiffControl
+{
+	extern KISMET_API const TArray<FReviewComment>*(*GGetReviewCommentsForFile)(const FString&);
+	extern KISMET_API void(*GPostReviewComment)(FReviewComment&);
+	extern KISMET_API void(*GEditReviewComment)(FReviewComment&);
+	extern KISMET_API FString(*GGetReviewerUsername)(void);
+	extern KISMET_API bool (*GIsFileInReview)(const FString& File);
+}
+
+static const TArray<FReviewComment>* GetReviewCommentsForFile(const FString& FilePath)
+{
+	const TWeakPtr<SSourceControlReview> ReviewWeak = FChangelistReviewModule::Get().GetActiveReview();
+	if (const TSharedPtr<SSourceControlReview> Review = ReviewWeak.Pin())
+	{
+		return Review->GetReviewCommentsForFile(FilePath);
+	}
+	return nullptr;
+}
+
+static void PostReviewComment(FReviewComment& Comment)
+{
+	const TWeakPtr<SSourceControlReview> ReviewWeak = FChangelistReviewModule::Get().GetActiveReview();
+	if (const TSharedPtr<SSourceControlReview> Review = ReviewWeak.Pin())
+	{
+		Review->PostComment(Comment);
+	}
+}
+
+static void EditReviewComment(FReviewComment& Comment)
+{
+	const TWeakPtr<SSourceControlReview> ReviewWeak = FChangelistReviewModule::Get().GetActiveReview();
+	if (const TSharedPtr<SSourceControlReview> Review = ReviewWeak.Pin())
+	{
+		Review->EditComment(Comment);
+	}
+}
+
+static FString GetReviewerUsername()
+{
+	const TWeakPtr<SSourceControlReview> ReviewWeak = FChangelistReviewModule::Get().GetActiveReview();
+	if (const TSharedPtr<SSourceControlReview> Review = ReviewWeak.Pin())
+	{
+		return Review->GetReviewerUsername();
+	}
+	return {};
+}
+
+static bool IsFileInReview(const FString& File)
+{
+	const TWeakPtr<SSourceControlReview> ReviewWeak = FChangelistReviewModule::Get().GetActiveReview();
+	if (const TSharedPtr<SSourceControlReview> Review = ReviewWeak.Pin())
+	{
+		return Review->IsFileInReview(File);
+	}
+	return false;
+}
+
 void FChangelistReviewModule::StartupModule()
 {
 	UToolMenu* SourceControlMenu = UToolMenus::Get()->ExtendMenu("StatusBar.ToolBar.SourceControl");
@@ -37,10 +94,21 @@ void FChangelistReviewModule::StartupModule()
 		.SetTooltipText(LOCTEXT("ChangelistsTabTooltip", "Opens a dialog to review shelved and submitted changelists."))
 		.SetIcon(SourceControlIcon)
 		.SetMenuType(ETabSpawnerMenuType::Hidden);
+
+	UE::DiffControl::GGetReviewCommentsForFile = GetReviewCommentsForFile;
+	UE::DiffControl::GPostReviewComment = PostReviewComment;
+	UE::DiffControl::GEditReviewComment = EditReviewComment;
+	UE::DiffControl::GGetReviewerUsername = GetReviewerUsername;
+	UE::DiffControl::GIsFileInReview = IsFileInReview;
 }
 
 void FChangelistReviewModule::ShutdownModule()
 {
+	UE::DiffControl::GGetReviewCommentsForFile = nullptr;
+	UE::DiffControl::GPostReviewComment = nullptr;
+	UE::DiffControl::GEditReviewComment = nullptr;
+	UE::DiffControl::GGetReviewerUsername = nullptr;
+	UE::DiffControl::GIsFileInReview = nullptr;
 }
 
 TSharedRef<SDockTab> FChangelistReviewModule::CreateReviewTab(const FSpawnTabArgs& Args)
@@ -54,13 +122,13 @@ TSharedRef<SDockTab> FChangelistReviewModule::CreateReviewTab(const FSpawnTabArg
 
 TSharedPtr<SWidget> FChangelistReviewModule::CreateReviewUI()
 {
-	TSharedPtr<SWidget> ReturnWidget;
+	TSharedPtr<SSourceControlReview> ReturnWidget;
 	if (IsInGameThread())
 	{
-		TSharedPtr<SSourceControlReview> SharedPtr = SNew(SSourceControlReview);
-		ReturnWidget = SharedPtr;
+		ReturnWidget = SNew(SSourceControlReview);
+		ReviewWidget = ReturnWidget;
 	}
-	 return ReturnWidget;
+	return ReturnWidget;
 }
 
 void FChangelistReviewModule::ShowReviewTab()
@@ -81,6 +149,11 @@ bool FChangelistReviewModule::CanShowReviewTab() const
 	return false;
 }
 
+TWeakPtr<SSourceControlReview> FChangelistReviewModule::GetActiveReview()
+{
+	return ReviewWidget;
+}
+
 
 #undef LOCTEXT_NAMESPACE
 	
diff --git a/Engine/Plugins/Editor/ChangelistReview/Source/Private/FSwarmCommentsAPI.cpp b/Engine/Plugins/Editor/ChangelistReview/Source/Private/FSwarmCommentsAPI.cpp
--- a/Engine/Plugins/Editor/ChangelistReview/Source/Private/FSwarmCommentsAPI.cpp
+++ b/Engine/Plugins/Editor/ChangelistReview/Source/Private/FSwarmCommentsAPI.cpp
@@ -14,10 +14,142 @@ THIRD_PARTY_INCLUDES_START
 #include <p4/clientapi.h>
 THIRD_PARTY_INCLUDES_END
 
-FSwarmCommentsAPI::FSwarmCommentsAPI()
-	: AuthTicket(RetrieveAuthorizationTicket())
-	, SwarmURL(RetrieveSwarmURL())
+namespace UE::SwarmCommentsAPI::Keywords
 {
+	// http
+	static const TCHAR* HttpPost = TEXT("POST");
+	static const TCHAR* HttpGet = TEXT("GET");
+	static const TCHAR* HttpPatch = TEXT("PATCH");
+
+	// swarm
+	static const TCHAR* SilenceNotification = TEXT("SilenceNotification");
+	static const TCHAR* True = TEXT("true");
+	static const TCHAR* False = TEXT("false");
+
+	// metadata
+	static const TCHAR* MetadataTag = TEXT("[UE_Metadata]");
+	static const TCHAR* MetadataBody = TEXT("Body");
+	static const TCHAR* MetadataFile = TEXT("File");
+	static const TCHAR* MetadataReplyTo = TEXT("ReplyTo");
+	static const TCHAR* MetadataCategory = TEXT("Category");
+
+	// userdata
+	static const TCHAR* UserdataTag = TEXT("[UE_Userdata]");
+	static const TCHAR* UserdataLikedComments = TEXT("LikedComments");
+	static const TCHAR* UserdataReadComments = TEXT("ReadComments");
+}
+using namespace UE::SwarmCommentsAPI;
+
+static TMap<FString, FString> ParseReviewCommentMetadata(const FString& Comment, const FString& MetadataKey = Keywords::MetadataTag)
+{
+	TMap<FString, FString> Result;
+	FString Body, Metadata;
+	if (!Comment.Split(MetadataKey, &Body, &Metadata))
+	{
+		return {{Keywords::MetadataBody, Comment.TrimEnd()}};
+	}
+	
+	Result.Add(Keywords::MetadataBody, Body.TrimEnd());
+	
+	TArray<FString> MetaDataLines;
+	Metadata.ParseIntoArrayLines(MetaDataLines);
+	for (const FString& Line : MetaDataLines)
+	{
+		const int32 KeyEnd = Line.Find(TEXT("="), ESearchCase::CaseSensitive, ESearchDir::FromStart);
+		if (KeyEnd == INDEX_NONE)
+		{
+			continue;
+		}
+		FString Key = Line.Left(KeyEnd);
+		
+		const int32 ValBegin = KeyEnd + 1; // skip the '=' char
+		const int32 ValLen = Line.Len() - ValBegin; // extend to end of line
+		const FString Val = Line.Mid(ValBegin, ValLen);
+		Result.Add(Key, Val);
+	}
+	return Result;
+}
+
+static TMap<FString, FString> ParseReviewCommentMetadata(const FReviewComment& Comment, const FString& MetadataKey = Keywords::MetadataTag )
+{
+	return ParseReviewCommentMetadata(Comment.Body.Get({}), MetadataKey);
+}
+
+static void SetReviewCommentMetadata(FReviewComment& Comment, TMap<FString, FString> Metadata, const FString& MetadataKey = Keywords::MetadataTag)
+{
+	FString CommentText;
+	if (FString* Body = Metadata.Find(Keywords::MetadataBody))
+	{
+		CommentText = *Body;
+		Metadata.Remove(Keywords::MetadataBody);
+	}
+	
+	if (Metadata.IsEmpty())
+	{
+		Comment.Body = CommentText;
+		return;
+	}
+	
+	CommentText += TEXT("\n\n")+MetadataKey+TEXT("\n");
+	for (auto&[Key,Value] : Metadata)
+	{
+		if (Key != Keywords::MetadataBody)
+		{
+			CommentText += FString::Format(TEXT("{0}={1}\n"), {Key, Value});
+		}
+	}
+	Comment.Body = CommentText;
+}
+
+// turn a comma separated string into a set of strings
+static TSet<FString> DeserializeSetString(const FString& SetString)
+{
+	TSet<FString> Result;
+	TArray<FString> LikesArray;
+	SetString.ParseIntoArray(LikesArray, TEXT(","));
+	for (const FString& Like : LikesArray)
+	{
+		Result.Add(Like);
+	}
+	return Result;
+}
+
+// turn a set of strings into a comma separated string
+static FString SerializeSetString(const TSet<FString>& Set)
+{
+	FString Result;
+	int32 Index = 0;
+	for (const FString& Item : Set)
+	{
+		Result += Item;
+		if (Index != Set.Num() - 1)
+		{
+			Result += TEXT(",");
+		}
+		++Index;
+	}
+	return Result;
+}
+
+FSwarmCommentsAPI::FSwarmCommentsAPI(const FString& AuthTicket, const FString& SwarmURL)
+	: AuthTicket(AuthTicket)
+	, SwarmURL(SwarmURL)
+{
+}
+
+TSharedPtr<FSwarmCommentsAPI> FSwarmCommentsAPI::TryConnect()
+{
+	const FString& AuthTicket = RetrieveAuthorizationTicket();
+	if (AuthTicket.IsEmpty())
+	{
+		return nullptr;
+	}
+	const FString& SwarmURL = RetrieveSwarmURL();
+	if (SwarmURL.IsEmpty())
+	{
+		return nullptr;
+	}
+	return MakeShared<FSwarmCommentsAPI>(AuthTicket, SwarmURL);
 }
 
 FString FSwarmCommentsAPI::GetUsername() const
@@ -37,9 +169,89 @@ void FSwarmCommentsAPI::GetComments(const FReviewTopic& Topic, const OnGetCommen
     		Topic.ToString(), // filter to comments in a specific review
     		5000 // limit to 5k comments. This should obviously be more than enough and loads at a reasonable rate
     	}));
-    HttpRequest->SetVerb(TEXT("GET"));
+    HttpRequest->SetVerb(Keywords::HttpGet);
+
+	TWeakPtr<const FSwarmCommentsAPI> WeakSelf = AsWeak();
+
+	auto BuildCommentArray = [WeakSelf, Topic](const TArray<TSharedPtr<FJsonValue>>& CommentsJson)->TArray<FReviewComment>
+	{
+		bool bUserDataCacheFound = false;
+		TArray<FReviewComment> Comments;
+		
+		TMap<uint32, TSet<FString>> CommentLikes;
+		TMap<uint32, TSet<FString>> CommentReadBy;
+
+		const TSharedPtr<const FSwarmCommentsAPI> Self = WeakSelf.Pin();
+		if (!Self)
+		{
+			return {};
+		}
+		
+		for (const TSharedPtr<FJsonValue>& CommentJson : CommentsJson)
+		{
+			FReviewComment Comment = FReviewComment::FromJson(CommentJson->AsObject());
+			if (FSwarmUserdataComment::IsUserdataComment(Comment))
+			{
+				FSwarmUserdataComment AsUserData = Comment;
+				if (Comment.User == Self->GetUsername())
+				{
+					bUserDataCacheFound = true;
+					Self->UserdataCache.Add(Topic, AsUserData);
+				}
+				
+				for (const FString& LikedCommentID : AsUserData.GetLikedComments())
+				{
+					CommentLikes.FindOrAdd(FCString::Atoi(*LikedCommentID)).Add(*Comment.User);
+				}
+				for (const FString& ReadCommentID : AsUserData.GetReadComments())
+				{
+					CommentReadBy.FindOrAdd(FCString::Atoi(*ReadCommentID)).Add(*Comment.User);
+				}
+				
+				continue; // don't include userdata comments in the comment list
+			}
+			
+			TakeMetadataFromBody(Comment);
+			Comments.Add(Comment);
+		}
+
+		if (!bUserDataCacheFound)
+		{
+			// every user should have a comment that stores their userdata
+			Self->PostNewUserdataComment(Topic);
+		}
+		
+		// parse Userdata for Likes and ReadBy state and update the associated comments
+		for (FReviewComment& Comment : Comments)
+		{
+			if (const TSet<FString>* Likes = CommentLikes.Find(*Comment.CommentID))
+			{
+				if (!Comment.Likes.IsSet() && !Likes->IsEmpty())
+				{
+					Comment.Likes = TSet<FString>{};
+				}
+				for (const FString& Like : *Likes)
+				{
+					Comment.Likes->Add(Like);
+				}
+			}
+			if (const TSet<FString>* Reads = CommentReadBy.Find(*Comment.CommentID))
+			{
+				if (!Comment.ReadBy.IsSet() && !Reads->IsEmpty())
+				{
+					Comment.ReadBy = TSet<FString>{};
+				}
+				for (const FString& ReadBy : *Reads)
+				{
+					Comment.ReadBy->Add(ReadBy);
+				}
+			}
+		}
+		
+		return Comments;
+	};
 	
-	HttpRequest->OnProcessRequestComplete().BindLambda([OnComplete, this](FHttpRequestPtr Request, FHttpResponsePtr Response, bool bConnectedSuccessfully)
+	HttpRequest->OnProcessRequestComplete().BindLambda([BuildCommentArray, OnComplete](FHttpRequestPtr Request, FHttpResponsePtr Response, bool bConnectedSuccessfully)
 	{
 		if (!bConnectedSuccessfully)
 		{
@@ -57,15 +269,7 @@ void FSwarmCommentsAPI::GetComments(const FReviewTopic& Topic, const OnGetCommen
 			}
 			if (JsonObject->HasField("comments"))
 			{
-				TArray<FReviewComment> Comments;
-				const TArray<TSharedPtr<FJsonValue>> CommentsJson = JsonObject->GetArrayField("comments");
- 				for (const TSharedPtr<FJsonValue>& CommentJson : CommentsJson)
-				{
- 					FReviewComment Comment = FReviewComment::FromJson(CommentJson->AsObject());
- 					TakeMetadataFromBody(Comment);
-					Comments.Add(Comment);
-				}
-				OnComplete.ExecuteIfBound(MoveTemp(Comments), {});
+				OnComplete.ExecuteIfBound(BuildCommentArray(JsonObject->GetArrayField("comments")), {});
 				return;
 			}
 		}
@@ -76,20 +280,30 @@ void FSwarmCommentsAPI::GetComments(const FReviewTopic& Topic, const OnGetCommen
 	HttpRequest->ProcessRequest();
 }
 
-void FSwarmCommentsAPI::PostComment(FReviewComment& Comment, const OnPostCommentComplete& OnComplete) const
+void FSwarmCommentsAPI::PostComment(FReviewComment& Comment, const OnPostCommentComplete& OnComplete, bool bSilenceNotification) const
 {
 	const TSharedRef<IHttpRequest> HttpRequest = FHttpModule::Get().CreateRequest();
 	HttpRequest->SetHeader(TEXT("Content-Type"), TEXT("application/json; charset=utf-8"));
 	HttpRequest->SetHeader(TEXT("Authorization"), AuthTicket);
 	HttpRequest->SetURL(CommentsURL());
-	HttpRequest->SetVerb(TEXT("POST"));
+	HttpRequest->SetVerb(Keywords::HttpPost);
 
 	Comment.User = AuthTicket.Username;
 	FReviewComment CleanedComment = Comment;
 	PutMetadataInBody(CleanedComment);
-	HttpRequest->SetContentAsString(CleanedComment.ToJson());
+	const TSharedPtr<FJsonObject> JsonObject = CleanedComment.ToJsonObject();
+	if (bSilenceNotification)
+	{
+		JsonObject->SetStringField(Keywords::SilenceNotification, Keywords::True);
+	}
+
+	FString JsonString;
+	const TSharedRef<TJsonWriter<>> Writer = TJsonWriterFactory<>::Create(&JsonString);
+	FJsonSerializer::Serialize(JsonObject.ToSharedRef(), Writer);
 	
-	HttpRequest->OnProcessRequestComplete().BindLambda([OnComplete, this](FHttpRequestPtr Request, FHttpResponsePtr Response, bool bConnectedSuccessfully)
+	HttpRequest->SetContentAsString(JsonString);
+	
+	HttpRequest->OnProcessRequestComplete().BindLambda([OnComplete, WeakSelf = AsWeak()](FHttpRequestPtr Request, FHttpResponsePtr Response, bool bConnectedSuccessfully)
 	{
 		if (!bConnectedSuccessfully)
 		{
@@ -107,8 +321,12 @@ void FSwarmCommentsAPI::PostComment(FReviewComment& Comment, const OnPostComment
 			}
 			if (JsonObject->HasField("comment"))
 			{
-				FReviewComment Comment = FReviewComment::FromJson(JsonObject->GetObjectField("comment"));
-				TakeMetadataFromBody(Comment);
+				const FReviewComment Comment = FReviewComment::FromJson(JsonObject->GetObjectField("comment"));
+				if (TSharedPtr<const FSwarmCommentsAPI> Self = WeakSelf.Pin())
+				{
+					// This comment may effect the userdata. update it if necessary
+					WeakSelf.Pin()->UpdateUserdata(Comment);
+				}
 				OnComplete.ExecuteIfBound(Comment, {});
 				return;
 			}
@@ -119,19 +337,45 @@ void FSwarmCommentsAPI::PostComment(FReviewComment& Comment, const OnPostComment
 	HttpRequest->ProcessRequest();
 }
 
-void FSwarmCommentsAPI::EditComment(const FReviewComment& Comment, const OnEditCommentComplete& OnComplete) const
+void FSwarmCommentsAPI::EditComment(const FReviewComment& Comment, const OnEditCommentComplete& OnComplete, bool bSilenceNotification) const
 {
+	// edits to this comment may effect the userdata. update it if necessary
+	UpdateUserdata(Comment);
+	
+	// you can't edit other people's comments. Assume that only likes or ReadBy were edited
+	if (Comment.User != GetUsername())
+	{
+		return;
+	}
+
+	// edits without bodies should only update Userdata
+	if (!Comment.Body.IsSet())
+	{
+		return;
+	}
+	
 	const TSharedRef<IHttpRequest> HttpRequest = FHttpModule::Get().CreateRequest();
 	HttpRequest->SetHeader(TEXT("Content-Type"), TEXT("application/json; charset=utf-8"));
 	HttpRequest->SetHeader(TEXT("Authorization"), AuthTicket);
-	const FString CommentID = FString::FromInt(Comment.CommentID.GetValue());
+	const FString CommentID = FString::FromInt(*Comment.CommentID);
 	HttpRequest->SetURL(CommentsURL() / CommentID);
-	HttpRequest->SetVerb(TEXT("PATCH"));
+	HttpRequest->SetVerb(Keywords::HttpPatch);
 	
 	FReviewComment CleanedComment = Comment;
 	PutMetadataInBody(CleanedComment);
-	HttpRequest->SetContentAsString(CleanedComment.ToJson());
-	HttpRequest->OnProcessRequestComplete().BindLambda([OnComplete, this](FHttpRequestPtr Request, FHttpResponsePtr Response, bool bConnectedSuccessfully)
+	
+	const TSharedPtr<FJsonObject> JsonObject = CleanedComment.ToJsonObject();
+	if (bSilenceNotification)
+	{
+		JsonObject->SetStringField(Keywords::SilenceNotification, Keywords::True);
+	}
+
+	FString JsonString;
+	const TSharedRef<TJsonWriter<>> Writer = TJsonWriterFactory<>::Create(&JsonString);
+	FJsonSerializer::Serialize(JsonObject.ToSharedRef(), Writer);
+	
+	HttpRequest->SetContentAsString(JsonString);
+	HttpRequest->OnProcessRequestComplete().BindLambda([OnComplete](FHttpRequestPtr Request, FHttpResponsePtr Response, bool bConnectedSuccessfully)
 	{
 		if (!bConnectedSuccessfully)
 		{
@@ -172,10 +416,18 @@ void FSwarmCommentsAPI::GetReviewTopicForCL(const FString& ChangelistNum,
 			ChangelistNum, // get the review for a specific CL
 			1 // we only want a single review.
 		}));
-	HttpRequest->SetVerb(TEXT("GET"));
+	HttpRequest->SetVerb(Keywords::HttpGet);
+
+	TWeakPtr<const FSwarmCommentsAPI> WeakSelf = AsWeak();
 	
-	HttpRequest->OnProcessRequestComplete().BindLambda([OnComplete, ChangelistNum, this](FHttpRequestPtr Request, FHttpResponsePtr Response, bool bConnectedSuccessfully)
+	HttpRequest->OnProcessRequestComplete().BindLambda([OnComplete, ChangelistNum, WeakSelf](FHttpRequestPtr Request, FHttpResponsePtr Response, bool bConnectedSuccessfully)
 	{
+		const TSharedPtr<const FSwarmCommentsAPI> Self = WeakSelf.Pin();
+		if (!Self)
+		{
+			return;
+		}
+		
 		if (!bConnectedSuccessfully)
 		{
 			OnComplete.ExecuteIfBound({}, TEXT("Connection Failed"));
@@ -195,7 +447,7 @@ void FSwarmCommentsAPI::GetReviewTopicForCL(const FString& ChangelistNum,
 				const TArray<TSharedPtr<FJsonValue>> ReviewsJson = JsonObject->GetArrayField("reviews");
 				if (ReviewsJson.IsEmpty())
 				{
-					CreateReviewTopicForCL(ChangelistNum, OnComplete);
+					Self->CreateReviewTopicForCL(ChangelistNum, OnComplete);
 					return;
 				}
 				
@@ -215,8 +467,7 @@ void FSwarmCommentsAPI::GetReviewTopicForCL(const FString& ChangelistNum,
 	HttpRequest->ProcessRequest();
 }
 
-void FSwarmCommentsAPI::CreateReviewTopicForCL(const FString& ChangelistNum,
-	const OnGetReviewTopicForCLComplete& OnComplete) const
+void FSwarmCommentsAPI::CreateReviewTopicForCL(const FString& ChangelistNum, const OnGetReviewTopicForCLComplete& OnComplete) const
 {
 	
 	const TSharedRef<IHttpRequest> HttpRequest = FHttpModule::Get().CreateRequest();
@@ -225,9 +476,9 @@ void FSwarmCommentsAPI::CreateReviewTopicForCL(const FString& ChangelistNum,
 	
 	HttpRequest->SetURL(ReviewsURL());
 	HttpRequest->SetContentAsString(FString::Format(TEXT("{\"change\":{0}}"), {ChangelistNum}));
-	HttpRequest->SetVerb(TEXT("POST"));
+	HttpRequest->SetVerb(Keywords::HttpPost);
 	
-	HttpRequest->OnProcessRequestComplete().BindLambda([OnComplete, ChangelistNum, this](FHttpRequestPtr Request, FHttpResponsePtr Response, bool bConnectedSuccessfully)
+	HttpRequest->OnProcessRequestComplete().BindLambda([OnComplete](FHttpRequestPtr Request, FHttpResponsePtr Response, bool bConnectedSuccessfully)
 	{
 		if (!bConnectedSuccessfully)
 		{
@@ -271,60 +522,120 @@ FString FSwarmCommentsAPI::ReviewsURL() const
 	return SwarmURL.IsEmpty() ? FString{} : SwarmURL / TEXT("api/v9/reviews");
 }
 
-
-static TMap<FString, FString> ParseReviewCommentMetadata(const FString& Comment)
+void FSwarmCommentsAPI::GetUserdata(const FReviewTopic& Topic, const OnGetUserData& OnComplete) const
 {
-	TMap<FString, FString> Result;
-	FString Body, Metadata;
-	if (!Comment.Split(TEXT("[metadata]"), &Body, &Metadata))
+	if (FSwarmUserdataComment* Found = UserdataCache.Find(Topic))
 	{
-		return {{TEXT("Body"), Comment.TrimEnd()}};
+		OnComplete.ExecuteIfBound(Found, {});
+		return;
 	}
-	
-	Result.Add(TEXT("Body"), Body.TrimEnd());
-	
-	TArray<FString> MetaDataLines;
-	Metadata.ParseIntoArrayLines(MetaDataLines);
-	for (const FString& Line : MetaDataLines)
-	{
-		const int32 KeyEnd = Line.Find(TEXT("="), ESearchCase::CaseSensitive, ESearchDir::FromStart);
-		if (KeyEnd == INDEX_NONE)
+
+	GetComments(Topic, OnGetCommentsComplete::CreateLambda(
+		[&OnComplete, Topic, WeakSelf = AsWeak()](const TArray<FReviewComment>& Comments, const FString& ErrorMessage)
 		{
-			continue;
+			if (!ErrorMessage.IsEmpty())
+			{
+				OnComplete.ExecuteIfBound(nullptr, ErrorMessage);
+				return;
+			}
+			const TSharedPtr<const FSwarmCommentsAPI> Self = WeakSelf.Pin();
+			if (!Self)
+			{
+				return;
+			}
+			
+			// GetComments should've set the UserdataCommentCache
+			if (FSwarmUserdataComment* Found = Self->UserdataCache.Find(Topic))
+			{
+				OnComplete.ExecuteIfBound(Found, {});
+				return;
+			}
 		}
-		FString Key = Line.Left(KeyEnd);
-		
-		const int32 ValBegin = KeyEnd + 1; // skip '=' char
-		const int32 ValLen = Line.Len() - ValBegin; // extend to end of line
-		const FString Val = Line.Mid(ValBegin, ValLen);
-		Result.Add(Key, Val);
-	}
-	return Result;
+	));
 }
 
-static TMap<FString, FString> ParseReviewCommentMetadata(const FReviewComment& Comment)
+void FSwarmCommentsAPI::PostNewUserdataComment(const FReviewTopic& Topic) const
 {
-	return ParseReviewCommentMetadata(Comment.Body.Get({}));
+	FSwarmUserdataComment UserDataComment(GetUsername(), Topic);
+			
+	PostComment(UserDataComment, OnPostCommentComplete::CreateLambda(
+		[WeakSelf = AsWeak()](const FReviewComment& Comment, const FString& ErrorMessage)
+		{
+			if (!ErrorMessage.IsEmpty())
+			{
+				UE_LOG(LogSourceControl, Error, TEXT("FSwarmCommentsAPI::PostComment Error: %s"), *ErrorMessage)
+				return;
+			}
+			if (const TSharedPtr<const FSwarmCommentsAPI> Self = WeakSelf.Pin())
+			{
+				// archive the comment so it's hidden in swarm
+				FReviewComment Edit = Comment;
+				Edit.bIsClosed = true;
+				Edit.TaskState.Reset();
+				Self->EditComment(Edit, OnEditCommentComplete::CreateLambda(
+					[WeakSelf](const FReviewComment& Comment, const FString& ErrorMessage)
+					{
+						if (!ErrorMessage.IsEmpty())
+						{
+							UE_LOG(LogSourceControl, Error, TEXT("FSwarmCommentsAPI::EditComment Error: %s"), *ErrorMessage)
+							return;
+						}
+						if (const TSharedPtr<const FSwarmCommentsAPI> Self = WeakSelf.Pin())
+						{
+							Self->UserdataCache.Add(*Comment.Topic, Comment);
+						}
+					}), true);
+			}
+		}), true);
 }
 
-static void SetReviewCommentMetadata(FReviewComment& Comment, TMap<FString, FString> Metadata)
+void FSwarmCommentsAPI::UpdateUserdata(const FReviewComment& ReviewComment) const
 {
-	FString CommentText = Metadata["Body"];
-	if (Metadata.Num() == 1)
+	// unset fields shouldn't change
+	if (!ReviewComment.Likes.IsSet() && !ReviewComment.ReadBy.IsSet())
 	{
-		Comment.Body = CommentText;
 		return;
 	}
 	
-	CommentText += TEXT("\n\n[metadata]\n");
-	for (auto&[Key,Value] : Metadata)
-	{
-		if (Key != "Body")
+	GetUserdata(*ReviewComment.Topic, OnGetUserData::CreateLambda(
+		[ReviewComment, WeakSelf = AsWeak()](FSwarmUserdataComment* UserDataComment, const FString& ErrorMessage)
 		{
-			CommentText += FString::Format(TEXT("{0}={1}\n"), {Key, Value});
-		}
-	}
-	Comment.Body = CommentText;
+			if (!ErrorMessage.IsEmpty())
+			{
+				UE_LOG(LogSourceControl, Error, TEXT("FSwarmCommentsAPI::UpdateUserdata Error: %s"), *ErrorMessage)
+				return;
+			}
+			const TSharedPtr<const FSwarmCommentsAPI> Self = WeakSelf.Pin();
+			if (!Self)
+			{
+				return;
+			}
+
+			bool bUserDataChanged = false;
+			
+			if (ReviewComment.Likes.IsSet())
+			{
+				bUserDataChanged |= UserDataComment->ReflectCommentLikes(ReviewComment);
+			}
+			if (ReviewComment.ReadBy.IsSet())
+			{
+				bUserDataChanged |= UserDataComment->ReflectCommentReads(ReviewComment);
+			}
+			if (bUserDataChanged)
+			{
+				Self->EditComment(*UserDataComment, nullptr, true);
+			}
+		}));
+}
+
+TWeakPtr<const FSwarmCommentsAPI> FSwarmCommentsAPI::AsWeak() const
+{
+	return StaticCastSharedRef<const FSwarmCommentsAPI>(AsShared()).ToWeakPtr();
+}
+
+TWeakPtr<FSwarmCommentsAPI> FSwarmCommentsAPI::AsWeak()
+{
+	return StaticCastSharedRef<FSwarmCommentsAPI>(AsShared()).ToWeakPtr();
 }
 
 void FSwarmCommentsAPI::PutMetadataInBody(FReviewComment& Comment)
@@ -332,43 +643,19 @@ void FSwarmCommentsAPI::PutMetadataInBody(FReviewComment& Comment)
 	TMap<FString, FString> Metadata = ParseReviewCommentMetadata(Comment);
 	if (Comment.Context.File.IsSet())
 	{
-		Metadata.FindOrAdd(TEXT("File"), Comment.Context.File.GetValue());
+		Metadata.FindOrAdd(Keywords::MetadataFile, *Comment.Context.File);
 		Comment.Context.File.Reset();
 	}
 	if (Comment.Context.ReplyTo.IsSet())
 	{
-		Metadata.FindOrAdd(TEXT("ReplyTo"), FString::FromInt(Comment.Context.ReplyTo.GetValue()));
+		Metadata.FindOrAdd(Keywords::MetadataReplyTo, FString::FromInt(*Comment.Context.ReplyTo));
 		Comment.Context.ReplyTo.Reset();
 	}
 	if (Comment.Context.Category.IsSet())
 	{
-		Metadata.FindOrAdd(TEXT("Category"), Comment.Context.Category.GetValue());
+		Metadata.FindOrAdd(Keywords::MetadataCategory, *Comment.Context.Category);
 		Comment.Context.Category.Reset();
 	}
-	if (Comment.Likes.IsSet())
-	{
-		FString LikesString;
-		TSet<FString> &Likes = Comment.Likes.GetValue();
-		if (!Likes.IsEmpty())
-		{
-			int32 Index = 0;
-			for (const FString& User : Likes)
-			{
-				LikesString += User;
-				if (Index != Likes.Num() - 1)
-				{
-					LikesString += TEXT(",");
-				}
-				++Index;
-			}
-			Metadata.FindOrAdd(TEXT("Likes"), LikesString);
-		}
-		else
-		{
-			// if Likes is set but it's empty, explicitly remove it from the metadata
-			Metadata.Remove(TEXT("Likes"));
-		}
-	}
 	SetReviewCommentMetadata(Comment, Metadata);
 }
 
@@ -377,29 +664,19 @@ void FSwarmCommentsAPI::TakeMetadataFromBody(FReviewComment& Comment)
 	// because swarm has very limited support for certain properties, they may be stored inside the body
 	// as a workaround. parse them out and put them in their proper place.
 	TMap<FString, FString> Metadata = ParseReviewCommentMetadata(Comment);
-	if (const FString* File = Metadata.Find(TEXT("File")))
+	if (const FString* File = Metadata.Find(Keywords::MetadataFile))
 	{
 		Comment.Context.File = *File;
 	}
-	if (const FString* ReplyTo = Metadata.Find(TEXT("ReplyTo")))
+	if (const FString* ReplyTo = Metadata.Find(Keywords::MetadataReplyTo))
 	{
 		Comment.Context.ReplyTo = FCString::Atoi(**ReplyTo);
 	}
-	if (const FString* Category = Metadata.Find(TEXT("Category")))
+	if (const FString* Category = Metadata.Find(Keywords::MetadataCategory))
 	{
 		Comment.Context.Category = *Category;
 	}
-	if (const FString* Likes = Metadata.Find(TEXT("Likes")))
-	{
-		Comment.Likes = TSet<FString>();
-		TArray<FString> LikesArray;
-		Likes->ParseIntoArray(LikesArray, TEXT(","));
-		for (const FString& Like : LikesArray)
-		{
-			Comment.Likes.GetValue().Add(Like);
-		}
-	}
-	Comment.Body = Metadata[TEXT("Body")];
+	Comment.Body = Metadata[Keywords::MetadataBody];
 }
 
 
@@ -442,7 +719,6 @@ FString FSwarmCommentsAPI::RetrieveAuthorizationTicket()
 		return {};
 	}
 
-	// TODO: @jordan.hoffmann: make a setting to let users choose which user they want to use
 	const FString& TicketString = TicketStrings[0];
 	
 	// find beginning of ticket
@@ -471,6 +747,7 @@ FString FSwarmCommentsAPI::RetrieveSwarmURL()
 		P4Error.Fmt(&ErrorMsg);
 		UE_LOG(LogSourceControl, Error, TEXT("P4ERROR: Invalid connection to server."));
 		UE_LOG(LogSourceControl, Error, TEXT("%s"), ANSI_TO_TCHAR(ErrorMsg.Text()));
+		return {};
 	}
 
 	// Create a ClientUser that can capture the output from the command
@@ -509,6 +786,7 @@ FString FSwarmCommentsAPI::RetrieveSwarmURL()
 		P4Error.Fmt(&ErrorMsg);
 		UE_LOG(LogSourceControl, Error, TEXT("P4ERROR: Failed to disconnect from Server."));
 		UE_LOG(LogSourceControl, Error, TEXT("%s"), (ANSI_TO_TCHAR(ErrorMsg.Text())));
+		return {};
 	}
 
 	return P4User.Result;
@@ -521,7 +799,7 @@ FSwarmCommentsAPI::FAuthTicket::FAuthTicket(const FString& InUsername, const FSt
 FSwarmCommentsAPI::FAuthTicket::FAuthTicket(const FString& TicketString)
 {
 	int32 ChopIndex;
-	if (ensure(TicketString.FindChar(':', ChopIndex)))
+	if (TicketString.FindChar(':', ChopIndex))
 	{
 		Username = TicketString.Left(ChopIndex);
 		Password = TicketString.RightChop(ChopIndex + 1);
@@ -532,3 +810,85 @@ FSwarmCommentsAPI::FAuthTicket::operator FString() const
 {
 	return TEXT("Basic ") + FBase64::Encode(Username + TEXT(":") + Password);
 }
+
+FSwarmUserdataComment::FSwarmUserdataComment(FString User, FReviewTopic Topic)
+{
+	UserDataComment.User = User;
+	UserDataComment.Topic = Topic;
+	DataMap.Add(Keywords::UserdataLikedComments, {});
+	DataMap.Add(Keywords::UserdataReadComments, {});
+}
+
+FSwarmUserdataComment::FSwarmUserdataComment(const FReviewComment& InUserDataComment)
+	: UserDataComment(InUserDataComment)
+	, DataMap(ParseReviewCommentMetadata(UserDataComment, Keywords::UserdataTag))
+{
+	// reset irrelevant fields
+	UserDataComment.Likes.Reset();
+	UserDataComment.ReadBy.Reset();
+	UserDataComment.TaskState.Reset();
+}
+
+FSwarmUserdataComment::operator FReviewComment&()
+{
+	SetReviewCommentMetadata(UserDataComment, DataMap, Keywords::UserdataTag);
+	return UserDataComment;
+}
+
+bool FSwarmUserdataComment::ReflectCommentLikes(const FReviewComment& ReviewComment)
+{
+	const FString& Username = *UserDataComment.User;
+	const FString CommentID = FString::FromInt(*ReviewComment.CommentID);
+	const EUserdataSetToggleVal ToggleVal = ReviewComment.Likes->Find(Username)? EUserdataSetToggleVal::ToggleSet : EUserdataSetToggleVal::ToggleUnset;
+	return ToggleSetItem(Keywords::UserdataLikedComments, CommentID, ToggleVal);
+}
+
+bool FSwarmUserdataComment::ReflectCommentReads(const FReviewComment& ReviewComment)
+{
+	const FString& Username = *UserDataComment.User;
+	const FString CommentID = FString::FromInt(*ReviewComment.CommentID);
+	const EUserdataSetToggleVal ToggleVal = ReviewComment.ReadBy->Find(Username)? EUserdataSetToggleVal::ToggleSet : EUserdataSetToggleVal::ToggleUnset;
+	return ToggleSetItem(Keywords::UserdataReadComments, CommentID, ToggleVal);
+}
+
+TSet<FString> FSwarmUserdataComment::GetLikedComments() const
+{
+	return DeserializeSetString(DataMap[Keywords::UserdataLikedComments]);
+}
+
+TSet<FString> FSwarmUserdataComment::GetReadComments() const
+{
+	return DeserializeSetString(DataMap[Keywords::UserdataReadComments]);
+}
+
+bool FSwarmUserdataComment::IsUserdataComment(const FReviewComment& Comment)
+{
+	if (!Comment.Body.IsSet())
+	{
+		return false;
+	}
+	return Comment.Body->Find(Keywords::UserdataTag) != INDEX_NONE;
+}
+
+bool FSwarmUserdataComment::ToggleSetItem(const FString& SetName, const FString& SetKey, EUserdataSetToggleVal ToggleVal)
+{
+	TSet<FString> Set = DeserializeSetString(DataMap[SetName]);
+	switch(ToggleVal)
+	{
+	case EUserdataSetToggleVal::ToggleSet:
+		Set.FindOrAdd(SetKey);
+		break;
+	case EUserdataSetToggleVal::ToggleUnset:
+		Set.Remove(SetKey);
+		break;
+	}
+
+	const FString NewSetString = SerializeSetString(Set);
+	if (DataMap[SetName] != NewSetString)
+	{
+		DataMap[SetName] = NewSetString;
+		SetReviewCommentMetadata(UserDataComment, DataMap, Keywords::UserdataTag);
+		return true;
+	}
+	return false;
+}
diff --git a/Engine/Plugins/Editor/ChangelistReview/Source/Private/FSwarmCommentsAPI.h b/Engine/Plugins/Editor/ChangelistReview/Source/Private/FSwarmCommentsAPI.h
--- a/Engine/Plugins/Editor/ChangelistReview/Source/Private/FSwarmCommentsAPI.h
+++ b/Engine/Plugins/Editor/ChangelistReview/Source/Private/FSwarmCommentsAPI.h
@@ -3,27 +3,58 @@
 #pragma once
 #include "ReviewComments.h"
 
+// helper class that stores likes and Read comments as an archived comment in a review
+class FSwarmUserdataComment
+{
+public:
+	FSwarmUserdataComment(FString User, FReviewTopic Topic);
+	FSwarmUserdataComment(const FReviewComment& InUserDataComment);
+	operator FReviewComment&();
+	bool ReflectCommentLikes(const FReviewComment& ReviewComment);
+	bool ReflectCommentReads(const FReviewComment& ReviewComment);
+	TSet<FString> GetLikedComments() const;
+	TSet<FString> GetReadComments() const;
+	
+	static bool IsUserdataComment(const FReviewComment& Comment);
+private:
+	enum class EUserdataSetToggleVal { ToggleSet, ToggleUnset };
+	bool ToggleSetItem(const FString& SetName, const FString& SetKey, EUserdataSetToggleVal ToggleVal);
+	FReviewComment UserDataComment;
+	TMap<FString,FString> DataMap;
+};
+
 class FSwarmCommentsAPI : public IReviewCommentAPI
 {
 public:
-	FSwarmCommentsAPI();
+	// prefer TryConnect over calling constructor directly
+	FSwarmCommentsAPI(const FString& AuthTicket, const FString& SwarmURL);
+	static TSharedPtr<FSwarmCommentsAPI> TryConnect();
 
 	virtual FString GetUsername() const override;
 
 	virtual void GetComments(const FReviewTopic& Topic, const OnGetCommentsComplete& OnComplete) const override;
 	// Body must be set
-	virtual void PostComment(FReviewComment& Comment, const OnPostCommentComplete& OnComplete) const override;
+	virtual void PostComment(FReviewComment& Comment, const OnPostCommentComplete& OnComplete, bool bSilenceNotification = false) const override;
 	// all unset members will be left unchanged.
-	virtual void EditComment(const FReviewComment& Comment, const OnEditCommentComplete& OnComplete) const override;
+	virtual void EditComment(const FReviewComment& Comment, const OnEditCommentComplete& OnComplete, bool bSilenceNotification = false) const override;
 	
 	// retrieves the review associated with a CL. If there is none, it creates a new review and returns that.
 	virtual void GetReviewTopicForCL(const FString &ChangelistNum, const OnGetReviewTopicForCLComplete& OnComplete) const override;
 
 private:
+	
 	void CreateReviewTopicForCL(const FString &ChangelistNum, const OnGetReviewTopicForCLComplete& OnComplete) const;
 
 	FString CommentsURL() const;
 	FString ReviewsURL() const;
+
+	DECLARE_DELEGATE_TwoParams(OnGetUserData, FSwarmUserdataComment*, const FString& /*ErrorMessage*/)
+	void GetUserdata(const FReviewTopic& Topic, const OnGetUserData& OnComplete) const;
+	void PostNewUserdataComment(const FReviewTopic& Topic) const;
+	void UpdateUserdata(const FReviewComment& ReviewComment) const;
+
+	TWeakPtr<const FSwarmCommentsAPI> AsWeak() const;
+	TWeakPtr<FSwarmCommentsAPI> AsWeak();
 	
 	static void PutMetadataInBody(FReviewComment& Comment);
 	static void TakeMetadataFromBody(FReviewComment& Comment);
@@ -45,4 +76,8 @@ private:
 
 	// base url for all swarm api requests
 	FString SwarmURL;
+
+	// maps user to the comment that stores the users data like their likes and views
+	// this is necessary because the swarm api unfortunately has incomplete support for updating like and view state
+	mutable TMap<FReviewTopic, FSwarmUserdataComment> UserdataCache;
 };
diff --git a/Engine/Plugins/Editor/ChangelistReview/Source/Private/SSourceControlReview.cpp b/Engine/Plugins/Editor/ChangelistReview/Source/Private/SSourceControlReview.cpp
--- a/Engine/Plugins/Editor/ChangelistReview/Source/Private/SSourceControlReview.cpp
+++ b/Engine/Plugins/Editor/ChangelistReview/Source/Private/SSourceControlReview.cpp
@@ -4,6 +4,7 @@
 
 #include "AssetRegistry/AssetData.h"
 #include "ClassIconFinder.h"
+#include "FSwarmCommentsAPI.h"
 #include "Framework/Views/TableViewMetadata.h"
 #include "ISourceControlModule.h"
 #include "SourceControlOperations.h"
@@ -29,6 +30,13 @@
 
 using namespace SourceControlReview;
 
+namespace UE::DiffControl
+{
+	DECLARE_MULTICAST_DELEGATE_OneParam(FOnCommentPosted, const FReviewComment&)
+	extern KISMET_API FOnCommentPosted GOnCommentPosted;
+}
+
+
 namespace ReviewHelpers
 {
 	const FString FileDepotKey = TEXT("depotFile");
@@ -80,6 +88,8 @@ void SSourceControlReview::Construct(const FArguments& InArgs)
 
 	const static FMargin InfoWidgetMargin(4.f, 2.f, 4.f, 8.f);
 
+	CommentsAPI = FSwarmCommentsAPI::TryConnect();
+	
 	LoadCLHistory();
 	
 	ChildSlot
@@ -305,6 +315,9 @@ void SSourceControlReview::LoadChangelist(const FString& Changelist)
 	}
 
 	ChangelistFiles.Empty();
+	GlobalReviewComments.Empty();
+	FileReviewComments.Empty();
+	TempLocalPathToDepotPath.Empty();
 
 	//This command runs p4 -describe (or similar for other version controls) to retrieve changelist record information
 	GetChangelistDetailsCommand = ISourceControlOperation::Create<FGetChangelistDetails>();
@@ -313,7 +326,130 @@ void SSourceControlReview::LoadChangelist(const FString& Changelist)
 	ISourceControlModule::Get().GetProvider().Execute(
 		GetChangelistDetailsCommand.ToSharedRef(),
 		EConcurrency::Asynchronous,
-		FSourceControlOperationComplete::CreateRaw(this, &SSourceControlReview::OnChangelistLoadComplete, Changelist));
+		FSourceControlOperationComplete::CreateRaw(this, &SSourceControlReview::OnGetChangelistDetails, Changelist));
+}
+
+
+const TArray<FReviewComment>* SSourceControlReview::GetReviewCommentsForFile(const FString& FilePath)
+{
+	return FileReviewComments.Find(AsDepotPath(FilePath));
+}
+
+void SSourceControlReview::UpdateReviewComments()
+{
+	if (bReviewCommentsDirty && CommentsAPI)
+	{
+		CommentsAPI->GetComments(ReviewTopic.GetValue(), FSwarmCommentsAPI::OnGetCommentsComplete::CreateRaw(this, &SSourceControlReview::OnGetReviewComments));
+		bReviewCommentsDirty = false;
+	}
+}
+
+void SSourceControlReview::PostComment(FReviewComment& Comment)
+{
+	if (!CommentsAPI)
+	{
+		return;
+	}
+	
+	bReviewCommentsDirty = true;
+	
+	// convert filepath from local to depot path
+	if (Comment.Context.File.IsSet())
+	{
+		FString& Path = Comment.Context.File.GetValue();
+		Path = AsDepotPath(Path);
+	}
+	
+	Comment.Topic = ReviewTopic;
+	CommentsAPI->PostComment(Comment, IReviewCommentAPI::OnPostCommentComplete::CreateLambda(
+		[](const FReviewComment& Comment, const FString& ErrorMessage)
+		{
+			if (!ErrorMessage.IsEmpty())
+			{
+				UE_LOG(LogSourceControl, Error, TEXT("IReviewCommentsAPI::PostComment Error: %s"), *ErrorMessage)
+				return;
+			}
+			// notify listeners that the comment successfully posted
+			if (UE::DiffControl::GOnCommentPosted.IsBound())
+			{
+				UE::DiffControl::GOnCommentPosted.Broadcast(Comment);
+			}
+		}));
+}
+
+void SSourceControlReview::EditComment(FReviewComment& Comment)
+{
+	if (!CommentsAPI)
+	{
+		return;
+	}
+	
+	bReviewCommentsDirty = true;
+	
+	// convert filepath from local to depot path
+	if (Comment.Context.File.IsSet())
+	{
+		FString& Path = Comment.Context.File.GetValue();
+		Path = AsDepotPath(Path);
+	}
+	
+	Comment.Topic = ReviewTopic;
+	CommentsAPI->EditComment(Comment, IReviewCommentAPI::OnEditCommentComplete::CreateLambda(
+		[](const FReviewComment& Comment, const FString& ErrorMessage)
+		{
+			if (!ErrorMessage.IsEmpty())
+			{
+				UE_LOG(LogSourceControl, Error, TEXT("IReviewCommentsAPI::EditComment Error: %s"), *ErrorMessage)
+			}
+		}));
+}
+
+FString SSourceControlReview::GetReviewerUsername() const
+{
+	return CommentsAPI? CommentsAPI->GetUsername() : FString{};
+}
+
+bool SSourceControlReview::IsFileInReview(const FString& File) const
+{
+	FString Path = File;
+	if (!Path.StartsWith(TEXT("//")))
+	{
+		// if a local path was passed in, commentable paths are found in TempLocalPathToDepotPath
+		Path = FPaths::ConvertRelativePathToFull(Path);
+		FPaths::NormalizeFilename(Path);
+		return TempLocalPathToDepotPath.Find(Path) != nullptr;
+	}
+
+	// if a depot path was passed in, we can early out if it's in FileReviewComments
+	if (FileReviewComments.Find(Path))
+	{
+		return true;
+	}
+
+	// if the file hasn't been commented on before, it might still be in the review. fallback to searching ChangelistFiles
+	for (const TSharedPtr<FChangelistFileData>& FileData : ChangelistFiles)
+	{
+		if (FileData->AssetDepotPath == Path)
+		{
+			return true;
+		}
+	}
+	return false;
+}
+
+FString SSourceControlReview::AsDepotPath(const FString& FilePath)
+{
+	if (!FilePath.StartsWith(TEXT("//")))
+	{
+		FString LocalPath = FPaths::ConvertRelativePathToFull(FilePath);
+		FPaths::NormalizeFilename(LocalPath);
+		if (const FString* DepotPath = TempLocalPathToDepotPath.Find(LocalPath))
+		{
+			return *DepotPath;
+		}
+		return {};
+	}
+	return FilePath;
 }
 
 void SSourceControlReview::CommitTempChangelistNumToHistory()
@@ -353,7 +489,17 @@ void SSourceControlReview::RemoveUncommittedChangelistNumFromHistory()
 	}
 }
 
-void SSourceControlReview::OnChangelistLoadComplete(const FSourceControlOperationRef& InOperation, ECommandResult::Type InResult, FString Changelist)
+void SSourceControlReview::IncrementItemsLoaded()
+{
+	NumItemsLoaded++;
+	LoadingProgressBar->SetPercent(NumItemsToLoad ? static_cast<float>(NumItemsLoaded) / static_cast<float>(NumItemsToLoad) : 1.f);
+	if (NumItemsToLoad == NumItemsLoaded)
+	{
+		OnLoadComplete();
+	}
+}
+
+void SSourceControlReview::OnGetChangelistDetails(const FSourceControlOperationRef& InOperation, ECommandResult::Type InResult, FString Changelist)
 {
 	// this command is cancelled when this widget is destroyed. Exit immediately to avoid touching invalid data
 	if (InResult == ECommandResult::Cancelled)
@@ -374,9 +520,9 @@ void SSourceControlReview::OnChangelistLoadComplete(const FSourceControlOperatio
 	const TMap<FString, FString>& ChangelistRecord = Record[ReviewHelpers::RecordIndex];
 	
 	//Num files we are going to expect retrieved from source control
-	FilesToLoad = 0;
+	NumItemsToLoad = 0;
 	//Num files we've loaded so far
-	FilesLoaded = 0;
+	NumItemsLoaded = 0;
 
 	// TODO: @jordan.hoffmann revise for other version controls
 	//Each file in p4 has an index "depotFile0" this is the index that is updated by the while loop to find all the files "depotFile1", "depotFile2" etc...
@@ -389,17 +535,23 @@ void SSourceControlReview::OnChangelistLoadComplete(const FSourceControlOperatio
 	FString RecordRevisionMapKey = ReviewHelpers::FileRevisionKey + RecordFileIndexStr;
 	//The p4 records is the map a revision key starts with "action" and is followed by file index 
 	FString RecordActionMapKey = ReviewHelpers::FileActionKey + RecordFileIndexStr;
+	const bool bIsShelved = ChangelistRecord[ReviewHelpers::ChangelistStatusKey] == ReviewHelpers::ChangelistPendingStatusKey;
 
-	SetChangelistInfo(ChangelistRecord);
+	SetChangelistInfo(ChangelistRecord, Changelist);
 	CommitTempChangelistNumToHistory();
 	SaveCLHistory();
 	
+	NumItemsToLoad += 2; // require review topic and comments to be loaded before continuing
+	if (CommentsAPI)
+	{
+		CommentsAPI->GetReviewTopicForCL(Changelist, FSwarmCommentsAPI::OnGetReviewTopicForCLComplete::CreateRaw(this, &SSourceControlReview::OnGetReviewTopic));
+	}
+	
 	//the loop checks if we have a valid record "depotFile(Index)" in the records to add a file entry
 	while (ChangelistRecord.Contains(RecordFileMapKey) && ChangelistRecord.Contains(RecordRevisionMapKey))
 	{
 		const FString &FileDepotPath = ChangelistRecord[RecordFileMapKey];
 
-		const bool bIsShelved = ChangelistRecord[ReviewHelpers::ChangelistStatusKey] == ReviewHelpers::ChangelistPendingStatusKey;
 		const int32 AssetRevision = FCString::Atoi(*ChangelistRecord[RecordRevisionMapKey]);
 	
 		TSharedPtr<FChangelistFileData> ChangelistFileData = MakeShared<FChangelistFileData>();
@@ -437,7 +589,7 @@ void SSourceControlReview::OnChangelistLoadComplete(const FSourceControlOperatio
 		// retrieve files directly from source control into a temp location
 		if (ChangelistFileData->FileSourceControlAction != ESourceControlAction::Delete)
 		{
-			FilesToLoad++;
+			NumItemsToLoad++;
 			
 			TSharedRef<FGetFile> GetFileToReviewCommand = ISourceControlOperation::Create<FGetFile>(Changelist, ChangelistRecord[RecordRevisionMapKey], ChangelistRecord[RecordFileMapKey], bIsShelved);
 
@@ -446,7 +598,7 @@ void SSourceControlReview::OnChangelistLoadComplete(const FSourceControlOperatio
 		
 		if (ChangelistFileData->FileSourceControlAction != ESourceControlAction::Add && ChangelistFileData->FileSourceControlAction != ESourceControlAction::Branch)
 		{
-			FilesToLoad++;
+			NumItemsToLoad++;
 
 			TSharedRef<FGetFile> GetPreviousFileCommand = ISourceControlOperation::Create<FGetFile>(Changelist, PreviousAssetRevisionStr, ChangelistRecord[RecordFileMapKey], false);
 
@@ -463,12 +615,123 @@ void SSourceControlReview::OnChangelistLoadComplete(const FSourceControlOperatio
 	}
 
 	//If we have no files to load flip loading bar visibility
-	if (FilesToLoad == 0)
+	if (NumItemsToLoad == 0)
 	{
 		SetLoading(false);
 	}
 }
 
+
+void SSourceControlReview::OnGetReviewTopic(const FReviewTopic& Topic, const FString& ErrorMessage)
+{
+	if(!ErrorMessage.IsEmpty())
+	{
+		UE_LOG(LogSourceControl, Error, TEXT("IReviewCommentsAPI::GetReviewTopicForCL Error: %s"), *ErrorMessage)
+		ReviewTopic= {CurrentChangelistInfo.ChangelistNum, CurrentChangelistInfo.Author.ToString() == TEXT("Swarm") ? EReviewTopicType::Review : EReviewTopicType::Change};
+	}
+	else
+	{
+		ReviewTopic = Topic;
+	}
+	IncrementItemsLoaded();
+	
+	if (CommentsAPI)
+	{
+		CommentsAPI->GetComments(ReviewTopic.GetValue(), FSwarmCommentsAPI::OnGetCommentsComplete::CreateRaw(this, &SSourceControlReview::OnInitReviewComments));
+	}
+	
+	// every 3 seconds, check whether comment data is dirty and if so, grab the update from swarm
+	ChangelistInfoWidget.Get()->RegisterActiveTimer(3.f, FWidgetActiveTimerDelegate::CreateLambda([this](double, float)
+	{
+		UpdateReviewComments();
+		return EActiveTimerReturnType::Continue;
+	}));
+}
+
+void SSourceControlReview::OnGetReviewComments(const TArray<FReviewComment>& Comments, const FString& ErrorMessage)
+{
+	bReviewCommentsDirty = false;
+	if (!ErrorMessage.IsEmpty())
+	{
+		UE_LOG(LogSourceControl, Error, TEXT("IReviewCommentsAPI::GetComments Error: %s"), *ErrorMessage)
+		return;
+	}
+
+	FileReviewComments.Empty();
+	GlobalReviewComments.Empty();
+	for (const FReviewComment& Comment : Comments)
+	{
+		if (Comment.Context.File.IsSet())
+		{
+			const FString& File = Comment.Context.File.GetValue();
+			TArray<FReviewComment>& FileComments = FileReviewComments.FindOrAdd(File);
+			FileComments.Add(Comment);
+			continue;
+		}
+		GlobalReviewComments.Add(Comment);
+	}
+}
+
+void SSourceControlReview::OnInitReviewComments(const TArray<FReviewComment>& Comments, const FString& ErrorMessage)
+{
+	OnGetReviewComments(Comments, ErrorMessage);
+	IncrementItemsLoaded();
+}
+
+void SSourceControlReview::OnGetFileFromSourceControl(TSharedPtr<FChangelistFileData> ChangelistFileData)
+{
+	if (ChangelistFileData && ChangelistFileData->IsDataValidForEntry())
+	{
+		ChangelistFiles.Add(ChangelistFileData);
+		if (!ChangelistFileData->ReviewFileTempPath.IsEmpty())
+		{
+			TempLocalPathToDepotPath.Add(ChangelistFileData->ReviewFileTempPath, ChangelistFileData->AssetDepotPath);
+		}
+		if (!ChangelistFileData->PreviousFileTempPath.IsEmpty())
+		{
+			TempLocalPathToDepotPath.Add(ChangelistFileData->PreviousFileTempPath, ChangelistFileData->AssetDepotPath);
+		}
+	}
+	
+	IncrementItemsLoaded();
+}
+
+void SSourceControlReview::OnLoadComplete()
+{
+	const FAssetRegistryModule& AssetRegistryModule = FModuleManager::LoadModuleChecked<FAssetRegistryModule>( TEXT( "AssetRegistry" ) );
+
+	// Merge asset moves and renames into single entries
+	FixupRedirectors();
+		
+	Algo::Sort(ChangelistFiles, [](const TSharedPtr<FChangelistFileData> &A, const TSharedPtr<FChangelistFileData> &B)
+	{
+		return A->RelativeFilePath < B->RelativeFilePath;
+	});
+	
+	TArray<FString> ChangelistFilePaths;
+	for (const TSharedPtr<FChangelistFileData>& FileData : ChangelistFiles)
+	{
+		if(FileData)
+		{
+			ChangelistFilePaths.Add(FileData->ReviewFileTempPath);
+		}
+	}
+	
+	AssetRegistryModule.Get().ScanFilesSynchronous(ChangelistFilePaths);
+
+	// now that the files are in the asset registry, cache their associated class so their class icons can be created quickly
+	for (const TSharedPtr<FChangelistFileData>& FileData : ChangelistFiles)
+	{
+		if(FileData)
+		{
+			FileData->GetIconClass();
+		}
+	}
+	
+	ChangelistEntriesWidget->RebuildList();
+	SetLoading(false);
+}
+
 FReply SSourceControlReview::OnLoadChangelistClicked()
 {
 	LoadChangelist(ChangelistNumText->GetText().ToString());
@@ -605,54 +868,6 @@ void SSourceControlReview::LoadCLHistory()
 	}
 }
 
-void SSourceControlReview::OnGetFileFromSourceControl(TSharedPtr<FChangelistFileData> ChangelistFileData)
-{
-	if (ChangelistFileData && ChangelistFileData->IsDataValidForEntry())
-	{
-		ChangelistFiles.Add(ChangelistFileData);
-	}
-
-	FilesLoaded++;
-
-	LoadingProgressBar->SetPercent(FilesToLoad ? static_cast<float>(FilesLoaded) / static_cast<float>(FilesToLoad) : 1.f);
-
-	if (FilesToLoad == FilesLoaded)
-	{
-		const FAssetRegistryModule& AssetRegistryModule = FModuleManager::LoadModuleChecked<FAssetRegistryModule>( TEXT( "AssetRegistry" ) );
-
-		// Merge asset moves and renames into single entries
-		FixupRedirectors();
-		
-		Algo::Sort(ChangelistFiles, [](const TSharedPtr<FChangelistFileData> &A, const TSharedPtr<FChangelistFileData> &B)
-		{
-			return A->RelativeFilePath < B->RelativeFilePath;
-		});
-		
-		TArray<FString> ChangelistFilePaths;
-		for (const TSharedPtr<FChangelistFileData>& FileData : ChangelistFiles)
-		{
-			if(FileData)
-			{
-				ChangelistFilePaths.Add(FileData->ReviewFileTempPath);
-			}
-		}
-		
-		AssetRegistryModule.Get().ScanFilesSynchronous(ChangelistFilePaths);
-
-		// now that the files are in the asset registry, cache their associated class so their class icons can be created quickly
-		for (const TSharedPtr<FChangelistFileData>& FileData : ChangelistFiles)
-		{
-			if(FileData)
-			{
-				FileData->GetIconClass();
-			}
-		}
-		
-		ChangelistEntriesWidget->RebuildList();
-		SetLoading(false);
-	}
-}
-
 void SSourceControlReview::SetLoading(bool bInLoading)
 {
 	bChangelistLoading = bInLoading;
@@ -796,18 +1011,20 @@ static FString GetSharedBranchPath(const TMap<FString, FString>& InChangelistRec
 	return TrimIndex == INDEX_NONE ? FString() : SharedBranchPath.Left(TrimIndex);
 }
 
-void SSourceControlReview::SetChangelistInfo(const TMap<FString, FString>& InChangelistRecord)
+void SSourceControlReview::SetChangelistInfo(const TMap<FString, FString>& InChangelistRecord, const FString& ChangelistNum)
 {
 	CurrentChangelistInfo.Author = FText::FromString(InChangelistRecord[ReviewHelpers::AuthorKey]);
 	CurrentChangelistInfo.Description = FText::FromString(InChangelistRecord[ReviewHelpers::DescriptionKey]);
 	CurrentChangelistInfo.Status = FText::FromString(InChangelistRecord[ReviewHelpers::ChangelistStatusKey]);
 	CurrentChangelistInfo.SharedPath = FText::FromString(GetSharedBranchPath(InChangelistRecord));
+	CurrentChangelistInfo.ChangelistNum = ChangelistNum;
 }
 
 TSharedRef<ITableRow> SSourceControlReview::OnGenerateFileRow(TSharedPtr<FChangelistFileData> FileData, const TSharedRef<STableViewBase>& Table) const
 {
 	return SNew(SSourceControlReviewEntry, Table)
-	.FileData(*FileData);
+	.FileData(*FileData)
+	.CommentsAPI(CommentsAPI);
 }
 
 FString SSourceControlReview::TrimSharedPath(FString FullCLPath) const
@@ -927,7 +1144,7 @@ SHeaderRow::FColumn::FArguments SSourceControlReview::HeaderColumn(FName HeaderN
 	if (HeaderName == ColumnIds::Status)
 	{
 		ColumnLabel = LOCTEXT("StatusColumnHeader", "Status");
-		ColumnWidth = 60;
+		ColumnWidth = 100;
 	}
 	else if (HeaderName == ColumnIds::File)
 	{
diff --git a/Engine/Plugins/Editor/ChangelistReview/Source/Private/SSourceControlReview.h b/Engine/Plugins/Editor/ChangelistReview/Source/Private/SSourceControlReview.h
deleted file mode 100644
diff --git a/Engine/Plugins/Editor/ChangelistReview/Source/Private/SSourceControlReviewEntry.cpp b/Engine/Plugins/Editor/ChangelistReview/Source/Private/SSourceControlReviewEntry.cpp
--- a/Engine/Plugins/Editor/ChangelistReview/Source/Private/SSourceControlReviewEntry.cpp
+++ b/Engine/Plugins/Editor/ChangelistReview/Source/Private/SSourceControlReviewEntry.cpp
@@ -4,6 +4,7 @@
 
 #include "AssetDefinition.h"
 #include "AssetToolsModule.h"
+#include "ChangelistReviewModule.h"
 #include "ClassIconFinder.h"
 #include "DiffUtils.h"
 #include "Engine/Blueprint.h"
@@ -30,6 +31,7 @@ namespace ReviewEntryConsts
 void SSourceControlReviewEntry::Construct(const FArguments& InArgs, const TSharedRef<STableViewBase>& InOwnerTableView)
 {
 	ChangelistFileData = InArgs._FileData;
+	CommentsAPI = InArgs._CommentsAPI;
 	SMultiColumnTableRow::Construct(FSuperRowType::FArguments(), InOwnerTableView);
 
 	// figure out how this asset diffs, and bind it to this->DiffMethod
@@ -42,10 +44,25 @@ TSharedRef<SWidget> SSourceControlReviewEntry::GenerateWidgetForColumn(const FNa
 	
 	if (ColumnName == SourceControlReview::ColumnIds::Status)
 	{
-		SAssignNew(InnerContent, SImage)
-		.DesiredSizeOverride(FVector2D(20.f, 20.f))
-		.ColorAndOpacity(this, &SSourceControlReviewEntry::GetSourceControlIconColor)
-		.Image(this, &SSourceControlReviewEntry::GetSourceControlIcon);
+		SAssignNew(InnerContent, SHorizontalBox)
+		+SHorizontalBox::Slot()
+		.Padding(0.f, 0.f, 5.f, 0.f)
+		[
+			SNew(SImage)
+			.DesiredSizeOverride(FVector2D(20.f, 20.f))
+			.ColorAndOpacity(this, &SSourceControlReviewEntry::GetSourceControlIconColor)
+			.Image(this, &SSourceControlReviewEntry::GetSourceControlIcon)
+		]
+		+SHorizontalBox::Slot()
+		.Padding(5.f, 0.f, 0.f, 0.f)
+		[
+			SNew(SImage)
+			.Visibility(this, &SSourceControlReviewEntry::GetUnreadCommentsIconVisibility)
+			.ToolTipText(this, &SSourceControlReviewEntry::GetUnreadCommentsTooltip)
+			.DesiredSizeOverride(FVector2D(20.f, 20.f))
+			.ColorAndOpacity(FStyleColors::Foreground)
+			.Image(FAppStyle::Get().GetBrush(TEXT("Icons.Comment")))
+		];
 	}
 	else if (ColumnName == SourceControlReview::ColumnIds::File)
 	{
@@ -441,6 +458,78 @@ FText SSourceControlReviewEntry::GetLocalAssetPathText() const
 	return FText::FromString(TEXT("...") + ChangelistFileData.RelativeFilePath);
 }
 
+const TArray<FReviewComment>* SSourceControlReviewEntry::GetReviewComments() const
+{
+	const TSharedPtr<SSourceControlReview> Review = FChangelistReviewModule::Get().GetActiveReview().Pin();
+	return Review? Review->GetReviewCommentsForFile(ChangelistFileData.AssetDepotPath) : nullptr;
+}
+
+FString SSourceControlReviewEntry::GetReviewerUsername() const
+{
+	const TSharedPtr<IReviewCommentAPI> ReviewComments = CommentsAPI.Pin();
+	return ReviewComments? ReviewComments->GetUsername() : FString();
+}
+
+static bool IsCommentUnread(const FReviewComment& Comment, const FString& Username)
+{
+	// treat archived comments as read
+	if (Comment.bIsClosed)
+	{
+		return false;
+	}
+	// if there is no 'ReadBy' set, then the comment is unread
+	if (!Comment.ReadBy.IsSet())
+	{
+		return true;
+	}
+	// if the logged in user isn't in the 'ReadBy' set, there's an unread comment
+	if (!Comment.ReadBy->Contains(Username))
+	{
+		return true;
+	}
+	return false;
+}
+
+EVisibility SSourceControlReviewEntry::GetUnreadCommentsIconVisibility() const
+{
+	const TArray<FReviewComment>* Comments = GetReviewComments();
+	const FString Username = GetReviewerUsername();
+	if (!Comments || Username.IsEmpty())
+	{
+		return EVisibility::Collapsed;
+	}
+	
+	for (const FReviewComment& Comment : *Comments)
+	{
+		if (IsCommentUnread(Comment, Username))
+		{
+			return EVisibility::Visible;
+		}
+	}
+	return EVisibility::Collapsed;
+}
+
+FText SSourceControlReviewEntry::GetUnreadCommentsTooltip() const
+{
+	const TArray<FReviewComment>* Comments = GetReviewComments();
+	const FString Username = GetReviewerUsername();
+	if (!Comments || Username.IsEmpty())
+	{
+		return {};
+	}
+	
+	int32 NumUnreadComments = 0;
+	for (const FReviewComment& Comment : *Comments)
+	{
+		if (IsCommentUnread(Comment, Username))
+		{
+			++NumUnreadComments;
+		}
+	}
+	
+	return FText::Format(LOCTEXT("UnreadCommentsTooltip", "{0} Unread {0}|plural(one=Comment,other=Comments)"), NumUnreadComments);
+}
+
 const FString& SSourceControlReviewEntry::GetSearchableString() const
 {
 	return ChangelistFileData.AssetName;
diff --git a/Engine/Plugins/Editor/ChangelistReview/Source/Private/SSourceControlReviewEntry.h b/Engine/Plugins/Editor/ChangelistReview/Source/Private/SSourceControlReviewEntry.h
--- a/Engine/Plugins/Editor/ChangelistReview/Source/Private/SSourceControlReviewEntry.h
+++ b/Engine/Plugins/Editor/ChangelistReview/Source/Private/SSourceControlReviewEntry.h
@@ -19,6 +19,7 @@ public:
 	
 	SLATE_BEGIN_ARGS(SSourceControlReviewEntry) {}
 		SLATE_ARGUMENT(FChangelistFileData, FileData);
+		SLATE_ARGUMENT(TWeakPtr<IReviewCommentAPI>, CommentsAPI);
 	SLATE_END_ARGS()
 	
 	/** Constructs the widget */
@@ -113,6 +114,14 @@ private:
 	 */
 	FText GetLocalAssetPathText() const;
 
+	const TArray<FReviewComment>* GetReviewComments() const;
+
+	FString GetReviewerUsername() const;
+
+	EVisibility GetUnreadCommentsIconVisibility() const;
+
+	FText GetUnreadCommentsTooltip() const;
+
 	/**
 	 * Returns asset name to be used for search/filter 
 	 */
@@ -126,6 +135,7 @@ private:
 private:
 	// File associated with this entry
 	FChangelistFileData ChangelistFileData;
+	TWeakPtr<IReviewCommentAPI> CommentsAPI;
 	
 	// Called to diff this file against it's previous revision
 	DECLARE_DELEGATE(FDiffMethod);
diff --git a/Engine/Plugins/Editor/ChangelistReview/Source/Public/ChangelistReviewModule.h b/Engine/Plugins/Editor/ChangelistReview/Source/Public/ChangelistReviewModule.h
--- a/Engine/Plugins/Editor/ChangelistReview/Source/Public/ChangelistReviewModule.h
+++ b/Engine/Plugins/Editor/ChangelistReview/Source/Public/ChangelistReviewModule.h
@@ -5,6 +5,7 @@
 //#include "SSourceControlReview.h"
 #include "Modules/ModuleManager.h"
 
+class SSourceControlReview;
 class FSpawnTabArgs;
 class SDockTab;
 class SWidget;
@@ -18,12 +19,18 @@ public:
 	virtual void ShutdownModule() override;
 	void ShowReviewTab();
 	bool CanShowReviewTab() const;
+	TWeakPtr<SSourceControlReview> GetActiveReview();
+	
+	static FChangelistReviewModule& Get()
+	{
+		return FModuleManager::LoadModuleChecked<FChangelistReviewModule>("ChangelistReview");
+	}
 private:
 	TSharedRef<SDockTab> CreateReviewTab(const FSpawnTabArgs& Args);
 	TSharedPtr<SWidget> CreateReviewUI();
 	
 	TWeakPtr<SDockTab> ReviewTab;
-	//TWeakPtr<SSourceControlReview> ReviewWidget;
+	TWeakPtr<SSourceControlReview> ReviewWidget;
 };
 
 #if UE_ENABLE_INCLUDE_ORDER_DEPRECATED_IN_5_2
diff --git a/Engine/Plugins/Editor/ChangelistReview/Source/Public/SSourceControlReview.h b/Engine/Plugins/Editor/ChangelistReview/Source/Public/SSourceControlReview.h
new file mode 100644
--- /dev/null
+++ b/Engine/Plugins/Editor/ChangelistReview/Source/Public/SSourceControlReview.h
@@ -0,0 +1,290 @@
+// Copyright Epic Games, Inc. All Rights Reserved.
+
+#pragma once
+#include "ISourceControlProvider.h"
+#include "Misc/PackagePath.h"
+#include "SourceControlOperations.h"
+#include "Widgets/Views/SHeaderRow.h"
+#include "FSwarmCommentsAPI.h"
+
+template <typename OptionType> class SComboBox;
+
+class SEditableText;
+class STableViewBase;
+class STextBlock;
+template <typename ItemType> class SListView;
+
+class SSourceControlReviewEntry;
+class SChangelistEditableTextBox;
+class SProgressBar;
+
+namespace SourceControlReview
+{
+	enum class ESourceControlAction : uint8
+	{
+		Add,
+		Edit,
+		Delete,
+		Branch,
+		Integrate,
+		Unset,
+
+		// keep this last
+		ActionCount
+	};
+
+	enum class EChangelistState : uint8
+	{
+		Submitted,
+		Pending
+	};
+
+	struct FChangelistFileData
+	{
+		FChangelistFileData() = default;
+
+		bool IsDataValidForEntry() const
+		{
+			if (FileSourceControlAction == ESourceControlAction::Unset)
+			{
+				return false;
+			}
+			if (ReviewFileTempPath.IsEmpty() && FileSourceControlAction != ESourceControlAction::Delete)
+			{
+				return false;
+			}
+			if (PreviousFileTempPath.IsEmpty() && FileSourceControlAction != ESourceControlAction::Add && FileSourceControlAction != ESourceControlAction::Branch)
+			{
+				return false;
+			}
+			
+			return true;
+		}
+
+		FString AssetName;
+
+		// path where the temporary asset was downloaded to
+		FString ReviewFileTempPath;
+
+		FString ReviewFileRevisionNum;
+
+		FDateTime ReviewFileDateTime;
+
+		FString PreviousAssetName;
+
+		// path where the temporary asset was downloaded to
+		FString PreviousFileTempPath;
+
+		FString PreviousFileRevisionNum;
+
+		FString RelativeFilePath;
+
+		FString AssetFilePath;
+		
+		FString AssetDepotPath;
+
+		int32 ChangelistNum = INDEX_NONE;
+
+		EChangelistState ChangelistState = EChangelistState::Submitted;
+
+		ESourceControlAction FileSourceControlAction = ESourceControlAction::Unset;
+
+		const UClass* GetIconClass();
+		
+	private:
+		TOptional<const UClass*> CachedIconClass;
+	};
+	
+	namespace ColumnIds
+	{
+		inline const FName Status = TEXT("Status");
+		inline const FName File = TEXT("File");
+		inline const FName Tools = TEXT("Tools");
+	}
+
+	struct FChangelistLightInfo
+	{
+		explicit FChangelistLightInfo(const FText& Number)
+			: Number(Number)
+		{}
+		FChangelistLightInfo(const FText& Number, const FText& Author, const FText& Description)
+			: Number(Number), Author(Author), Description(Description)
+		{}
+		
+		FText Number;
+		FText Author;
+		FText Description;
+	};
+}
+
+
+/**
+ * Used to select a changelist and diff it's changes
+ */
+class SSourceControlReview : public SCompoundWidget
+{
+public:
+	using ESourceControlAction = SourceControlReview::ESourceControlAction;
+	using EChangelistState = SourceControlReview::EChangelistState;
+	using FChangelistFileData = SourceControlReview::FChangelistFileData;
+	
+	SLATE_BEGIN_ARGS(SSourceControlReview) {}
+	SLATE_END_ARGS()
+	
+	/** Constructs the widget */
+	void Construct(const FArguments& InArgs);
+
+	virtual ~SSourceControlReview() override;
+
+	/**
+	 * Pulls up changelist record from source control
+	 * @param Changelist Changelist number to get
+	 */
+	void LoadChangelist(const FString& Changelist);
+
+	/**
+	 * Returns an array of review comments that are 
+	 * @param FilePath the file to retrieve comments from. (supports AssetDepotPaths and temp local file paths)
+	 */
+	const TArray<FReviewComment>* GetReviewCommentsForFile(const FString& FilePath);
+
+
+	void UpdateReviewComments();
+	void PostComment(FReviewComment& Comment);
+	void EditComment(FReviewComment& Comment);
+	FString GetReviewerUsername() const;
+	bool IsFileInReview(const FString& File) const;
+	
+private:
+
+	FString AsDepotPath(const FString& FilePath);
+	
+	/**
+	 * Called when FGetChangelistDetails completes
+	 */
+	void OnGetChangelistDetails(const FSourceControlOperationRef& InOperation, ECommandResult::Type InResult, FString Changelist);
+	
+	/**
+	 * Called when FSwarmCommentsAPI::GetReviewTopicForCL completes
+	 */
+	void OnGetReviewTopic(const FReviewTopic& Topic, const FString& ErrorMessage);
+	
+	/**
+	 * Called when FSwarmCommentsAPI::OnGetReviewComments completes
+	 */
+	void OnGetReviewComments(const TArray<FReviewComment>& Comments, const FString& ErrorMessage);
+	void OnInitReviewComments(const TArray<FReviewComment>& Comments, const FString& ErrorMessage);
+	
+	/**
+	 * Called when file is retrieved from source control
+	 * @param ChangelistFileData carried through worker to be filled with information about retrieved file from source control
+	 */
+	void OnGetFileFromSourceControl(TSharedPtr<FChangelistFileData> ChangelistFileData);
+	
+	/**
+	 * Called when changelist details, changelist files, and changelist review info has all been retrieved
+	 */
+	void OnLoadComplete();
+	
+	FReply OnLoadChangelistClicked();
+	bool IsSourceControlActive() const;
+	FText LoadChangelistTooltip() const;
+	void OnChangelistNumChanged(const FText& Text);
+	void OnChangelistNumCommitted(const FText& Text, ETextCommit::Type CommitMethod);
+
+	TSharedRef<SWidget> MakeCLComboOption(TSharedPtr<SourceControlReview::FChangelistLightInfo>Item) const;
+	void OnCLComboSelection(TSharedPtr<SourceControlReview::FChangelistLightInfo> Item, ESelectInfo::Type SelectInfo) const;
+	void SaveCLHistory();
+	void LoadCLHistory();
+	
+	/**
+	 * Sets bChangelistLoading and then fires the BP_SetLoading event
+	 * @param bInLoading when true flips visibility on loading bar and is also blocking from loading multiple changelists at the same time until loading is finished
+	 */
+	void SetLoading(bool bInLoading);
+	
+	/**
+	 * True if the tool is currently loading a changelist
+	 */
+	bool IsLoading() const;
+
+	/**
+	 * Checks changelist information to make sure it's valid for the diff tool
+	 * @param InRecord changelist record
+	 */
+	bool IsChangelistRecordValid(const TArray<TMap<FString, FString>>& InRecord) const;
+
+	/**
+	 * Sets the source control action type for the loaded file
+	 * @param ChangelistFileData carried through worker to be filled with information about retrieved file from source control
+	 * @param SourceControlAction source action operation type retrieved from source control
+	 */
+	static void SetFileSourceControlAction(TSharedPtr<FChangelistFileData> ChangelistFileData,
+	                                       const FString& SourceControlAction);
+
+	/**
+	 * Sets changelist information Author, Description, Status and Depot name that changelist is going to
+	 */
+	void SetChangelistInfo(const TMap<FString, FString>& InChangelistRecord, const FString& ChangelistNum);
+
+	TSharedRef<ITableRow> OnGenerateFileRow(TSharedPtr<FChangelistFileData> FileData, const TSharedRef<STableViewBase>& Table) const;
+
+	/**
+	 * Removes CurrentChangelistInfo.SharedPath from the beginning of FullCLPath
+	 */
+	FString TrimSharedPath(FString FullCLPath) const;
+	
+	void FixupRedirectors();
+
+	static SHeaderRow::FColumn::FArguments HeaderColumn(FName HeaderName);
+
+	/**
+	 * Removes the game directory from the beginning of the FullCLPath
+	 */
+	static FString AsAssetPath(const FString& FullCLPath);
+
+	// helper function that sets bUncommittedChangelistNum to false and upgrades CLHistory[0] from a entry to an official one
+	void CommitTempChangelistNumToHistory();
+	
+	// helper function that sets bUncommittedChangelistNum to false and removes CLHistory[0] if it's uncommitted
+	void RemoveUncommittedChangelistNumFromHistory();
+
+	// used for asynchronous changelist loading
+	bool bChangelistLoading = false;
+
+	void IncrementItemsLoaded();
+	uint32 NumItemsToLoad = 0; // counts the number of API calls (swarm, p4, etc) that need to be loaded.
+	uint32 NumItemsLoaded = 0; // counts the number of API calls (swarm, p4, etc) that have been loaded
+	TArray<TSharedPtr<FChangelistFileData>> ChangelistFiles;
+	TSharedPtr<FGetChangelistDetails> GetChangelistDetailsCommand;
+	TArray<TSharedPtr<SourceControlReview::FChangelistLightInfo>> CLHistory;
+	bool bUncommittedChangelistNum = false;
+	
+	TSharedPtr<SComboBox<TSharedPtr<SourceControlReview::FChangelistLightInfo>>> ChangelistNumComboBox;
+	TSharedPtr<SEditableText> ChangelistNumText;
+	TSharedPtr<STextBlock> EnterChangelistTextBlock;
+	TSharedPtr<STextBlock> LoadingTextBlock;
+	TSharedPtr<SProgressBar> LoadingProgressBar;
+	TSharedPtr<SWidget> ChangelistInfoWidget;
+	TSharedPtr<SListView<TSharedPtr<FChangelistFileData>>> ChangelistEntriesWidget;
+
+	// Info about the current chagnelist
+	struct FChangelistInfo
+	{
+		FText Author;
+		FText SharedPath;
+		FText Status;
+		FText Description;
+		FString ChangelistNum;
+	} CurrentChangelistInfo;
+
+	TSharedPtr<IReviewCommentAPI> CommentsAPI;
+	TOptional<FReviewTopic> ReviewTopic;
+	// review comments that aren't attached to a specific file
+	TArray<FReviewComment> GlobalReviewComments;
+	// review comments mapped by the file they're attached to (keyed by both repository path and temp local path)
+	TMap<FString, TArray<FReviewComment>> FileReviewComments;
+	bool bReviewCommentsDirty = false;
+
+	TMap<FString, FString> TempLocalPathToDepotPath;
+};
diff --git a/Engine/Source/Editor/EditorStyle/Private/StarshipStyle.cpp b/Engine/Source/Editor/EditorStyle/Private/StarshipStyle.cpp
--- a/Engine/Source/Editor/EditorStyle/Private/StarshipStyle.cpp
+++ b/Engine/Source/Editor/EditorStyle/Private/StarshipStyle.cpp
@@ -452,6 +452,8 @@ void FStarshipEditorStyle::FStyle::SetupGeneralStyles()
 		Set("Icons.Unpinned", new IMAGE_BRUSH_SVG("Starship/Common/Unpinned", Icon16x16));
 		Set("Icons.Tools", new IMAGE_BRUSH_SVG("Starship/Common/EditorModes", Icon16x16));
 		Set("Icons.Clipboard", new IMAGE_BRUSH_SVG("Starship/Common/Clipboard", Icon16x16));
+		Set("Icons.HollowHeart", new IMAGE_BRUSH_SVG("Starship/Common/HollowHeart", Icon16x16));
+		Set("Icons.Heart", new IMAGE_BRUSH_SVG("Starship/Common/Heart", Icon16x16));
 
 
 		Set("Icons.Toolbar.Play", new IMAGE_BRUSH_SVG("Starship/Common/play", Icon20x20));
diff --git a/Engine/Source/Editor/Kismet/Private/DiffControl.cpp b/Engine/Source/Editor/Kismet/Private/DiffControl.cpp
--- a/Engine/Source/Editor/Kismet/Private/DiffControl.cpp
+++ b/Engine/Source/Editor/Kismet/Private/DiffControl.cpp
@@ -7,9 +7,23 @@
 #include "SBlueprintDiff.h"
 #include "Kismet2/BlueprintEditorUtils.h"
 #include "IDetailsView.h"
+#include "ReviewComments.h"
+#include "Widgets/Input/SMultiLineEditableTextBox.h"
 
 #define LOCTEXT_NAMESPACE "SBlueprintDif"
 
+namespace UE::DiffControl
+{
+	KISMET_API const TArray<FReviewComment>*(*GGetReviewCommentsForFile)(const FString&) = nullptr;
+	KISMET_API void(*GPostReviewComment)(FReviewComment&) = nullptr;
+	KISMET_API void(*GEditReviewComment)(FReviewComment&) = nullptr;
+	KISMET_API FString (*GGetReviewerUsername)() = nullptr;
+	KISMET_API bool (*GIsFileInReview)(const FString& File) = nullptr;
+	
+	DECLARE_MULTICAST_DELEGATE_OneParam(FOnCommentPosted, const FReviewComment&)
+	KISMET_API FOnCommentPosted GOnCommentPosted;
+}
+
 // well known algorithm for computing the Longest Common Subsequence table of two ordered lists
 template <typename RangeType, typename ComparePredicate>
 static TArray<TArray<int32>> CalculateLCSTable(const RangeType& Range1, const RangeType& Range2, ComparePredicate Comparison)
@@ -60,6 +74,22 @@ TSharedRef<SWidget> IDiffControl::GenerateObjectDiffWidget(FSingleObjectDiffEntr
 		.ColorAndOpacity(DiffViewUtils::Differs());
 }
 
+void IDiffControl::EnableComments(TWeakPtr<STreeView<TSharedPtr<FBlueprintDifferenceTreeEntry>>> TreeView,
+	const UObject* OldObject, const UObject* NewObject)
+{
+	ReviewCommentsDiffControl = MakeShared<FReviewCommentsDiffControl>(OldObject, NewObject, TreeView);
+}
+
+void IDiffControl::GenerateCategoryCommentTreeEntries(TArray<TSharedPtr<FBlueprintDifferenceTreeEntry>>& OutChildrenEntries,
+                                                      TArray<TSharedPtr<FBlueprintDifferenceTreeEntry>>& OutRealDifferences, FString CategoryKey)
+{
+	if (ReviewCommentsDiffControl)
+	{
+		ReviewCommentsDiffControl->SetCategory(CategoryKey);
+		ReviewCommentsDiffControl->GenerateTreeEntries(OutChildrenEntries, OutRealDifferences);
+	}
+}
+
 
 /////////////////////////////////////////////////////////////////////////////
 /// FMyBlueprintDiffControl
@@ -149,6 +179,15 @@ void FMyBlueprintDiffControl::GenerateTreeEntries(TArray<TSharedPtr<FBlueprintDi
 		Children,
 		bHasDifferences
 	));
+
+	// add comments as children to this category
+	GenerateCategoryCommentTreeEntries(OutTreeEntries.Last()->Children, OutRealDifferences, TEXT("My Blueprint"));
+	
+}
+
+void FMyBlueprintDiffControl::EnableComments(TWeakPtr<STreeView<TSharedPtr<FBlueprintDifferenceTreeEntry>>> TreeView)
+{
+	IDiffControl::EnableComments(TreeView, OldBlueprint, NewBlueprint);
 }
 
 
@@ -213,6 +252,14 @@ void FSCSDiffControl::GenerateTreeEntries(TArray<TSharedPtr<FBlueprintDifference
 		Children,
 		bHasDifferences
 	));
+
+	// add comments as children to this category
+	GenerateCategoryCommentTreeEntries(OutTreeEntries.Last()->Children, OutRealDifferences, TEXT("Components"));
+}
+
+void FSCSDiffControl::EnableComments(TWeakPtr<STreeView<TSharedPtr<FBlueprintDifferenceTreeEntry>>> TreeView)
+{
+	IDiffControl::EnableComments(TreeView, OldSCS.GetBlueprint(), NewSCS.GetBlueprint());
 }
 
 FDetailsDiffControl::FDetailsDiffControl(const UObject* InOldObject, const UObject* InNewObject,
@@ -248,6 +295,16 @@ FDetailsDiffControl::FDetailsDiffControl(const UObject* InOldObject, const UObje
 }
 
 void FDetailsDiffControl::GenerateTreeEntries(TArray<TSharedPtr<FBlueprintDifferenceTreeEntry>>& OutTreeEntries, TArray<TSharedPtr<FBlueprintDifferenceTreeEntry>>& OutRealDifferences)
+{
+	GenerateTreeEntriesWithoutComments(OutTreeEntries, OutRealDifferences);
+	
+	// add comments
+	GenerateCategoryCommentTreeEntries(OutTreeEntries, OutRealDifferences, TEXT("Details"));
+}
+
+void FDetailsDiffControl::GenerateTreeEntriesWithoutComments(
+	TArray<TSharedPtr<FBlueprintDifferenceTreeEntry>>& OutTreeEntries,
+	TArray<TSharedPtr<FBlueprintDifferenceTreeEntry>>& OutRealDifferences)
 {
 	for (const FSingleObjectDiffEntry& Difference : DifferingProperties)
 	{
@@ -263,6 +320,11 @@ void FDetailsDiffControl::GenerateTreeEntries(TArray<TSharedPtr<FBlueprintDiffer
 	}
 }
 
+void FDetailsDiffControl::EnableComments(TWeakPtr<STreeView<TSharedPtr<FBlueprintDifferenceTreeEntry>>> TreeView)
+{
+	IDiffControl::EnableComments(TreeView, OldDetails.GetDisplayedObject(), NewDetails.GetDisplayedObject());
+}
+
 void FDetailsDiffControl::OnSelectDiffEntry(FPropertySoftPath PropertyName)
 {
 	SelectionCallback.ExecuteIfBound();
@@ -344,7 +406,7 @@ FCDODiffControl::FCDODiffControl(const UObject* InOldObject, const UObject* InNe
 
 void FCDODiffControl::GenerateTreeEntries(TArray<TSharedPtr<FBlueprintDifferenceTreeEntry>>& OutTreeEntries, TArray<TSharedPtr<FBlueprintDifferenceTreeEntry>>& OutRealDifferences)
 {
-	FDetailsDiffControl::GenerateTreeEntries(OutTreeEntries, OutRealDifferences);
+	FDetailsDiffControl::GenerateTreeEntriesWithoutComments(OutTreeEntries, OutRealDifferences);
 
 	const bool bHasDifferences = Children.Num() != 0;
 	if (!bHasDifferences)
@@ -362,6 +424,10 @@ void FCDODiffControl::GenerateTreeEntries(TArray<TSharedPtr<FBlueprintDifference
 		Children,
 		bHasDifferences
 	));
+
+	// add comments as children to this category
+	GenerateCategoryCommentTreeEntries(OutTreeEntries.Last()->Children, OutRealDifferences, TEXT("Defaults"));
+	Children = OutTreeEntries.Last()->Children;
 }
 
 /////////////////////////////////////////////////////////////////////////////
@@ -374,7 +440,7 @@ FClassSettingsDiffControl::FClassSettingsDiffControl(const UObject* InOldObject,
 
 void FClassSettingsDiffControl::GenerateTreeEntries(TArray<TSharedPtr<FBlueprintDifferenceTreeEntry>>& OutTreeEntries, TArray<TSharedPtr<FBlueprintDifferenceTreeEntry>>& OutRealDifferences)
 {
-	FDetailsDiffControl::GenerateTreeEntries(OutTreeEntries, OutRealDifferences);
+	FDetailsDiffControl::GenerateTreeEntriesWithoutComments(OutTreeEntries, OutRealDifferences);
 
 	const bool bHasDifferences = Children.Num() != 0;
 	if (!bHasDifferences)
@@ -391,7 +457,11 @@ void FClassSettingsDiffControl::GenerateTreeEntries(TArray<TSharedPtr<FBlueprint
 		SelectionCallback,
 		Children,
 		bHasDifferences
-	));
+		));
+
+	// add comments as children to this category
+	GenerateCategoryCommentTreeEntries(OutTreeEntries.Last()->Children, OutRealDifferences, TEXT("Class Settings"));
+	Children = OutTreeEntries.Last()->Children;
 }
 
 
@@ -449,6 +519,10 @@ void FBlueprintTypeDiffControl::GenerateTreeEntries(TArray< TSharedPtr<FBlueprin
 			Children);
 	OutTreeEntries.Push(CategoryEntry);
 
+	// add comments as children to this category
+	const FString CommentCategoryKey = (BlueprintNew ? BlueprintNew : BlueprintOld)->GetClass()->GetName();
+	GenerateCategoryCommentTreeEntries(OutTreeEntries.Last()->Children, OutRealDifferences, CommentCategoryKey);
+
 	// Now add subobject diffs, one category per object
 	for (const TSharedPtr<FSubObjectDiff>& SubObjectDiff : SubObjectDiffs)
 	{
@@ -463,11 +537,18 @@ void FBlueprintTypeDiffControl::GenerateTreeEntries(TArray< TSharedPtr<FBlueprin
 			FOnDiffEntryFocused::CreateSP(AsShared(), &FBlueprintTypeDiffControl::OnSelectSubobjectDiff, FPropertySoftPath(), SubObjectDiff),
 			Children,
 			true);
-
 		OutTreeEntries.Push(SubObjectEntry);
+
+		// add comments as children to this category
+		GenerateCategoryCommentTreeEntries(OutTreeEntries.Last()->Children, OutRealDifferences, SubObjectDiff->SourceResult.OwningObjectPath);
 	}
 }
 
+void FBlueprintTypeDiffControl::EnableComments(TWeakPtr<STreeView<TSharedPtr<FBlueprintDifferenceTreeEntry>>> TreeView)
+{
+	IDiffControl::EnableComments(TreeView, BlueprintOld, BlueprintNew);
+}
+
 TSharedRef<SWidget> FBlueprintTypeDiffControl::GenerateCategoryWidget(bool bHasRealDiffs)
 {
 	FLinearColor Color = FLinearColor::White;
@@ -492,6 +573,7 @@ void FBlueprintTypeDiffControl::BuildDiffSourceArray()
 {
 	if (!BlueprintNew || !BlueprintOld)
 	{
+		bDiffSucceeded = true;
 		return;
 	}
 	
@@ -607,6 +689,22 @@ void FGraphToDiff::GenerateTreeEntries(TArray< TSharedPtr<FBlueprintDifferenceTr
 			FGenerateDiffEntryWidget::CreateSP(AsShared(), &FGraphToDiff::GenerateCategoryWidget),
 			Children);
 	OutTreeEntries.Push(Entry);
+	
+	// add comments as children to this category
+	FString GraphName;
+	const UEdGraph* Graph = GraphOld ? GraphOld : GraphNew;
+	if (const UEdGraphSchema* Schema = Graph->GetSchema())
+	{
+		FGraphDisplayInfo DisplayInfo;
+		Schema->GetGraphDisplayInformation(*Graph, DisplayInfo);
+
+		GraphName = DisplayInfo.DisplayName.ToString();
+	}
+	else
+	{
+		GraphName = Graph->GetFName().ToString();
+	}
+	GenerateCategoryCommentTreeEntries(OutTreeEntries.Last()->Children, OutRealDifferences, GraphName);
 }
 
 UEdGraph* FGraphToDiff::GetGraphOld() const
@@ -619,6 +717,11 @@ UEdGraph* FGraphToDiff::GetGraphNew() const
 	return GraphNew;
 }
 
+void FGraphToDiff::EnableComments(TWeakPtr<STreeView<TSharedPtr<FBlueprintDifferenceTreeEntry>>> TreeView)
+{
+	IDiffControl::EnableComments(TreeView, GetGraphOld(), GetGraphNew());
+}
+
 FText FGraphToDiff::GetToolTip()
 {
 	if (GraphOld && GraphNew)
@@ -706,4 +809,761 @@ void FGraphToDiff::OnGraphChanged( const FEdGraphEditAction& Action )
 	DiffWidget->OnGraphChanged(this);
 }
 
-#undef LOCTEXT_NAMESPACE
\ No newline at end of file
+// returns whether a file being diffed should display comments
+static bool AreCommentsEnabled(FString File)
+{
+	if (UE::DiffControl::GIsFileInReview)
+	{
+		return UE::DiffControl::GIsFileInReview(File);
+	}
+	return false;
+}
+
+FCommentTreeEntry::FCommentTreeEntry(TWeakPtr<FReviewCommentsDiffControl> InCommentsControl, const FReviewComment& InComment, const TArray<TSharedPtr<FBlueprintDifferenceTreeEntry>>& InChildren)
+	: FBlueprintDifferenceTreeEntry({}, FGenerateDiffEntryWidget::CreateRaw(this, &FCommentTreeEntry::CreateWidget), InChildren)
+	, Comment(InComment)
+	, CommentsControl(InCommentsControl)
+{}
+
+FCommentTreeEntry::~FCommentTreeEntry()
+{
+	if (OnCommentPostedHandle.IsValid())
+	{
+		UE::DiffControl::GOnCommentPosted.Remove(OnCommentPostedHandle);
+	}
+}
+
+TSharedRef<FCommentTreeEntry> FCommentTreeEntry::Make(TWeakPtr<FReviewCommentsDiffControl> CommentsControl, const FReviewComment& Comment,
+                                                      const TArray<TSharedPtr<FBlueprintDifferenceTreeEntry>>& Children)
+{
+	return MakeShared<FCommentTreeEntry>(CommentsControl, Comment, Children);
+}
+
+int32 FCommentTreeEntry::GetCommentIDChecked() const
+{
+	return *Comment.CommentID;
+}
+
+void FCommentTreeEntry::AwaitCommentPost()
+{
+	OnCommentPostedHandle = UE::DiffControl::GOnCommentPosted.AddSP(this, &FCommentTreeEntry::OnCommentPosted);
+}
+
+TSharedRef<SWidget> FCommentTreeEntry::CreateWidget()
+{
+	// if a widget for this comment is being created that means it's in view. Mark the ReadBy state to reflect that the user is reading it
+	if (!Comment.ReadBy.IsSet())
+	{
+		Comment.ReadBy = TSet<FString>();
+	}
+	const FString Username = UE::DiffControl::GGetReviewerUsername();
+	if (!Comment.ReadBy->Find(Username))
+	{
+		Comment.ReadBy->Add(Username);
+		
+		FReviewComment ReadByEdit;
+		ReadByEdit.CommentID = Comment.CommentID;
+		ReadByEdit.ReadBy = Comment.ReadBy;
+		UE::DiffControl::GEditReviewComment(ReadByEdit);
+	}
+	
+	FDateTime DateTime = FDateTime::UtcNow();
+	if (Comment.EditedTime.IsSet())
+	{
+		DateTime = *Comment.EditedTime;
+	}
+	else if (Comment.CreatedTime.IsSet())
+	{
+		DateTime = *Comment.CreatedTime;
+	}
+	FText DateText;
+	if (DateTime.GetDay() == FDateTime::UtcNow().GetDay())
+	{
+		DateText = FText::AsTime(DateTime, EDateTimeStyle::Short);
+	}
+	else
+	{
+		DateText = FText::AsDate(DateTime, EDateTimeStyle::Short);
+	}
+
+	// constructs a hyperlink-like text button
+	auto MakeCommentOptionButton = [](const FText &Text)
+	{
+		const FLinearColor HoveredColor = FStyleColors::AccentBlue.GetSpecifiedColor();
+		const FLinearColor UnhoveredColor = HoveredColor.Desaturate(0.2f);
+		TSharedRef<SButton> Button = SNew(SButton)
+			.ButtonStyle(FAppStyle::Get(), "NoBorder")
+			.ForegroundColor(UnhoveredColor)
+			.Text(Text);
+
+		Button->SetOnHovered(FSimpleDelegate::CreateLambda([ButtonWeak = Button.ToWeakPtr(), HoveredColor]()
+		{
+			if (const TSharedPtr<SButton> Button = ButtonWeak.Pin())
+			{
+				Button->SetForegroundColor(HoveredColor);
+			}
+		}));
+		
+		Button->SetOnUnhovered(FSimpleDelegate::CreateLambda([ButtonWeak = Button.ToWeakPtr(), UnhoveredColor]()
+		{
+			if (const TSharedPtr<SButton> Button = ButtonWeak.Pin())
+			{
+				Button->SetForegroundColor(UnhoveredColor);
+			}
+		}));
+		return Button;
+	};
+
+	const TSharedPtr<SButton> ReplyButton = MakeCommentOptionButton(LOCTEXT("ReplyToComment","Reply"));
+	ReplyButton->SetOnClicked(FOnClicked::CreateSP(this, &FCommentTreeEntry::OnClickReply));
+	
+	const TSharedPtr<SButton> EditButton = MakeCommentOptionButton(LOCTEXT("EditComment","Edit"));
+	EditButton->SetOnClicked(FOnClicked::CreateSP(this, &FCommentTreeEntry::OnClickEdit));
+	EditButton->SetVisibility(TAttribute<EVisibility>::CreateSP(this, &FCommentTreeEntry::GetEditButtonVisibility));
+
+	FLinearColor BorderColor = FStyleColors::Recessed.GetSpecifiedColor();
+	BorderColor.A = 0.95f;
+
+	return SAssignNew(Content, SBorder)
+	.IsEnabled_Static(&AreCommentsEnabled, *Comment.Context.File)
+	.BorderImage(FCoreStyle::Get().GetBrush("GenericWhiteBox"))
+	.BorderBackgroundColor(BorderColor)
+	[
+		SNew(SVerticalBox)
+		+SVerticalBox::Slot()
+		.VAlign(VAlign_Top)
+		.AutoHeight()
+		[
+			SNew(SHorizontalBox)
+			+SHorizontalBox::Slot()
+			.AutoWidth()
+			[
+				SNew(STextBlock)
+				.ColorAndOpacity(this, &FCommentTreeEntry::GetUsernameColor)
+				.Text(FText::FromString(Comment.User.Get(TEXT("[Unknown User]"))))
+				.Font(FAppStyle::GetFontStyle(TEXT("BoldFont")))
+			]
+			+SHorizontalBox::Slot()
+			.HAlign(HAlign_Right)
+			.FillWidth(1)
+			.Padding(10.f,0.f,0.f,0.f)
+			[
+				SNew(STextBlock)
+					.ColorAndOpacity(FSlateColor::UseSubduedForeground())
+					.Text(DateText)
+			]
+		]
+		+SVerticalBox::Slot()
+		.VAlign(VAlign_Top)
+		.AutoHeight()
+		[
+			SAssignNew(CommentTextBox, SMultiLineEditableTextBox)
+			.Text(FText::FromString(GetCommentString()))
+			.AutoWrapText(true)
+			.IsReadOnly(this, &FCommentTreeEntry::IsCommentTextBoxReadonly)
+		]
+		+SVerticalBox::Slot()
+		.VAlign(VAlign_Top)
+		.AutoHeight()
+		[
+			SAssignNew(EditReplyButtonGroup, SHorizontalBox)
+			.Visibility(this, &FCommentTreeEntry::GetEditReplyButtonGroupVisibility)
+			+SHorizontalBox::Slot()
+			.AutoWidth()
+			[
+				ReplyButton.ToSharedRef()
+			]
+			+SHorizontalBox::Slot()
+			.AutoWidth()
+			[
+				EditButton.ToSharedRef()
+			]
+			+SHorizontalBox::Slot()
+			.AutoWidth()
+			[
+				SNew(SButton)
+            	.ButtonStyle(FAppStyle::Get(), "SimpleButton")
+            	.OnClicked(this, &FCommentTreeEntry::OnLikeToggle)
+            	.ToolTipText(this, &FCommentTreeEntry::GetLikeTooltip)
+            	.ContentScale(0.8f)
+            	.HAlign(HAlign_Center)
+            	.VAlign(VAlign_Center)
+            	[
+					SNew(SImage)
+					.Image(this, &FCommentTreeEntry::GetLikeIcon)
+					.ColorAndOpacity(this, &FCommentTreeEntry::GetLikeIconColor)
+				]
+				
+			]
+		]
+		+SVerticalBox::Slot()
+		.VAlign(VAlign_Top)
+		.AutoHeight()
+		[
+			SAssignNew(SubmitCancelButtonGroup, SHorizontalBox)
+			.Visibility(this, &FCommentTreeEntry::GetSubmitCancelButtonGroupVisibility)
+			+SHorizontalBox::Slot()
+			.AutoWidth()
+			[
+				SNew(SButton)
+				.Text(LOCTEXT("SubmitEdit","Submit"))
+				.IsEnabled(this, &FCommentTreeEntry::IsSubmitButtonEnabled)
+				.OnClicked(this, &FCommentTreeEntry::OnEditSubmitClicked)
+			]
+			+SHorizontalBox::Slot()
+			.AutoWidth()
+			[
+				SNew(SButton)
+				.Text(LOCTEXT("CancelEdit","Cancel"))
+				.OnClicked(this, &FCommentTreeEntry::OnEditCancelClicked)
+			]
+		]
+	];
+}
+
+bool FCommentTreeEntry::IsCommentTextBoxReadonly() const
+{
+	return !IsEditMode();
+}
+
+FReply FCommentTreeEntry::OnClickReply()
+{
+	if (const TSharedPtr<FReviewCommentsDiffControl> Control = CommentsControl.Pin())
+	{
+		Control->DraftReply(AsShared());
+		return FReply::Handled();
+	}
+	return FReply::Unhandled();
+}
+
+FReply FCommentTreeEntry::OnClickEdit()
+{
+	SetEditMode(true);
+	return FReply::Handled();
+}
+
+FReply FCommentTreeEntry::OnLikeToggle()
+{
+	if (!Comment.Likes.IsSet())
+	{
+		Comment.Likes = TSet<FString>{};
+	}
+	const FString Username = UE::DiffControl::GGetReviewerUsername();
+	if (Comment.Likes->Contains(Username))
+	{
+		Comment.Likes->Remove(Username);
+	}
+	else
+	{
+		Comment.Likes->Add(Username);
+	}
+	
+	FReviewComment LikeEdit;
+	LikeEdit.CommentID = Comment.CommentID;
+	LikeEdit.Likes = Comment.Likes;
+    UE::DiffControl::GEditReviewComment(LikeEdit);
+	
+	return FReply::Handled();
+}
+
+// convert an iterable structure to a text based list
+// for example,
+// {"a"} -> "{Range[0]}" -> "a"
+// {"a", "b"} -> "{Range[0]} {Conjunction} {Range[1]}" -> "a and b"
+// {"a", "b", "c"} -> "{Range[0]}{Delimiter} {Range[1]}{Delimiter} {Conjunction} {Range[2]}" -> "a, b, and c"
+template <typename RangeType, typename PredicateType>
+static FText RangeAsText(RangeType&& Range, const FText& Delimiter, const FText& Conjunction, PredicateType Transformer)
+{
+	TArray<FText> ItemsAsText;
+	int32 Index = 0;
+	const int32 LastIndex = Range.Num() - 1;
+	if (Range.Num() == 1)
+	{
+		return Invoke(Transformer, *Range.begin());
+	}
+	for (auto&& Elem : Forward<RangeType>(Range))
+	{
+		FText ElementText = Invoke(Transformer, Elem);
+		if (Index == LastIndex and !Conjunction.IsEmpty())
+		{
+			// for the last item, prepend a conjunction
+			ItemsAsText.Add(FText::FormatNamed(
+				LOCTEXT("ListLastElement", "{Conjunction} {Element}"),
+				TEXT("Conjunction"), Conjunction,
+				TEXT("Element"), Invoke(Transformer, Elem)
+			));
+		}
+		else
+		{
+			ItemsAsText.Add(Invoke(Transformer, Elem));
+		}
+		++Index;
+	}
+	if (Range.Num() == 2)
+	{
+		// for two elements, don't add a comma
+		return FText::Join(LOCTEXT("Space"," "), ItemsAsText);
+	}
+	return FText::Join(FText::Format(LOCTEXT("DelimeterWithSpace","{0} "), Delimiter), ItemsAsText);
+}
+
+FText FCommentTreeEntry::GetLikeTooltip() const
+{
+	if (Comment.Likes.IsSet() && !Comment.Likes->IsEmpty())
+	{
+		FText(* const StringToText)(const FString&) = &FText::FromString;
+		const FText LikesListText = RangeAsText(*Comment.Likes, LOCTEXT("ListDelimeter",","), LOCTEXT("ConjunctionAnd","and"), StringToText);
+		return FText::Format(LOCTEXT("LikedBy", "Liked By: {0}"), LikesListText);
+	}
+	return LOCTEXT("NoLikes","No Likes");
+}
+
+const FSlateBrush* FCommentTreeEntry::GetLikeIcon() const
+{
+	if (Comment.Likes.IsSet() && !Comment.Likes->IsEmpty())
+	{
+		return FAppStyle::Get().GetBrush(TEXT("Icons.Heart"));
+	}
+	return FAppStyle::Get().GetBrush(TEXT("Icons.HollowHeart"));
+}
+
+FSlateColor FCommentTreeEntry::GetUsernameColor() const
+{
+	const FString Username = UE::DiffControl::GGetReviewerUsername();
+	if (Comment.User == Username)
+	{
+		return FStyleColors::AccentBlue;
+	}
+	return FStyleColors::AccentGreen.GetSpecifiedColor().Desaturate(0.2f);
+}
+
+FSlateColor FCommentTreeEntry::GetLikeIconColor() const
+{
+	const FString Username = UE::DiffControl::GGetReviewerUsername();
+	if (Comment.Likes.IsSet())
+	{
+		if (Comment.Likes->Contains(Username))
+		{
+			return FSlateColor::UseForeground();
+		}
+	}
+	return FSlateColor::UseSubduedForeground();
+}
+
+EVisibility FCommentTreeEntry::GetEditReplyButtonGroupVisibility() const
+{
+	if (IsEditMode())
+	{
+		return EVisibility::Collapsed;
+	}
+	// if the comment is still being submitted, don't allow edits or replies yet
+	if (!Comment.CommentID.IsSet())
+	{
+		return EVisibility::Hidden;
+	}
+	return Content->IsHovered() ? EVisibility::Visible : EVisibility::Hidden;
+}
+
+EVisibility FCommentTreeEntry::GetEditButtonVisibility() const
+{
+	const FString Reviewer = UE::DiffControl::GGetReviewerUsername();
+	if (Reviewer == Comment.User)
+	{
+		return EVisibility::Visible;
+	}
+	return EVisibility::Collapsed;
+}
+
+EVisibility FCommentTreeEntry::GetSubmitCancelButtonGroupVisibility() const
+{
+	return IsEditMode() ? EVisibility::Visible : EVisibility::Collapsed;
+}
+
+bool FCommentTreeEntry::IsSubmitButtonEnabled() const
+{
+	// don't allow submission of empty strings
+	return !CommentTextBox->GetText().ToString().TrimEnd().IsEmpty();
+}
+
+FReply FCommentTreeEntry::OnEditSubmitClicked()
+{
+	if (HasCommentStringChanged())
+	{
+		SetCommentString(CommentTextBox->GetText().ToString());
+		FReviewComment BodyEdit;
+		BodyEdit.CommentID = Comment.CommentID;
+		BodyEdit.Body = Comment.Body;
+		UE::DiffControl::GEditReviewComment(BodyEdit);
+	}
+	CommentTextBox->SetText(FText::FromString(GetCommentString()));
+	SetEditMode(false);
+	return FReply::Handled();
+}
+
+FReply FCommentTreeEntry::OnEditCancelClicked()
+{
+	CommentTextBox->SetText(FText::FromString(GetCommentString()));
+	SetEditMode(false);
+	return FReply::Handled();
+}
+
+void FCommentTreeEntry::OnCommentPosted(const FReviewComment& InComment)
+{
+	if (InComment.Body != Comment.Body)
+	{
+		return;
+	}
+	if (InComment.User == Comment.User)
+	{
+		Comment = InComment;
+		UE::DiffControl::GOnCommentPosted.Remove(OnCommentPostedHandle);
+	}
+}
+
+FString FCommentTreeEntry::GetCommentString() const
+{
+	return Comment.Body.Get(TEXT("[Missing Comment Body]")).TrimEnd();
+}
+
+void FCommentTreeEntry::SetCommentString(const FString& NewComment)
+{
+	Comment.Body = NewComment.TrimEnd();
+}
+
+bool FCommentTreeEntry::HasCommentStringChanged() const
+{
+	const FString EditedCommentString = CommentTextBox->GetText().ToString().TrimEnd();
+	const FString CommentString = GetCommentString();
+	return EditedCommentString != CommentString;
+}
+
+bool FCommentTreeEntry::IsEditMode() const
+{
+	if (bExpectedEditMode)
+	{
+		// if user de-focused the textbox without changing anything, cancel the edit
+		if (!CommentTextBox->HasKeyboardFocus() && !HasCommentStringChanged())
+		{
+			CommentTextBox->SetText(FText::FromString(GetCommentString()));
+			return false;
+		}
+	}
+	return bExpectedEditMode;
+}
+
+void FCommentTreeEntry::SetEditMode(bool bIsEditMode)
+{
+	if (bIsEditMode)
+	{
+		FSlateApplication::Get().SetKeyboardFocus(CommentTextBox);
+	}
+	bExpectedEditMode = bIsEditMode;
+}
+
+FCommentDraftTreeEntry::FCommentDraftTreeEntry(TWeakPtr<FReviewCommentsDiffControl> InCommentsControl, TArray<TSharedPtr<FBlueprintDifferenceTreeEntry>>* InSiblings, int32 InReplyID)
+	: FBlueprintDifferenceTreeEntry({}, FGenerateDiffEntryWidget::CreateRaw(this, &FCommentDraftTreeEntry::CreateWidget), {})
+	, CommentsControl(InCommentsControl)
+	, Siblings(InSiblings)
+	, ReplyID(InReplyID)
+{}
+
+TSharedRef<FCommentDraftTreeEntry> FCommentDraftTreeEntry::MakeCommentDraft(TWeakPtr<FReviewCommentsDiffControl> CommentsControl, TArray<TSharedPtr<FBlueprintDifferenceTreeEntry>>* Siblings)
+{
+	return MakeShared<FCommentDraftTreeEntry>(CommentsControl, Siblings);
+}
+
+TSharedRef<FCommentDraftTreeEntry> FCommentDraftTreeEntry::MakeReplyDraft(TWeakPtr<FReviewCommentsDiffControl> CommentsControl, TSharedPtr<FCommentTreeEntry> InParent)
+{
+	return MakeShared<FCommentDraftTreeEntry>(CommentsControl, &InParent->Children, InParent->GetCommentIDChecked());
+}
+
+void FCommentDraftTreeEntry::ReassignSiblings(TArray<TSharedPtr<FBlueprintDifferenceTreeEntry>>* InSiblings, int32 InReplyID)
+{
+	const int32 MyIndex = Siblings->Find(AsShared());
+	if (MyIndex != INDEX_NONE)
+	{
+		Siblings->RemoveAt(MyIndex);
+	}
+	Siblings = InSiblings;
+	ReplyID = InReplyID;
+}
+
+void FCommentDraftTreeEntry::ReassignReplyParent(TSharedPtr<FCommentTreeEntry> InParent)
+{
+	ReassignSiblings(&InParent->Children, InParent->GetCommentIDChecked());
+}
+
+bool FCommentDraftTreeEntry::IsReply() const
+{
+	return ReplyID != -1;
+}
+
+TSharedRef<SWidget> FCommentDraftTreeEntry::CreateWidget()
+{
+	FLinearColor BorderColor = FStyleColors::Recessed.GetSpecifiedColor();
+	BorderColor.A = 0.95f;
+
+	FString FilePath;
+	if (const TSharedPtr<FReviewCommentsDiffControl> Control = CommentsControl.Pin())
+	{
+		FilePath = Control->GetCommentFilePath();
+	}
+
+	SAssignNew(Content, SBorder)
+	.IsEnabled_Static(&AreCommentsEnabled, FilePath)
+	.BorderImage(FCoreStyle::Get().GetBrush("GenericWhiteBox"))
+	.BorderBackgroundColor(BorderColor)
+	[
+		SNew(SVerticalBox)
+		+SVerticalBox::Slot()
+		.VAlign(VAlign_Top)
+		.AutoHeight()
+		[
+			SNew(SBox)
+			.MinDesiredHeight(50.f)
+			[
+				SAssignNew(CommentTextBox, SMultiLineEditableTextBox)
+				.HintText(IsReply()? LOCTEXT("AddReply", "Reply...") : LOCTEXT("AddAComment", "Add a comment"))
+				.ForegroundColor(FStyleColors::White)
+				.AutoWrapText(true)
+			]
+		]
+		+SVerticalBox::Slot()
+		.VAlign(VAlign_Top)
+		.AutoHeight()
+		.Padding(0.f, 0.f, 0.f, 20.f)
+		[
+			SNew(SHorizontalBox)
+			+SHorizontalBox::Slot()
+			.VAlign(VAlign_Center)
+			.AutoWidth()
+			.Padding(0.f,0.f,10.f,0.f)
+			[
+				SNew(SButton)
+				.Text(LOCTEXT("PostComment","Post"))
+				.IsEnabled(this, &FCommentDraftTreeEntry::IsPostButtonEnabled)
+				.OnClicked(this, &FCommentDraftTreeEntry::OnCommentPostClicked)
+			]
+			+SHorizontalBox::Slot()
+			.VAlign(VAlign_Center)
+			.AutoWidth()
+			.Padding(0.f,0.f,10.f,0.f)
+			[
+				SAssignNew(FlagAsTaskCheckBox, SCheckBox)
+			]
+			+SHorizontalBox::Slot()
+			.VAlign(VAlign_Center)
+			.AutoWidth()
+			[
+				SNew(STextBlock)
+				.Text(LOCTEXT("FlagCommentAsTask","Flag as Task"))
+			]
+		]
+	];
+	
+	return Content.ToSharedRef();
+}
+
+bool FCommentDraftTreeEntry::IsPostButtonEnabled() const
+{
+	// only allow posting if the comment has non-whitespace characters
+	return !CommentTextBox->GetText().ToString().TrimEnd().IsEmpty();
+}
+
+FReply FCommentDraftTreeEntry::OnCommentPostClicked()
+{
+	if (const TSharedPtr<FReviewCommentsDiffControl> Control = CommentsControl.Pin())
+	{
+		const FString& CommentCategory = Control->GetCommentCategory();
+		const FString& CommentFilePath = Control->GetCommentFilePath();
+		
+		FReviewComment Comment;
+		Comment.Body = CommentTextBox->GetText().ToString().TrimEnd();
+		Comment.Context.File = CommentFilePath;
+		Comment.ReadBy = TSet<FString>{UE::DiffControl::GGetReviewerUsername()};
+		
+		if (!CommentCategory.IsEmpty())
+		{
+			Comment.Context.Category = CommentCategory;
+		}
+		if (IsReply())
+		{
+			Comment.Context.ReplyTo = ReplyID;
+		}
+		
+		const bool bIsTask = FlagAsTaskCheckBox->GetCheckedState() == ECheckBoxState::Checked;
+		Comment.TaskState = bIsTask ? EReviewCommentTaskState::Open : EReviewCommentTaskState::Comment;
+
+		const int32 InsertIndex = Siblings->Find(AsShared());
+		check(InsertIndex != INDEX_NONE);
+
+		Control->PostComment(Comment);
+		const TSharedRef<FCommentTreeEntry> CommentEntry = FCommentTreeEntry::Make(Control, Comment);
+		CommentEntry->AwaitCommentPost(); // tell entry to update once comment is posted
+		Siblings->Insert(CommentEntry, InsertIndex);
+		
+		return FReply::Handled();
+	}
+	return FReply::Unhandled();
+}
+
+FReviewCommentsDiffControl::FReviewCommentsDiffControl(const UObject* InOldObject, const UObject* InNewObject,
+                                                       TWeakPtr<STreeView<TSharedPtr<FBlueprintDifferenceTreeEntry>>> TreeView)
+	: OldObject(InOldObject)
+	, NewObject(InNewObject)
+	, CommentsTreeView(TreeView)
+{
+	const UPackage* Package = NewObject ? NewObject->GetPackage() : OldObject->GetPackage();
+	const FPackagePath PackagePath = FPackagePath::FromPackageNameChecked(Package->GetName());
+	CommentFilePath = PackagePath.GetLocalFullPath();
+}
+
+void FReviewCommentsDiffControl::GenerateCommentThreadRecursive(const FReviewComment& Comment,
+	const TMap<int32, TArray<const FReviewComment*>>& CommentReplyMap,
+	TArray<TSharedPtr<FBlueprintDifferenceTreeEntry>>& OutTreeEntries)
+{
+	if (Comment.bIsClosed)
+	{
+		return;
+	}
+	
+	TArray<TSharedPtr<FBlueprintDifferenceTreeEntry>> Children;
+	if (const TArray<const FReviewComment*>* Replies = CommentReplyMap.Find(*Comment.CommentID))
+	{
+		for (const FReviewComment* Reply : *Replies)
+		{
+			GenerateCommentThreadRecursive(*Reply, CommentReplyMap, Children);
+		}
+	}
+	
+	OutTreeEntries.Push(FCommentTreeEntry::Make(AsShared(), Comment, Children));
+}
+
+void FReviewCommentsDiffControl::GenerateTreeEntries(TArray<TSharedPtr<FBlueprintDifferenceTreeEntry>>& OutTreeEntries,
+                                                     TArray<TSharedPtr<FBlueprintDifferenceTreeEntry>>& OutRealDifferences)
+{
+	if (!AreCommentsEnabled(GetCommentFilePath()))
+	{
+		return;
+	}
+
+
+	TMap<int32, TArray<const FReviewComment*>> CommentReplyMap;
+	TArray<const FReviewComment*> HeadComments;
+	if (const TArray<FReviewComment>* FileComments = UE::DiffControl::GGetReviewCommentsForFile(CommentFilePath))
+	{
+		for (const FReviewComment& Comment : *FileComments)
+		{
+			if (Comment.Context.ReplyTo.IsSet())
+			{
+				CommentReplyMap.FindOrAdd(*Comment.Context.ReplyTo, {}).Add(&Comment);
+				continue;
+			}
+			if (Comment.Context.Category.Get({}) != CommentCategory)
+			{
+				continue;
+			}
+			
+			HeadComments.Add(&Comment);
+		}
+	}
+
+	// add an empty entry that visually spaces the comments from the diff entries
+	TSharedPtr<FBlueprintDifferenceTreeEntry> PaddingEntry = TSharedPtr<FBlueprintDifferenceTreeEntry>(new FBlueprintDifferenceTreeEntry(
+		{}
+		, FGenerateDiffEntryWidget::CreateLambda([]()
+		{
+			return
+			SNew(SBox)
+			.Padding(0.f, 20.f, 0.f, 10.f)
+			.HAlign(HAlign_Fill)
+			[
+				SNew(SImage)
+				.DesiredSizeOverride(FVector2D{100.0, 2.0})
+				.Image(FCoreStyle::Get().GetBrush("GenericWhiteBox"))
+				.ColorAndOpacity(FStyleColors::White25)
+			];
+		})
+		, {}
+	));
+	OutTreeEntries.Push(PaddingEntry);
+	
+	for (const FReviewComment* Comment : HeadComments)
+	{
+		// add comment as a thread (including replies as it's children)
+		GenerateCommentThreadRecursive(*Comment, CommentReplyMap, OutTreeEntries);
+	}
+	
+	OutTreeEntries.Push(FCommentDraftTreeEntry::MakeCommentDraft(AsShared(), &OutTreeEntries));
+}
+
+void FReviewCommentsDiffControl::SetCategory(const FString& CategoryKey)
+{
+	CommentCategory = CategoryKey;
+}
+
+void FReviewCommentsDiffControl::PostComment(FReviewComment& Comment)
+{
+	UE::DiffControl::GPostReviewComment(Comment);
+	if (const TSharedPtr<STreeView<TSharedPtr<FBlueprintDifferenceTreeEntry>>> TreeView = CommentsTreeView.Pin())
+	{
+		TreeView->RebuildList();
+	}
+}
+
+static void DeferredCallback(TSharedPtr<SWidget> Widget, int32 FramesToSkip, TFunction<void(void)> Callback)
+{
+	Widget->RegisterActiveTimer(0.f, FWidgetActiveTimerDelegate::CreateLambda(
+			[Callback, FramesToSkip, FrameCount = 0](double, float) mutable
+			{
+				if (++FrameCount < FramesToSkip)
+				{
+					return EActiveTimerReturnType::Continue;
+				}
+				Callback();
+				return EActiveTimerReturnType::Stop;
+			}
+		)
+	);
+}
+
+TWeakPtr<FCommentDraftTreeEntry> FReviewCommentsDiffControl::ReplyDraftEntry;
+
+void FReviewCommentsDiffControl::DraftReply(TSharedPtr<FCommentTreeEntry> ParentComment)
+{
+	if (const TSharedPtr<STreeView<TSharedPtr<FBlueprintDifferenceTreeEntry>>> TreeView = CommentsTreeView.Pin())
+	{
+		TSharedPtr<FCommentDraftTreeEntry> DraftEntry = ReplyDraftEntry.Pin();
+		if (DraftEntry)
+		{
+			DraftEntry->ReassignReplyParent(ParentComment);
+		}
+		else
+		{
+			DraftEntry = FCommentDraftTreeEntry::MakeReplyDraft(AsShared(), ParentComment);
+			ReplyDraftEntry = DraftEntry;
+		}
+		ParentComment->Children.Push(DraftEntry);
+		TreeView->SetItemExpansion(ParentComment, true);
+		TreeView->RebuildList();
+		
+		// focus comment text box once it's finished constructing
+		DeferredCallback(TreeView, 2, [DraftEntryWeak = TWeakPtr<FCommentDraftTreeEntry>(DraftEntry)]()
+		{
+			if (const TSharedPtr<FCommentDraftTreeEntry> DraftEntry = DraftEntryWeak.Pin())
+			{
+				FSlateApplication::Get().SetKeyboardFocus(DraftEntry->GetCommentTextBox());
+			}
+		});
+	}
+}
+
+void FReviewCommentsDiffControl::RebuildListView() const
+{
+	if (const TSharedPtr<STreeView<TSharedPtr<FBlueprintDifferenceTreeEntry>>> TreeView = CommentsTreeView.Pin())
+	{
+		TreeView->RebuildList();
+	}
+}
+
+#undef LOCTEXT_NAMESPACE
diff --git a/Engine/Source/Editor/Kismet/Private/DiffControl.h b/Engine/Source/Editor/Kismet/Private/DiffControl.h
--- a/Engine/Source/Editor/Kismet/Private/DiffControl.h
+++ b/Engine/Source/Editor/Kismet/Private/DiffControl.h
@@ -7,9 +7,13 @@
 #include "DiffResults.h"
 #include "IAssetTypeActions.h"
 #include "IDetailsView.h"
+#include "ReviewComments.h"
 #include "Algo/Transform.h"
 #include "Widgets/SWidget.h"
+#include "DiffUtils.h"
 
+struct FReviewCommentsDiffControl;
+class SMultiLineEditableTextBox;
 struct FDiffResultItem;
 class SBlueprintDiff;
 class UEdGraph;
@@ -28,7 +32,16 @@ public:
 	static FText RightRevision;
 	static TSharedRef<SWidget> GenerateSimpleDiffWidget(FText DiffText);
 	static TSharedRef<SWidget> GenerateObjectDiffWidget(FSingleObjectDiffEntry DiffEntry, FText ObjectName);
+
+	// to support comment posting, set this to the tree view that contains the comments.
+	void EnableComments(TWeakPtr<STreeView<TSharedPtr<FBlueprintDifferenceTreeEntry>>> TreeView, const UObject* OldObject, const UObject* NewObject);
+	virtual void EnableComments(TWeakPtr<STreeView<TSharedPtr<FBlueprintDifferenceTreeEntry>>> TreeView) {}
+	
+protected:
+	// to be called at the end of GenerateTreeEntries if you want this diff control to support review comments
+	void GenerateCategoryCommentTreeEntries(TArray< TSharedPtr<FBlueprintDifferenceTreeEntry> >& , TArray< TSharedPtr<FBlueprintDifferenceTreeEntry> >& OutRealDifferences, FString CategoryKey);
 	
+	TSharedPtr<FReviewCommentsDiffControl> ReviewCommentsDiffControl;
 };
 
 /** Shows all differences for the blueprint structure itself that aren't picked up elsewhere */
@@ -39,6 +52,9 @@ public:
 
 	virtual void GenerateTreeEntries(TArray< TSharedPtr<FBlueprintDifferenceTreeEntry> >& OutTreeEntries, TArray< TSharedPtr<FBlueprintDifferenceTreeEntry> >& OutRealDifferences) override;
 
+	// to support comment posting, set this to the tree view that contains the comments.
+	virtual void EnableComments(TWeakPtr<STreeView<TSharedPtr<FBlueprintDifferenceTreeEntry>>> TreeView) override;
+
 private:
 	FOnDiffEntryFocused SelectionCallback;
 	const UBlueprint* OldBlueprint;
@@ -59,6 +75,9 @@ public:
 	TSharedRef<SWidget> OldTreeWidget() { return OldSCS.TreeWidget(); }
 	TSharedRef<SWidget> NewTreeWidget() { return NewSCS.TreeWidget(); }
 
+	// to support comment posting, set this to the tree view that contains the comments.
+	virtual void EnableComments(TWeakPtr<STreeView<TSharedPtr<FBlueprintDifferenceTreeEntry>>> TreeView) override;
+
 private:
 	FOnDiffEntryFocused SelectionCallback;
 	FSCSDiffRoot DifferingProperties;
@@ -74,10 +93,14 @@ public:
 	FDetailsDiffControl(const UObject* InOldObject, const UObject* InNewObject, FOnDiffEntryFocused InSelectionCallback, bool bPopulateOutTreeEntries);
 
 	virtual void GenerateTreeEntries(TArray< TSharedPtr<FBlueprintDifferenceTreeEntry> >& OutTreeEntries, TArray< TSharedPtr<FBlueprintDifferenceTreeEntry> >& OutRealDifferences) override;
+	void GenerateTreeEntriesWithoutComments(TArray< TSharedPtr<FBlueprintDifferenceTreeEntry> >& OutTreeEntries, TArray< TSharedPtr<FBlueprintDifferenceTreeEntry> >& OutRealDifferences);
 
 	TSharedRef<SWidget> OldDetailsWidget() { return OldDetails.DetailsWidget(); }
 	TSharedRef<SWidget> NewDetailsWidget() { return NewDetails.DetailsWidget(); }
 
+	// to support comment posting, set this to the tree view that contains the comments.
+	virtual void EnableComments(TWeakPtr<STreeView<TSharedPtr<FBlueprintDifferenceTreeEntry>>> TreeView) override;
+
 protected:
 	virtual void OnSelectDiffEntry(FPropertySoftPath PropertyName);
 	TAttribute<TArray<FVector2f>> GetLinkedScrollRateAttribute(const TSharedRef<IDetailsView>& OldDetailsView, const TSharedRef<IDetailsView>& NewDetailsView);
@@ -141,7 +164,10 @@ struct FBlueprintTypeDiffControl : public TSharedFromThis<FBlueprintTypeDiffCont
 
 	/** Generate difference tree widgets */
 	virtual void GenerateTreeEntries(TArray< TSharedPtr<FBlueprintDifferenceTreeEntry> >& OutTreeEntries, TArray< TSharedPtr<FBlueprintDifferenceTreeEntry> >& OutRealDifferences) override;
-
+	
+	// to support comment posting, set this to the tree view that contains the comments.
+	virtual void EnableComments(TWeakPtr<STreeView<TSharedPtr<FBlueprintDifferenceTreeEntry>>> TreeView) override;
+	
 	/** The old blueprint (left) */
 	const UBlueprint* BlueprintOld;
 
@@ -191,6 +217,9 @@ struct FGraphToDiff	: public TSharedFromThis<FGraphToDiff>, IDiffControl
 	/** Get new(right) graph*/
 	UEdGraph* GetGraphNew() const;
 
+	// to support comment posting, set this to the tree view that contains the comments.
+	virtual void EnableComments(TWeakPtr<STreeView<TSharedPtr<FBlueprintDifferenceTreeEntry>>> TreeView) override;
+	
 	/** Source for list view */
 	TArray<TSharedPtr<FDiffResultItem>> DiffListSource;
 	TSharedPtr<TArray<FDiffSingleResult>> FoundDiffs;
@@ -225,4 +254,125 @@ private:
 
 	/** Handle to the registered OnGraphChanged delegate. */
 	FDelegateHandle OnGraphChangedDelegateHandle;
-};
\ No newline at end of file
+};
+
+class FCommentTreeEntry : public FBlueprintDifferenceTreeEntry, public TSharedFromThis<FCommentTreeEntry>
+{
+public:
+	FCommentTreeEntry(TWeakPtr<FReviewCommentsDiffControl> InCommentsControl, const FReviewComment& InComment, const TArray<TSharedPtr<FBlueprintDifferenceTreeEntry>>& InChildren);
+	virtual ~FCommentTreeEntry() override;
+	static TSharedRef<FCommentTreeEntry> Make(TWeakPtr<FReviewCommentsDiffControl> CommentsControl, const FReviewComment& Comment, const TArray<TSharedPtr<FBlueprintDifferenceTreeEntry>>& Children = {});
+
+	int32 GetCommentIDChecked() const;
+	void AwaitCommentPost();
+private:
+	// FBlueprintDifferenceTreeEntry callbacks
+	TSharedRef<SWidget> CreateWidget();
+
+	// slate callbacks
+	bool IsCommentTextBoxReadonly() const;
+	FReply OnClickReply();
+	FReply OnClickEdit();
+	FReply OnLikeToggle();
+	FText GetLikeTooltip() const;
+	const FSlateBrush* GetLikeIcon() const;
+	FSlateColor GetUsernameColor() const;
+	FSlateColor GetLikeIconColor() const;
+	EVisibility GetEditReplyButtonGroupVisibility() const;
+	EVisibility GetEditButtonVisibility() const;
+	EVisibility GetSubmitCancelButtonGroupVisibility() const;
+	bool IsSubmitButtonEnabled() const;
+	FReply OnEditSubmitClicked();
+	FReply OnEditCancelClicked();
+
+	// comment api callbacks
+	void OnCommentPosted(const FReviewComment& Comment);
+
+	FString GetCommentString() const;
+	void SetCommentString(const FString& NewComment);
+
+	// during a pending edit, returns whether the comment is different from the saved version
+	bool HasCommentStringChanged() const;
+	
+	bool IsEditMode() const;
+	void SetEditMode(bool bIsEditMode);
+	
+	FReviewComment Comment;
+
+	TSharedPtr<SWidget> Content;
+	TSharedPtr<SHorizontalBox> EditReplyButtonGroup;
+	TSharedPtr<SHorizontalBox> SubmitCancelButtonGroup;
+	TSharedPtr<SMultiLineEditableTextBox> CommentTextBox;
+	
+	TWeakPtr<FReviewCommentsDiffControl> CommentsControl;
+	FDelegateHandle OnCommentPostedHandle;
+
+	// prefer the use of IsEditMode().
+	bool bExpectedEditMode = false;
+};
+
+class FCommentDraftTreeEntry : public FBlueprintDifferenceTreeEntry, public TSharedFromThis<FCommentDraftTreeEntry>
+{
+public:
+	FCommentDraftTreeEntry(TWeakPtr<FReviewCommentsDiffControl> InCommentsControl, TArray<TSharedPtr<FBlueprintDifferenceTreeEntry>>* InSiblings, int32 InReplyID = -1);
+	static TSharedRef<FCommentDraftTreeEntry> MakeCommentDraft(TWeakPtr<FReviewCommentsDiffControl> CommentsControl, TArray<TSharedPtr<FBlueprintDifferenceTreeEntry>>* Siblings);
+	static TSharedRef<FCommentDraftTreeEntry> MakeReplyDraft(TWeakPtr<FReviewCommentsDiffControl> CommentsControl, TSharedPtr<FCommentTreeEntry> InParent);
+
+	void ReassignReplyParent(TSharedPtr<FCommentTreeEntry> InParent);
+
+	TSharedPtr<SMultiLineEditableTextBox> GetCommentTextBox() const {return CommentTextBox;}
+
+	bool IsReply() const;
+private:
+	// FBlueprintDifferenceTreeEntry callbacks
+	TSharedRef<SWidget> CreateWidget();
+	
+	// slate callbacks
+	bool IsPostButtonEnabled() const;
+	FReply OnCommentPostClicked();
+	
+	void ReassignSiblings(TArray<TSharedPtr<FBlueprintDifferenceTreeEntry>>* InSiblings, int32 InReplyID);
+	
+	TSharedPtr<SWidget> Content;
+	TSharedPtr<SMultiLineEditableTextBox> CommentTextBox;
+	TSharedPtr<SCheckBox> FlagAsTaskCheckBox;
+	
+	TWeakPtr<FReviewCommentsDiffControl> CommentsControl;
+	
+	TArray<TSharedPtr<FBlueprintDifferenceTreeEntry>>* Siblings;
+	
+	// if this is a reply, store the comment id of it's parent
+	int32 ReplyID = -1;
+};
+
+/** Category list item for a graph*/
+struct FReviewCommentsDiffControl : public TSharedFromThis<FReviewCommentsDiffControl>, IDiffControl
+{
+	FReviewCommentsDiffControl(const UObject* InOldObject, const UObject* InNewObject, TWeakPtr<STreeView<TSharedPtr<FBlueprintDifferenceTreeEntry>>> TreeView);
+
+	virtual void GenerateTreeEntries(TArray< TSharedPtr<FBlueprintDifferenceTreeEntry> >& OutTreeEntries, TArray< TSharedPtr<FBlueprintDifferenceTreeEntry> >& OutRealDifferences) override;
+
+	void SetCategory(const FString& CategoryKey);
+
+	void PostComment(FReviewComment& Comment);
+	void DraftReply(TSharedPtr<FCommentTreeEntry> ParentComment);
+	void RebuildListView() const;
+
+	const FString& GetCommentFilePath() {return CommentFilePath;}
+	const FString& GetCommentCategory() {return CommentCategory;}
+protected:
+	
+	void GenerateCommentThreadRecursive(const FReviewComment& Comment,
+		const TMap<int32, TArray<const FReviewComment*>>& CommentReplyMap,
+		TArray<TSharedPtr<FBlueprintDifferenceTreeEntry>>& OutTreeEntries);
+	
+	const UObject* OldObject;
+	const UObject* NewObject;
+	
+	FString CommentFilePath;
+	FString CommentCategory;
+	TWeakPtr<STreeView<TSharedPtr<FBlueprintDifferenceTreeEntry>>> CommentsTreeView;
+
+	// we only want one reply draft widget in existence so we can use a static to track it
+	static TWeakPtr<FCommentDraftTreeEntry> ReplyDraftEntry;
+};
diff --git a/Engine/Source/Editor/Kismet/Private/ReviewComments.cpp b/Engine/Source/Editor/Kismet/Private/ReviewComments.cpp
--- a/Engine/Source/Editor/Kismet/Private/ReviewComments.cpp
+++ b/Engine/Source/Editor/Kismet/Private/ReviewComments.cpp
@@ -335,7 +335,7 @@ TSharedPtr<FJsonObject> FReviewComment::ToJsonObject() const
 	SET_JSON_FIELD_NAMED(Result, *this, UpdatedTime, TEXT("updated"));
 	if (bIsClosed)
 	{
-		TrySetJsonObjectField(Result, TArray<FString>{TEXT("closed")}, TEXT("type"));
+		TrySetJsonObjectField(Result, TArray<FString>{TEXT("closed")}, TEXT("flags"));
 	}
 	
 	return Result;
diff --git a/Engine/Source/Editor/Kismet/Private/SBlueprintDiff.cpp b/Engine/Source/Editor/Kismet/Private/SBlueprintDiff.cpp
--- a/Engine/Source/Editor/Kismet/Private/SBlueprintDiff.cpp
+++ b/Engine/Source/Editor/Kismet/Private/SBlueprintDiff.cpp
@@ -35,6 +35,7 @@
 #include "Misc/AssertionMacros.h"
 #include "Misc/CString.h"
 #include "PropertyEditorDelegates.h"
+#include "ReviewComments.h"
 #include "SKismetInspector.h"
 #include "SMyBlueprint.h"
 #include "SlateOptMacros.h"
@@ -68,6 +69,7 @@ static const FName MyBlueprintMode = FName(TEXT("MyBlueprintMode"));
 static const FName DefaultsMode = FName(TEXT("DefaultsMode"));
 static const FName ClassSettingsMode = FName(TEXT("ClassSettingsMode"));
 static const FName ComponentsMode = FName(TEXT("ComponentsMode"));
+static const FName CommentsMode = FName(TEXT("CommentsMode"));
 static const FName GraphMode = FName(TEXT("GraphMode"));
 
 TSharedRef<SWidget>	FDiffResultItem::GenerateWidget() const
@@ -390,6 +392,8 @@ void SBlueprintDiff::OnCloseAssetEditor(UObject* Asset, EAssetEditorCloseReason
 void SBlueprintDiff::CreateGraphEntry( UEdGraph* GraphOld, UEdGraph* GraphNew )
 {
 	Graphs.Add(MakeShared<FGraphToDiff>(this, GraphOld, GraphNew, PanelOld.RevisionInfo, PanelNew.RevisionInfo));
+	Graphs.Last()->EnableComments(DifferencesTreeView.ToWeakPtr());
+	
 }
 
 void SBlueprintDiff::OnGraphSelectionChanged(TSharedPtr<FGraphToDiff> Item, ESelectInfo::Type SelectionType)
@@ -945,6 +949,8 @@ void SBlueprintDiff::GenerateDifferencesList()
 	{
 		Graph->GenerateTreeEntries(PrimaryDifferencesList, RealDifferences);
 	}
+	
+	ModePanels.Add(CommentsMode, GenerateGeneralFileCommentEntries());
 
 	DifferencesTreeView->RebuildList();
 }
@@ -952,6 +958,7 @@ void SBlueprintDiff::GenerateDifferencesList()
 SBlueprintDiff::FDiffControl SBlueprintDiff::GenerateBlueprintTypePanel()
 {
 	TSharedPtr<FBlueprintTypeDiffControl> NewDiffControl = MakeShared<FBlueprintTypeDiffControl>(PanelOld.Blueprint, PanelNew.Blueprint, FOnDiffEntryFocused::CreateRaw(this, &SBlueprintDiff::SetCurrentMode, BlueprintTypeMode));
+	NewDiffControl->EnableComments(DifferencesTreeView.ToWeakPtr());
 	NewDiffControl->GenerateTreeEntries(PrimaryDifferencesList, RealDifferences);
 
 	SBlueprintDiff::FDiffControl Ret;
@@ -984,6 +991,7 @@ SBlueprintDiff::FDiffControl SBlueprintDiff::GenerateBlueprintTypePanel()
 SBlueprintDiff::FDiffControl SBlueprintDiff::GenerateMyBlueprintPanel()
 {
 	TSharedPtr<FMyBlueprintDiffControl> NewDiffControl = MakeShared<FMyBlueprintDiffControl>(PanelOld.Blueprint, PanelNew.Blueprint, FOnDiffEntryFocused::CreateRaw(this, &SBlueprintDiff::SetCurrentMode, MyBlueprintMode));
+	NewDiffControl->EnableComments(DifferencesTreeView.ToWeakPtr());
 	NewDiffControl->GenerateTreeEntries(PrimaryDifferencesList, RealDifferences);
 
 	SBlueprintDiff::FDiffControl Ret;
@@ -1090,6 +1098,7 @@ SBlueprintDiff::FDiffControl SBlueprintDiff::GenerateDefaultsPanel()
 	const UObject* B = DiffUtils::GetCDO(PanelNew.Blueprint);
 
 	TSharedPtr<FCDODiffControl> NewDiffControl = MakeShared<FCDODiffControl>(A, B, FOnDiffEntryFocused::CreateRaw(this, &SBlueprintDiff::SetCurrentMode, DefaultsMode));
+	NewDiffControl->EnableComments(DifferencesTreeView.ToWeakPtr());
 	NewDiffControl->GenerateTreeEntries(PrimaryDifferencesList, RealDifferences);
 
 	SBlueprintDiff::FDiffControl Ret;
@@ -1113,6 +1122,7 @@ SBlueprintDiff::FDiffControl SBlueprintDiff::GenerateDefaultsPanel()
 SBlueprintDiff::FDiffControl SBlueprintDiff::GenerateClassSettingsPanel()
 {
 	TSharedPtr<FClassSettingsDiffControl> NewDiffControl = MakeShared<FClassSettingsDiffControl>(PanelOld.Blueprint, PanelNew.Blueprint, FOnDiffEntryFocused::CreateRaw(this, &SBlueprintDiff::SetCurrentMode, ClassSettingsMode));
+	NewDiffControl->EnableComments(DifferencesTreeView.ToWeakPtr());
 	NewDiffControl->GenerateTreeEntries(PrimaryDifferencesList, RealDifferences);
 
 	SBlueprintDiff::FDiffControl Ret;
@@ -1136,6 +1146,7 @@ SBlueprintDiff::FDiffControl SBlueprintDiff::GenerateClassSettingsPanel()
 SBlueprintDiff::FDiffControl SBlueprintDiff::GenerateComponentsPanel()
 {
 	TSharedPtr<FSCSDiffControl> NewDiffControl = MakeShared<FSCSDiffControl>(PanelOld.Blueprint, PanelNew.Blueprint, FOnDiffEntryFocused::CreateRaw(this, &SBlueprintDiff::SetCurrentMode, ComponentsMode));
+	NewDiffControl->EnableComments(DifferencesTreeView.ToWeakPtr());
 	NewDiffControl->GenerateTreeEntries(PrimaryDifferencesList, RealDifferences);
 
 	SBlueprintDiff::FDiffControl Ret;
@@ -1156,6 +1167,18 @@ SBlueprintDiff::FDiffControl SBlueprintDiff::GenerateComponentsPanel()
 	return Ret;
 }
 
+SBlueprintDiff::FDiffControl SBlueprintDiff::GenerateGeneralFileCommentEntries()
+{
+	TSharedPtr<FReviewCommentsDiffControl> NewDiffControl = MakeShared<FReviewCommentsDiffControl>(
+		PanelOld.Blueprint, PanelNew.Blueprint, DifferencesTreeView.ToWeakPtr());
+	NewDiffControl->GenerateTreeEntries(PrimaryDifferencesList, RealDifferences);
+
+	SBlueprintDiff::FDiffControl Ret;
+	Ret.DiffControl = NewDiffControl;
+
+	return Ret;
+}
+
 TSharedRef<SOverlay> SBlueprintDiff::GenerateGraphWidgetForPanel(FDiffPanel& OutDiffPanel) const
 {
 	return SNew(SOverlay)
diff --git a/Engine/Source/Editor/Kismet/Private/SDetailsDiff.cpp b/Engine/Source/Editor/Kismet/Private/SDetailsDiff.cpp
--- a/Engine/Source/Editor/Kismet/Private/SDetailsDiff.cpp
+++ b/Engine/Source/Editor/Kismet/Private/SDetailsDiff.cpp
@@ -328,6 +328,7 @@ void SDetailsDiff::GenerateDifferencesList()
 SDetailsDiff::FDiffControl SDetailsDiff::GenerateDetailsPanel()
 {
 	TSharedPtr<FDetailsDiffControl> NewDiffControl = MakeShared<FDetailsDiffControl>(PanelOld.Object, PanelNew.Object, FOnDiffEntryFocused::CreateRaw(this, &SDetailsDiff::SetCurrentMode, DetailsMode), true);
+	NewDiffControl->EnableComments(DifferencesTreeView.ToWeakPtr());
 	NewDiffControl->GenerateTreeEntries(PrimaryDifferencesList, RealDifferences);
 
 	SDetailsDiff::FDiffControl Ret;
diff --git a/Engine/Source/Editor/Kismet/Public/ReviewComments.h b/Engine/Source/Editor/Kismet/Public/ReviewComments.h
--- a/Engine/Source/Editor/Kismet/Public/ReviewComments.h
+++ b/Engine/Source/Editor/Kismet/Public/ReviewComments.h
@@ -66,8 +66,17 @@ public:
 	// returns topic in format "TopicType/ChangelistNum"
 	FString ToString() const;
 	static FReviewTopic FromString(const FString& ReviewTopic);
+	bool operator==(const FReviewTopic& Other) const
+	{
+		return ChangelistNum == Other.ChangelistNum;
+	}
 };
 
+inline uint32 GetTypeHash(const FReviewTopic& Topic)
+{
+	return GetTypeHash(Topic.ChangelistNum);
+}
+
 struct KISMET_API FReviewComment
 {
 public:	
@@ -105,7 +114,7 @@ public:
 	static FReviewComment FromJson(const TSharedPtr<FJsonObject>& JsonObject);
 };
 
-class KISMET_API IReviewCommentAPI
+class KISMET_API IReviewCommentAPI : public TSharedFromThis<IReviewCommentAPI>
 {
 public:
 	DECLARE_DELEGATE_TwoParams(OnGetCommentsComplete, const TArray<FReviewComment>&, const FString& /*ErrorMessage*/)
@@ -120,9 +129,9 @@ public:
 	
 	virtual void GetComments(const FReviewTopic& Topic, const OnGetCommentsComplete& OnComplete) const = 0;
 	// Body must be set; username will get set by PostComment
-	virtual void PostComment(FReviewComment& Comment, const OnPostCommentComplete& OnComplete) const = 0;
+	virtual void PostComment(FReviewComment& Comment, const OnPostCommentComplete& OnComplete, bool bSilenceNotification = false) const = 0;
 	// all unset members will be left unchanged.
-	virtual void EditComment(const FReviewComment& Comment, const OnEditCommentComplete& OnComplete) const = 0;
+	virtual void EditComment(const FReviewComment& Comment, const OnEditCommentComplete& OnComplete, bool bSilenceNotification = false) const = 0;
 	
 	virtual void GetReviewTopicForCL(const FString &ChangelistNum, const OnGetReviewTopicForCLComplete& OnComplete) const = 0;
 };
diff --git a/Engine/Source/Editor/Kismet/Public/SBlueprintDiff.h b/Engine/Source/Editor/Kismet/Public/SBlueprintDiff.h
--- a/Engine/Source/Editor/Kismet/Public/SBlueprintDiff.h
+++ b/Engine/Source/Editor/Kismet/Public/SBlueprintDiff.h
@@ -244,6 +244,7 @@ protected:
 	FDiffControl GenerateDefaultsPanel();
 	FDiffControl GenerateClassSettingsPanel();
 	FDiffControl GenerateComponentsPanel();
+	FDiffControl GenerateGeneralFileCommentEntries();
 
 	TSharedRef<SOverlay> GenerateGraphWidgetForPanel(FDiffPanel& OutDiffPanel) const;
 	TSharedRef<SBox> GenerateRevisionInfoWidgetForPanel(TSharedPtr<SWidget>& OutGeneratedWidget,const FText& InRevisionText) const;
diff --git a/Engine/Source/Editor/UnrealEd/Private/DiffUtils.cpp b/Engine/Source/Editor/UnrealEd/Private/DiffUtils.cpp
--- a/Engine/Source/Editor/UnrealEd/Private/DiffUtils.cpp
+++ b/Engine/Source/Editor/UnrealEd/Private/DiffUtils.cpp
@@ -782,7 +782,10 @@ UPackage* DiffUtils::LoadPackageForDiff(const FPackagePath& InTempPackagePath, c
 {
 	// set up instancing context
 	FLinkerInstancingContext Context;
-	Context.AddPackageMapping(InOriginalPackagePath.GetPackageFName(), InTempPackagePath.GetPackageFName());
+	if (!InOriginalPackagePath.GetLocalFullPath().IsEmpty())
+	{
+		Context.AddPackageMapping(InOriginalPackagePath.GetPackageFName(), InTempPackagePath.GetPackageFName());
+	}
 	
 	return LoadPackage(nullptr, *InTempPackagePath.GetPackageName(),
 		LOAD_ForDiff | LOAD_DisableCompileOnLoad | LOAD_DisableEngineVersionChecks, nullptr, &Context);
diff --git a/Engine/Source/Editor/UnrealEd/Public/DiffUtils.h b/Engine/Source/Editor/UnrealEd/Public/DiffUtils.h
--- a/Engine/Source/Editor/UnrealEd/Public/DiffUtils.h
+++ b/Engine/Source/Editor/UnrealEd/Public/DiffUtils.h
@@ -312,6 +312,8 @@ public:
 		check( InGenerateWidget.IsBound() );
 	}
 
+	virtual ~FBlueprintDifferenceTreeEntry() = default;
+
 	/** Displays message to user saying there are no differences */
 	static TSharedPtr<FBlueprintDifferenceTreeEntry> NoDifferencesEntry();
 
