Non Unity Fix

diff --git a/Engine/Source/Editor/Kismet/Private/DiffControl.cpp b/Engine/Source/Editor/Kismet/Private/DiffControl.cpp
--- a/Engine/Source/Editor/Kismet/Private/DiffControl.cpp
+++ b/Engine/Source/Editor/Kismet/Private/DiffControl.cpp
@@ -9,6 +9,9 @@
 #include "IDetailsView.h"
 #include "ReviewComments.h"
 #include "Widgets/Input/SMultiLineEditableTextBox.h"
+#include "Widgets/Input/SButton.h"
+#include "Widgets/Input/SCheckBox.h"
+#include "Framework/Application/SlateApplication.h"
 
 #define LOCTEXT_NAMESPACE "SBlueprintDif"
 
diff --git a/Engine/Source/Editor/Kismet/Private/DiffControl.h b/Engine/Source/Editor/Kismet/Private/DiffControl.h
--- a/Engine/Source/Editor/Kismet/Private/DiffControl.h
+++ b/Engine/Source/Editor/Kismet/Private/DiffControl.h
@@ -18,6 +18,7 @@ struct FDiffResultItem;
 class SBlueprintDiff;
 class UEdGraph;
 struct FEdGraphEditAction;
+class SCheckBox;
 
 
 /** Interface responsible for generating FBlueprintDifferenceTreeEntry's for visual diff tools */
