Fix race between BP compilation and reinstanciation now that the latter is back on GT

#rnx

diff --git a/Engine/Source/Editor/Kismet/Private/BlueprintCompilationManager.cpp b/Engine/Source/Editor/Kismet/Private/BlueprintCompilationManager.cpp
--- a/Engine/Source/Editor/Kismet/Private/BlueprintCompilationManager.cpp
+++ b/Engine/Source/Editor/Kismet/Private/BlueprintCompilationManager.cpp
@@ -171,6 +171,11 @@ struct FBlueprintCompilationManagerImpl : public FGCObject
 
 	// State stored so that we can check what stage of compilation we're in:
 	bool bGeneratedClassLayoutReady;
+
+#if WITH_EDITOR
+	// Used to avoid reinstanciation on the GT while compiling on the loading thread
+	FCriticalSection Lock;
+#endif
 };
 
 // free function that we use to cross a module boundary (from CoreUObject to here)
@@ -223,6 +228,9 @@ void FBlueprintCompilationManagerImpl::RegisterCompilerExtension(TSubclassOf<UBl
 
 void FBlueprintCompilationManagerImpl::QueueForCompilation(const FBPCompileRequestInternal& CompileJob)
 {
+#if WITH_EDITOR
+	FScopeLock ScopeLock(&Lock);
+#endif
 	if(!CompileJob.UserData.BPToCompile->bQueuedForCompilation)
 	{
 		if(GCompilingBlueprint)
@@ -245,6 +253,9 @@ void FBlueprintCompilationManagerImpl::QueueForCompilation(const FBPCompileReque
 
 void FBlueprintCompilationManagerImpl::CompileSynchronouslyImpl(const FBPCompileRequestInternal& Request)
 {
+#if WITH_EDITOR
+	FScopeLock ScopeLock(&Lock);
+#endif
 	DECLARE_SCOPE_HIERARCHICAL_COUNTER_FUNC()
 
 	Request.UserData.BPToCompile->bQueuedForCompilation = true;
@@ -555,6 +566,10 @@ void FBlueprintCompilationManagerImpl::FlushCompilationQueueImpl(bool bSuppressB
 	TRACE_CPUPROFILER_EVENT_SCOPE(FBlueprintCompilationManager::FlushCompilationQueue);
 	DECLARE_SCOPE_HIERARCHICAL_COUNTER_FUNC()
 
+#if WITH_EDITOR
+	FScopeLock ScopeLock(&Lock);
+#endif
+
 	TGuardValue<bool> GuardTemplateNameFlag(GCompilingBlueprint, true);
 	ensure(bGeneratedClassLayoutReady);
 
@@ -1876,6 +1891,10 @@ void FBlueprintCompilationManagerImpl::FlushReinstancingQueueImpl(bool bFindAndR
 {
 	DECLARE_SCOPE_HIERARCHICAL_COUNTER_FUNC()
 
+#if WITH_EDITOR
+	FScopeLock ScopeLock(&Lock);
+#endif
+
 	if(GCompilingBlueprint)
 	{
 		return;
