changed use of 'and' keyword to &&

diff --git a/Engine/Source/Editor/Kismet/Private/DiffControl.cpp b/Engine/Source/Editor/Kismet/Private/DiffControl.cpp
--- a/Engine/Source/Editor/Kismet/Private/DiffControl.cpp
+++ b/Engine/Source/Editor/Kismet/Private/DiffControl.cpp
@@ -1087,7 +1087,7 @@ static FText RangeAsText(RangeType&& Range, const FText& Delimiter, const FText&
 	for (auto&& Elem : Forward<RangeType>(Range))
 	{
 		FText ElementText = Invoke(Transformer, Elem);
-		if (Index == LastIndex and !Conjunction.IsEmpty())
+		if (Index == LastIndex && !Conjunction.IsEmpty())
 		{
 			// for the last item, prepend a conjunction
 			ItemsAsText.Add(FText::FormatNamed(
