Fixes: while compiling a Level Blueprint with an Actor Object Reference array variable that changed typage to Soft Object Reference, editor will crash.

diff --git a/Engine/Source/Editor/Kismet/Private/BlueprintEditor.cpp b/Engine/Source/Editor/Kismet/Private/BlueprintEditor.cpp
--- a/Engine/Source/Editor/Kismet/Private/BlueprintEditor.cpp
+++ b/Engine/Source/Editor/Kismet/Private/BlueprintEditor.cpp
@@ -4113,6 +4113,8 @@ void FBlueprintEditor::Compile()
 		// send record when player clicks compile and send the result
 		// this will make sure how the users activity is
 		AnalyticsTrackCompileEvent(BlueprintObj, LogResults.NumErrors, LogResults.NumWarnings);
+
+		RefreshInspector();
 	}
 }
 
