Relocate where we purge stale references from the Blueprint editor's MRU document list in order to avoid attempts to resolve those paths during async editor loads.

diff --git a/Engine/Source/Editor/Kismet/Private/BlueprintEditor.cpp b/Engine/Source/Editor/Kismet/Private/BlueprintEditor.cpp
--- a/Engine/Source/Editor/Kismet/Private/BlueprintEditor.cpp
+++ b/Engine/Source/Editor/Kismet/Private/BlueprintEditor.cpp
@@ -9989,6 +9989,7 @@ void FBlueprintEditor::RestoreEditedObjectState()
 		}
 	}
 
+	TSet<FSoftObjectPath> PathsToRemove;
 	for (int32 i = 0; i < Blueprint->LastEditedDocuments.Num(); i++)
 	{
 		if (UObject* Obj = Blueprint->LastEditedDocuments[i].EditedObjectPath.ResolveObject())
@@ -10037,6 +10038,20 @@ void FBlueprintEditor::RestoreEditedObjectState()
 				TSharedPtr<SDockTab> TabWithGraph = OpenDocument(Obj, FDocumentTracker::RestorePreviousDocument);
 			}
 		}
+		else
+		{
+			PathsToRemove.Add(Blueprint->LastEditedDocuments[i].EditedObjectPath);
+		}
+	}
+
+	// Older assets may have neglected to clean up this array when referenced objects were deleted, so
+	// we'll check for that now. This is done to ensure we don't store invalid object paths indefinitely.
+	if (PathsToRemove.Num() > 0)
+	{
+		Blueprint->LastEditedDocuments.RemoveAll([&PathsToRemove](const FEditedDocumentInfo& Entry)
+		{
+			return PathsToRemove.Contains(Entry.EditedObjectPath);
+		});
 	}
 }
 
diff --git a/Engine/Source/Editor/UnrealEd/Private/Kismet2/BlueprintEditorUtils.cpp b/Engine/Source/Editor/UnrealEd/Private/Kismet2/BlueprintEditorUtils.cpp
--- a/Engine/Source/Editor/UnrealEd/Private/Kismet2/BlueprintEditorUtils.cpp
+++ b/Engine/Source/Editor/UnrealEd/Private/Kismet2/BlueprintEditorUtils.cpp
@@ -6908,8 +6908,6 @@ void FBlueprintEditorUtils::PurgeNullGraphs(UBlueprint* Blueprint)
 	CleanNullGraphReferencesInArray(Blueprint, Blueprint->FunctionGraphs);
 	CleanNullGraphReferencesInArray(Blueprint, Blueprint->DelegateSignatureGraphs);
 	CleanNullGraphReferencesInArray(Blueprint, Blueprint->MacroGraphs);
-
-	Blueprint->LastEditedDocuments.RemoveAll([](const FEditedDocumentInfo& TestDoc) { return TestDoc.EditedObjectPath.ResolveObject() == nullptr; });
 }
 
 struct FConformCallsToParentFunctionUtils
