Gather search data for the BlueprintExtension. Used by the FindInBlueprint system.

diff --git a/Engine/Plugins/Runtime/ModelViewViewModel/Source/ModelViewViewModelBlueprint/ModelViewViewModelBlueprint.Build.cs b/Engine/Plugins/Runtime/ModelViewViewModel/Source/ModelViewViewModelBlueprint/ModelViewViewModelBlueprint.Build.cs
--- a/Engine/Plugins/Runtime/ModelViewViewModel/Source/ModelViewViewModelBlueprint/ModelViewViewModelBlueprint.Build.cs
+++ b/Engine/Plugins/Runtime/ModelViewViewModel/Source/ModelViewViewModelBlueprint/ModelViewViewModelBlueprint.Build.cs
@@ -23,6 +23,7 @@ public class ModelViewViewModelBlueprint : ModuleRules
 			new string[]
 			{
 				"BlueprintGraph",
+				"Kismet",
 				"KismetCompiler",
 				"PropertyEditor",
 				"PropertyPath",
diff --git a/Engine/Plugins/Runtime/ModelViewViewModel/Source/ModelViewViewModelBlueprint/Private/MVVMWidgetBlueprintExtension_View.cpp b/Engine/Plugins/Runtime/ModelViewViewModel/Source/ModelViewViewModelBlueprint/Private/MVVMWidgetBlueprintExtension_View.cpp
--- a/Engine/Plugins/Runtime/ModelViewViewModel/Source/ModelViewViewModelBlueprint/Private/MVVMWidgetBlueprintExtension_View.cpp
+++ b/Engine/Plugins/Runtime/ModelViewViewModel/Source/ModelViewViewModelBlueprint/Private/MVVMWidgetBlueprintExtension_View.cpp
@@ -5,8 +5,11 @@
 #include "MVVMViewBlueprintCompiler.h"
 #include "View/MVVMViewClass.h"
 
+#include "FindInBlueprintManager.h"
+
 #include UE_INLINE_GENERATED_CPP_BY_NAME(MVVMWidgetBlueprintExtension_View)
 
+#define LOCTEXT_NAMESPACE "MVVMBlueprintExtensionView"
 
 namespace UE::MVVM::Private
 {
@@ -133,6 +136,38 @@ void UMVVMWidgetBlueprintExtension_View::HandleFinishCompilingClass(UWidgetBluep
 	}
 }
 
+UMVVMWidgetBlueprintExtension_View::FSearchData UMVVMWidgetBlueprintExtension_View::HandleGatherSearchData(const UBlueprint* OwningBlueprint) const
+{
+	UMVVMWidgetBlueprintExtension_View::FSearchData SearchData;
+	if (GetBlueprintView())
+	{
+		{
+			TUniquePtr<FSearchArrayData> ViewModelContextSearchData = MakeUnique<FSearchArrayData>();
+			ViewModelContextSearchData->Identifier = LOCTEXT("ViewmodelSearchTag", "Viewmodels");
+			for (const FMVVMBlueprintViewModelContext& ViewModelContext : GetBlueprintView()->GetViewModels())
+			{
+				FSearchData& ViewModelSearchData = ViewModelContextSearchData->SearchSubList.AddDefaulted_GetRef();
+				ViewModelSearchData.Datas.Emplace(LOCTEXT("ViewmodelGuidSearchTag", "Guid"), FText::FromString(ViewModelContext.GetViewModelId().ToString(EGuidFormats::Digits)));
+				ViewModelSearchData.Datas.Emplace(FFindInBlueprintSearchTags::FiB_Name, ViewModelContext.GetDisplayName());
+				ViewModelSearchData.Datas.Emplace(FFindInBlueprintSearchTags::FiB_ClassName, ViewModelContext.GetViewModelClass() ? ViewModelContext.GetViewModelClass()->GetDisplayNameText() : FText::GetEmpty());
+				ViewModelSearchData.Datas.Emplace(LOCTEXT("ViewmodelCreationTypeSearchTag", "CreationType"), StaticEnum<EMVVMBlueprintViewModelContextCreationType>()->GetDisplayNameTextByValue((int64)ViewModelContext.CreationType));
+			}
+			SearchData.SearchArrayDatas.Add(MoveTemp(ViewModelContextSearchData));
+		}
+		{
+			TUniquePtr<FSearchArrayData> BindingContextSearchData = MakeUnique<FSearchArrayData>();
+			BindingContextSearchData->Identifier = LOCTEXT("BindingSearchTag", "Bindings");
+			for (const FMVVMBlueprintViewBinding& Binding : GetBlueprintView()->GetBindings())
+			{
+				FSearchData& BindingSearchData = BindingContextSearchData->SearchSubList.AddDefaulted_GetRef();
+				BindingSearchData.Datas.Emplace(LOCTEXT("ViewBindingSearchTag", "Binding"), FText::FromString(Binding.GetDisplayNameString(GetWidgetBlueprint())));
+			}
+			SearchData.SearchArrayDatas.Add(MoveTemp(BindingContextSearchData));
+		}
+	}
+	return SearchData;
+}
+
 void UMVVMWidgetBlueprintExtension_View::SetFilterSettings(FMVVMViewBindingFilterSettings InFilterSettings)
 {
 	FilterSettings = InFilterSettings;
@@ -148,4 +183,6 @@ void UMVVMWidgetBlueprintExtension_View::PostInitProperties()
 		SetFilterSettings(GetDefault<UMVVMDeveloperProjectSettings>()->FilterSettings);
 	}
 }
-#endif
\ No newline at end of file
+#endif
+
+#undef LOCTEXT_NAMESPACE
diff --git a/Engine/Plugins/Runtime/ModelViewViewModel/Source/ModelViewViewModelBlueprint/Public/MVVMWidgetBlueprintExtension_View.h b/Engine/Plugins/Runtime/ModelViewViewModel/Source/ModelViewViewModelBlueprint/Public/MVVMWidgetBlueprintExtension_View.h
--- a/Engine/Plugins/Runtime/ModelViewViewModel/Source/ModelViewViewModelBlueprint/Public/MVVMWidgetBlueprintExtension_View.h
+++ b/Engine/Plugins/Runtime/ModelViewViewModel/Source/ModelViewViewModelBlueprint/Public/MVVMWidgetBlueprintExtension_View.h
@@ -59,6 +59,7 @@ public:
 	virtual void HandleCreateFunctionList() override;
 	virtual void HandleFinishCompilingClass(UWidgetBlueprintGeneratedClass* Class) override;
 	virtual void HandleEndCompilation() override;
+	virtual FSearchData HandleGatherSearchData(const UBlueprint* OwningBlueprint) const override;
 	//~ End UWidgetBlueprintExtension interface
 
 	void SetFilterSettings(FMVVMViewBindingFilterSettings InFilterSettings);
diff --git a/Engine/Source/Editor/Kismet/Private/FindInBlueprintManager.cpp b/Engine/Source/Editor/Kismet/Private/FindInBlueprintManager.cpp
--- a/Engine/Source/Editor/Kismet/Private/FindInBlueprintManager.cpp
+++ b/Engine/Source/Editor/Kismet/Private/FindInBlueprintManager.cpp
@@ -38,6 +38,7 @@
 #include "WorkspaceMenuStructure.h"
 #include "WorkspaceMenuStructureModule.h"
 
+#include "Blueprint/BlueprintExtension.h"
 #include "Engine/SimpleConstructionScript.h"
 #include "Kismet2/BlueprintEditorUtils.h"
 #include "Kismet2/KismetEditorUtilities.h"
@@ -74,6 +75,7 @@ const FText FFindInBlueprintSearchTags::FiB_UberGraphs = LOCTEXT("Uber", "Uber")
 const FText FFindInBlueprintSearchTags::FiB_Functions = LOCTEXT("Functions", "Functions");
 const FText FFindInBlueprintSearchTags::FiB_Macros = LOCTEXT("Macros", "Macros");
 const FText FFindInBlueprintSearchTags::FiB_SubGraphs = LOCTEXT("Sub", "Sub");
+const FText FFindInBlueprintSearchTags::FiB_Extensions = LOCTEXT("Extensions", "Extensions");
 
 const FText FFindInBlueprintSearchTags::FiB_Name = LOCTEXT("Name", "Name");
 const FText FFindInBlueprintSearchTags::FiB_NativeName = LOCTEXT("NativeName", "Native Name");
@@ -529,12 +531,12 @@ namespace BlueprintSearchMetaDataHelpers
 
 	/** Json Writer used for serializing FText's in the correct format for Find-in-Blueprints */
 	class FFindInBlueprintJsonWriter : public TFindInBlueprintJsonStringWriter<TCondensedJsonPrintPolicy<TCHAR>>
-			{
+	{
 	public:
 		FFindInBlueprintJsonWriter(FString* const InOutString, int32 InFormatVersion)
 			:TFindInBlueprintJsonStringWriter<TCondensedJsonPrintPolicy<TCHAR>>(InOutString, InFormatVersion)
 			,JsonOutput(InOutString)
-				{
+		{
 		}
 
 		virtual bool Close() override
@@ -1192,6 +1194,63 @@ namespace BlueprintSearchMetaDataHelpers
 		}
 	}
 
+	template<class PrintPolicy>
+	void GatherExtensionsSearchData_Recursive(const TSharedRef<TFindInBlueprintJsonStringWriter<PrintPolicy>>& InWriter, const UBlueprintExtension::FSearchData& SearchData)
+	{
+		for (const UBlueprintExtension::FSearchTagDataPair& Data : SearchData.Datas)
+		{
+			InWriter->WriteValue(Data.Key, Data.Value);
+		}
+
+		for (const TUniquePtr<UBlueprintExtension::FSearchArrayData>& ArrayData : SearchData.SearchArrayDatas)
+		{
+			InWriter->WriteArrayStart(ArrayData->Identifier.IsEmpty() ? FFindInBlueprintSearchTags::FiB_Extensions : ArrayData->Identifier);
+			for (const UBlueprintExtension::FSearchData& Data : ArrayData->SearchSubList)
+			{
+				InWriter->WriteObjectStart();
+				GatherExtensionsSearchData_Recursive(InWriter, Data);
+				InWriter->WriteObjectEnd();
+			}
+			InWriter->WriteArrayEnd();
+		}
+	}
+
+	template<class PrintPolicy>
+	void GatherExtensionsSearchData(const TSharedRef<TFindInBlueprintJsonStringWriter<PrintPolicy>>& InWriter, const UBlueprint* InBlueprint)
+	{
+		if (InBlueprint->GetExtensions().Num() > 0)
+		{
+			// Collect all extensions
+			bool bExtensionArrayStarted = false;
+			for (UBlueprintExtension* Extension : InBlueprint->GetExtensions())
+			{
+				if (Extension)
+				{
+					if (!bExtensionArrayStarted)
+					{
+						InWriter->WriteArrayStart(FFindInBlueprintSearchTags::FiB_Extensions);
+						bExtensionArrayStarted = true;
+					}
+
+					InWriter->WriteObjectStart();
+					InWriter->WriteValue(FFindInBlueprintSearchTags::FiB_Name, FText::FromString(Extension->GetName()));
+					InWriter->WriteValue(FFindInBlueprintSearchTags::FiB_ClassName, FText::FromString(Extension->GetClass()->GetName()));
+
+					// Retrieve the custom search metadata from the extension.
+					UBlueprintExtension::FSearchData SearchData = Extension->GatherSearchData(InBlueprint);
+					GatherExtensionsSearchData_Recursive(InWriter, SearchData);
+
+					InWriter->WriteObjectEnd();
+				}
+			}
+
+			if (bExtensionArrayStarted)
+			{
+				InWriter->WriteArrayEnd();
+			}
+		}
+	}
+
 	template<class PrintPolicy>
 	void GatherBlueprintSearchMetadata(const TSharedRef<TFindInBlueprintJsonStringWriter<PrintPolicy>>& InWriter, const UBlueprint* Blueprint)
 	{
@@ -1273,6 +1332,8 @@ namespace BlueprintSearchMetaDataHelpers
 			InWriter->WriteArrayEnd(); // Components
 		}
 
+		GatherExtensionsSearchData(InWriter, Blueprint);
+
 		InWriter->WriteObjectEnd();
 		InWriter->Close();
 	}
diff --git a/Engine/Source/Editor/Kismet/Public/FindInBlueprintManager.h b/Engine/Source/Editor/Kismet/Public/FindInBlueprintManager.h
--- a/Engine/Source/Editor/Kismet/Public/FindInBlueprintManager.h
+++ b/Engine/Source/Editor/Kismet/Public/FindInBlueprintManager.h
@@ -90,6 +90,8 @@ struct KISMET_API FFindInBlueprintSearchTags
 	static const FText FiB_Macros;
 	/** Sub graph tag, for any sub-graphs in a Blueprint */
 	static const FText FiB_SubGraphs;
+	/** Blueprint extension tag. */
+	static const FText FiB_Extensions;
 
 	/** Name tag */
 	static const FText FiB_Name;
diff --git a/Engine/Source/Runtime/Engine/Public/Blueprint/BlueprintExtension.h b/Engine/Source/Runtime/Engine/Public/Blueprint/BlueprintExtension.h
--- a/Engine/Source/Runtime/Engine/Public/Blueprint/BlueprintExtension.h
+++ b/Engine/Source/Runtime/Engine/Public/Blueprint/BlueprintExtension.h
@@ -23,6 +23,29 @@ public:
 
 #if WITH_EDITORONLY_DATA
 
+	/** Used to add Key/Value pairs in the FindInBlueprintManager */
+	struct FSearchTagDataPair
+	{
+		FSearchTagDataPair(FText InKey, FText InValue)
+			: Key(InKey)
+			, Value(InValue)
+		{}
+
+		FText Key;
+		FText Value;
+	};
+	struct FSearchArrayData;
+	struct FSearchData
+	{
+		TArray<FSearchTagDataPair> Datas;
+		TArray<TUniquePtr<FSearchArrayData>> SearchArrayDatas;
+	};
+	struct FSearchArrayData
+	{
+		FText Identifier;
+		TArray<FSearchData> SearchSubList;
+	};
+
 	/**
 	 * Called during compilation (after skeleton class generation) in order to generate dynamic function graphs for this blueprint
 	 */
@@ -38,6 +61,14 @@ public:
 	{
 		HandlePreloadObjectsForCompilation(OwningBlueprint);
 	}
+	
+	/**
+	 * Called when the find in blueprint system gather the search data for the blueprint
+	 */
+	FSearchData GatherSearchData(const UBlueprint* OwningBlueprint) const
+	{
+		return HandleGatherSearchData(OwningBlueprint);
+	}
 
 private:
 
@@ -51,5 +82,10 @@ private:
 	 */
 	virtual void HandleGenerateFunctionGraphs(FKismetCompilerContext* CompilerContext) {}
 
+	/**
+	 * Override this function to define custom function generation logic for gathering search data
+	 */
+	virtual FSearchData HandleGatherSearchData(const UBlueprint* OwningBlueprint) const { return FSearchData(); }
+
 #endif // WITH_EDITORONLY_DATA
 };
