Changing corner text in the Blueprint Editor when we're in an Editor Utility Blueprint

diff --git a/Engine/Source/Editor/Kismet/Private/BlueprintEditor.cpp b/Engine/Source/Editor/Kismet/Private/BlueprintEditor.cpp
--- a/Engine/Source/Editor/Kismet/Private/BlueprintEditor.cpp
+++ b/Engine/Source/Editor/Kismet/Private/BlueprintEditor.cpp
@@ -1748,21 +1748,28 @@ FGraphAppearanceInfo FBlueprintEditor::GetGraphAppearance(UEdGraph* InGraph) con
 	UBlueprint* Blueprint = (InGraph != nullptr) ? FBlueprintEditorUtils::FindBlueprintForGraph(InGraph) : GetBlueprintObj();
 	if (Blueprint != nullptr)
 	{
-	switch (Blueprint->BlueprintType)
-	{
-	case BPTYPE_LevelScript:
-		AppearanceInfo.CornerText = LOCTEXT("AppearanceCornerText_LevelScript", "LEVEL BLUEPRINT");
-		break;
-	case BPTYPE_MacroLibrary:
-		AppearanceInfo.CornerText = LOCTEXT("AppearanceCornerText_Macro", "MACRO");
-		break;
-	case BPTYPE_Interface:
-		AppearanceInfo.CornerText = LOCTEXT("AppearanceCornerText_Interface", "INTERFACE");
-		break;
-	default:
-		AppearanceInfo.CornerText = LOCTEXT("AppearanceCornerText_Blueprint", "BLUEPRINT");
-		break;
-	}
+		if (FBlueprintEditorUtils::IsEditorUtilityBlueprint(Blueprint))
+		{
+			AppearanceInfo.CornerText = LOCTEXT("EditorUtilityAppearanceCornerText", "EDITOR UTILITY");
+		}
+		else
+		{
+			switch (Blueprint->BlueprintType)
+			{
+			case BPTYPE_LevelScript:
+				AppearanceInfo.CornerText = LOCTEXT("AppearanceCornerText_LevelScript", "LEVEL BLUEPRINT");
+				break;
+			case BPTYPE_MacroLibrary:
+				AppearanceInfo.CornerText = LOCTEXT("AppearanceCornerText_Macro", "MACRO");
+				break;
+			case BPTYPE_Interface:
+				AppearanceInfo.CornerText = LOCTEXT("AppearanceCornerText_Interface", "INTERFACE");
+				break;
+			default:
+				AppearanceInfo.CornerText = LOCTEXT("AppearanceCornerText_Blueprint", "BLUEPRINT");
+				break;
+			}
+		}
 	}
 
 	UEdGraph const* EditingGraph = GetFocusedGraph();
