Don't replace references on transaction buffer unless on main thread
#rnx

diff --git a/Engine/Source/Editor/Kismet/Private/BlueprintCompilationManager.cpp b/Engine/Source/Editor/Kismet/Private/BlueprintCompilationManager.cpp
--- a/Engine/Source/Editor/Kismet/Private/BlueprintCompilationManager.cpp
+++ b/Engine/Source/Editor/Kismet/Private/BlueprintCompilationManager.cpp
@@ -2456,7 +2456,7 @@ void FBlueprintCompilationManagerImpl::ReinstanceBatch(TArray<FReinstancingJob>&
 
 	// The transaction buffer could reference archetypes, and tag serialization
 	// will be simpler if we update the instance:
-	if(GUnrealEd && GUnrealEd->Trans)
+	if (IsInGameThread() && GUnrealEd && GUnrealEd->Trans)
 	{
 		ArchetypeReferencers.Add(GUnrealEd->Trans);
 	}
