Added TCallTraits::ConstPointerType overrides for common smart pointer types so that they you can look up a TSmartPtr<const T> in a map or set with a TSmartPtr<T> key.
Fixed some const incorrectness caused as a fallout.

diff --git a/Engine/Plugins/Runtime/Metasound/Source/MetasoundEditor/Private/SMetasoundPinValueInspector.cpp b/Engine/Plugins/Runtime/Metasound/Source/MetasoundEditor/Private/SMetasoundPinValueInspector.cpp
--- a/Engine/Plugins/Runtime/Metasound/Source/MetasoundEditor/Private/SMetasoundPinValueInspector.cpp
+++ b/Engine/Plugins/Runtime/Metasound/Source/MetasoundEditor/Private/SMetasoundPinValueInspector.cpp
@@ -625,7 +625,7 @@ namespace Metasound
 			return Message == Other->Message;
 		}
 
-		uint32 FMetasoundNumericDebugLineItem::GetHash()
+		uint32 FMetasoundNumericDebugLineItem::GetHash() const
 		{
 			return GetTypeHash(Message);
 		}
diff --git a/Engine/Plugins/Runtime/Metasound/Source/MetasoundEditor/Private/SMetasoundPinValueInspector.h b/Engine/Plugins/Runtime/Metasound/Source/MetasoundEditor/Private/SMetasoundPinValueInspector.h
--- a/Engine/Plugins/Runtime/Metasound/Source/MetasoundEditor/Private/SMetasoundPinValueInspector.h
+++ b/Engine/Plugins/Runtime/Metasound/Source/MetasoundEditor/Private/SMetasoundPinValueInspector.h
@@ -88,7 +88,7 @@ namespace Metasound
 
 			virtual bool Compare(const FDebugLineItem* BaseOther) const override;
 
-			virtual uint32 GetHash() override;
+			virtual uint32 GetHash() const override;
 
 		private:
 			FLinearColor GetEdgeStyleColorAtIndex(int32 InIndex) const;
diff --git a/Engine/Source/Editor/Kismet/Private/Debugging/SKismetDebugTreeView.cpp b/Engine/Source/Editor/Kismet/Private/Debugging/SKismetDebugTreeView.cpp
--- a/Engine/Source/Editor/Kismet/Private/Debugging/SKismetDebugTreeView.cpp
+++ b/Engine/Source/Editor/Kismet/Private/Debugging/SKismetDebugTreeView.cpp
@@ -592,8 +592,8 @@ public:
 		*/
 		static FORCEINLINE bool Matches(KeyInitType A, KeyInitType B)
 		{
-			FDebugLineItem* APtr = A.Get();
-			FDebugLineItem* BPtr = B.Get();
+			const FDebugLineItem* APtr = A.Get();
+			const FDebugLineItem* BPtr = B.Get();
 			if (APtr && BPtr)
 			{
 				return (APtr->Type == BPtr->Type) && APtr->Compare(BPtr);
@@ -604,7 +604,7 @@ public:
 		/** Calculates a hash index for a key. */
 		static FORCEINLINE uint32 GetKeyHash(KeyInitType Key)
 		{
-			if (FDebugLineItem* KeyPtr = Key.Get())
+			if (const FDebugLineItem* KeyPtr = Key.Get())
 			{
 				return KeyPtr->GetHash();
 			}
@@ -757,7 +757,7 @@ protected:
 		return Message == Other->Message;
 	}
 
-	virtual uint32 GetHash() override
+	virtual uint32 GetHash() const override
 	{
 		return GetTypeHash(Message);
 	}
@@ -798,7 +798,7 @@ protected:
 			(UUID == Other->UUID);
 	}
 
-	virtual uint32 GetHash() override
+	virtual uint32 GetHash() const override
 	{
 		return HashCombine(GetTypeHash(UUID), GetTypeHash(ParentObjectRef));
 	}
@@ -1061,7 +1061,7 @@ protected:
 			Data->DisplayName.CompareTo(Other->Data->DisplayName) == 0;
 	}
 
-	virtual uint32 GetHash() override
+	virtual uint32 GetHash() const override
 	{
 		return HashCombine(GetTypeHash(Data->Property), GetTypeHash(Data->DisplayName.ToString()));
 	}
@@ -1272,7 +1272,7 @@ protected:
 		return (ObjectToWatch.Get() == Other->ObjectToWatch.Get());
 	}
 
-	virtual uint32 GetHash() override
+	virtual uint32 GetHash() const override
 	{
 		return GetTypeHash(ObjectToWatch);
 	}
@@ -1409,7 +1409,7 @@ protected:
 			(PathToProperty == Other->PathToProperty);
 	}
 
-	virtual uint32 GetHash() override
+	virtual uint32 GetHash() const override
 	{
 		return HashCombine(GetTypeHash(ParentObjectRef), GetTypeHash(ObjectRef));
 	}
@@ -2181,7 +2181,7 @@ protected:
 			(BreakpointNode == Other->BreakpointNode);
 	}
 
-	virtual uint32 GetHash() override
+	virtual uint32 GetHash() const override
 	{
 		return HashCombine(GetTypeHash(ParentObjectRef), GetTypeHash(BreakpointNode));
 	}
@@ -2398,7 +2398,7 @@ protected:
 		return false;
 	}
 
-	virtual uint32 GetHash() override
+	virtual uint32 GetHash() const override
 	{
 		check(false);
 		return 0;
@@ -2601,7 +2601,7 @@ protected:
 		return ObjectRef.Get() == Other->ObjectRef.Get();
 	}
 
-	virtual uint32 GetHash() override
+	virtual uint32 GetHash() const override
 	{
 		return GetTypeHash(ObjectRef);
 	}
@@ -2731,7 +2731,7 @@ protected:
 		return false;
 	}
 
-	virtual uint32 GetHash() override
+	virtual uint32 GetHash() const override
 	{
 		check(false);
 		return 0;
@@ -2850,7 +2850,7 @@ protected:
 		return false;
 	}
 
-	virtual uint32 GetHash() override
+	virtual uint32 GetHash() const override
 	{
 		check(false);
 		return 0;
diff --git a/Engine/Source/Editor/Kismet/Public/Debugging/SKismetDebugTreeView.h b/Engine/Source/Editor/Kismet/Public/Debugging/SKismetDebugTreeView.h
--- a/Engine/Source/Editor/Kismet/Public/Debugging/SKismetDebugTreeView.h
+++ b/Engine/Source/Editor/Kismet/Public/Debugging/SKismetDebugTreeView.h
@@ -152,7 +152,7 @@ protected:
 	virtual bool Compare(const FDebugLineItem* Other) const = 0;
 
 	// used for sets
-	virtual uint32 GetHash() = 0;
+	virtual uint32 GetHash() const = 0;
 
 	// Used to update the state of a line item rather than replace it.
 	// called after Compare returns true
diff --git a/Engine/Source/Runtime/Core/Public/Templates/SharedPointer.h b/Engine/Source/Runtime/Core/Public/Templates/SharedPointer.h
--- a/Engine/Source/Runtime/Core/Public/Templates/SharedPointer.h
+++ b/Engine/Source/Runtime/Core/Public/Templates/SharedPointer.h
@@ -652,6 +652,12 @@ template <class ObjectType, ESPMode InMode> constexpr bool TIsTSharedRef_V<const
 template <class ObjectType, ESPMode InMode> constexpr bool TIsTSharedRef_V<      volatile TSharedRef<ObjectType, InMode>> = true;
 template <class ObjectType, ESPMode InMode> constexpr bool TIsTSharedRef_V<const volatile TSharedRef<ObjectType, InMode>> = true;
 
+template <class ObjectType, ESPMode Mode>
+struct TCallTraits<TSharedRef<ObjectType, Mode>> : public TCallTraitsBase<TSharedRef<ObjectType, Mode>>
+{
+	using ConstPointerType = TSharedRef<const ObjectType, Mode>;
+};
+
 
 /**
  * Wrapper for a type that yields a reference to that type.
@@ -1236,6 +1242,12 @@ DECLARE_TEMPLATE_INTRINSIC_TYPE_LAYOUT((template<class ObjectType, ESPMode Mode>
 
 template<class ObjectType, ESPMode Mode> struct TIsZeroConstructType<TSharedPtr<ObjectType, Mode>> { enum { Value = true }; };
 
+template <class ObjectType, ESPMode Mode>
+struct TCallTraits<TSharedPtr<ObjectType, Mode>> : public TCallTraitsBase<TSharedPtr<ObjectType, Mode>>
+{
+	using ConstPointerType = TSharedPtr<const ObjectType, Mode>;
+};
+
 
 /**
  * TWeakPtr is a non-intrusive reference-counted weak object pointer.  This weak pointer will be
@@ -1542,6 +1554,12 @@ template <class ObjectType, ESPMode InMode> constexpr bool TIsTWeakPtr_V<const v
 template<class T, ESPMode Mode> struct TIsWeakPointerType<TWeakPtr<T, Mode> > { enum { Value = true }; };
 template<class T, ESPMode Mode> struct TIsZeroConstructType<TWeakPtr<T, Mode> > { enum { Value = true }; };
 
+template <class ObjectType, ESPMode Mode>
+struct TCallTraits<TWeakPtr<ObjectType, Mode>> : public TCallTraitsBase<TWeakPtr<ObjectType, Mode>>
+{
+	using ConstPointerType = TWeakPtr<const ObjectType, Mode>;
+};
+
 
 /**
  * Derive your class from TSharedFromThis to enable access to a TSharedRef directly from an object
diff --git a/Engine/Source/Runtime/Core/Public/UObject/WeakObjectPtrTemplates.h b/Engine/Source/Runtime/Core/Public/UObject/WeakObjectPtrTemplates.h
--- a/Engine/Source/Runtime/Core/Public/UObject/WeakObjectPtrTemplates.h
+++ b/Engine/Source/Runtime/Core/Public/UObject/WeakObjectPtrTemplates.h
@@ -396,6 +396,12 @@ template<class T> struct TIsPODType<TAutoWeakObjectPtr<T> > { enum { Value = tru
 template<class T> struct TIsZeroConstructType<TAutoWeakObjectPtr<T> > { enum { Value = true }; };
 template<class T> struct TIsWeakPointerType<TAutoWeakObjectPtr<T> > { enum { Value = true }; };
 
+template <typename T>
+struct TCallTraits<TWeakObjectPtr<T>> : public TCallTraitsBase<TWeakObjectPtr<T>>
+{
+	using ConstPointerType = TWeakObjectPtr<const T>;
+};
+
 /** Utility function to fill in a TArray<ClassName*> from a TArray<TWeakObjectPtr<ClassName>> */
 template<typename DestArrayType, typename SourceArrayType>
 void CopyFromWeakArray(DestArrayType& Dest, const SourceArrayType& Src)
diff --git a/Engine/Source/Runtime/CoreUObject/Public/Templates/SubclassOf.h b/Engine/Source/Runtime/CoreUObject/Public/Templates/SubclassOf.h
--- a/Engine/Source/Runtime/CoreUObject/Public/Templates/SubclassOf.h
+++ b/Engine/Source/Runtime/CoreUObject/Public/Templates/SubclassOf.h
@@ -170,6 +170,12 @@ private:
 	UClass* Class = nullptr;
 };
 
+template <typename T>
+struct TCallTraits<TSubclassOf<T>> : public TCallTraitsBase<TSubclassOf<T>>
+{
+	using ConstPointerType = TSubclassOf<const T>;
+};
+
 template <typename T>
 FArchive& operator<<(FArchive& Ar, TSubclassOf<T>& SubclassOf)
 {
diff --git a/Engine/Source/Runtime/CoreUObject/Public/UObject/ScriptInterface.h b/Engine/Source/Runtime/CoreUObject/Public/UObject/ScriptInterface.h
--- a/Engine/Source/Runtime/CoreUObject/Public/UObject/ScriptInterface.h
+++ b/Engine/Source/Runtime/CoreUObject/Public/UObject/ScriptInterface.h
@@ -146,7 +146,7 @@ public:
 	TScriptInterface(TYPE_OF_NULLPTR) {}
 
 	/**
-	 * Assignment from an object type that implements the InterfaceType native interface class
+	 * Construction from an object type that implements the InterfaceType native interface class
 	 */
 	template <
 		typename U,
@@ -161,6 +161,21 @@ public:
 		SetInterface(SourceInterface);
 	}
 
+	/**
+	 * Construction from another script interface of a compatible interface type
+	 */
+	template <
+		typename OtherInterfaceType,
+		decltype(ImplicitConv<InterfaceType*>(std::declval<OtherInterfaceType>()))* = nullptr
+	>
+	FORCEINLINE TScriptInterface(const TScriptInterface<OtherInterfaceType>& Other)
+	{
+		SetObject(Other.GetObject());
+
+		InterfaceType* SourceInterface = Other.GetInterface();
+		SetInterface(SourceInterface);
+	}
+
 	/**
 	 * Assignment from an object type that implements the InterfaceType native interface class
 	 */
@@ -173,6 +188,12 @@ public:
 		SetInterface(SourceInterface);
 	}
 
+	template <typename OtherInterfaceType>
+	TScriptInterface(const TScriptInterface<OtherInterfaceType>& Other)
+	{
+
+	}
+
 	/**
 	 * Copyable
 	 */
@@ -201,6 +222,19 @@ public:
 		return *this;
 	}
 
+	/**
+	 * Assignment from another script interface of a compatible interface type
+	 */
+	template <
+		typename OtherInterfaceType,
+		decltype(ImplicitConv<InterfaceType*>(std::declval<OtherInterfaceType>()))* = nullptr
+	>
+	TScriptInterface& operator=(const TScriptInterface<OtherInterfaceType>& Other)
+	{
+		*this = TScriptInterface(Other);
+		return *this;
+	}
+
 	/**
 	 * Assignment from an object type that implements the InterfaceType native interface class
 	 */
@@ -299,6 +333,12 @@ public:
 	}
 };
 
+template <typename InterfaceType>
+struct TCallTraits<TScriptInterface<InterfaceType>> : public TCallTraitsBase<TScriptInterface<InterfaceType>>
+{
+	using ConstPointerType = TScriptInterface<const InterfaceType>;
+};
+
 #if UE_ENABLE_INCLUDE_ORDER_DEPRECATED_IN_5_2
 #include "CoreMinimal.h"
 #endif
diff --git a/Engine/Source/Runtime/CoreUObject/Public/UObject/SoftObjectPtr.h b/Engine/Source/Runtime/CoreUObject/Public/UObject/SoftObjectPtr.h
--- a/Engine/Source/Runtime/CoreUObject/Public/UObject/SoftObjectPtr.h
+++ b/Engine/Source/Runtime/CoreUObject/Public/UObject/SoftObjectPtr.h
@@ -412,6 +412,12 @@ FArchive& operator<<(FArchive& Ar, TSoftObjectPtr<T>& Ptr)
 template<class T> struct TIsPODType<TSoftObjectPtr<T> > { enum { Value = TIsPODType<FSoftObjectPtr>::Value }; };
 template<class T> struct TIsWeakPointerType<TSoftObjectPtr<T> > { enum { Value = TIsWeakPointerType<FSoftObjectPtr>::Value }; };
 
+template <typename T>
+struct TCallTraits<TSoftObjectPtr<T>> : public TCallTraitsBase<TSoftObjectPtr<T>>
+{
+	using ConstPointerType = TSoftObjectPtr<const T>;
+};
+
 /** Utility to create a TSoftObjectPtr without specifying the type */
 template <class T>
 TSoftObjectPtr<std::remove_cv_t<T>> MakeSoftObjectPtr(T* Object)
@@ -652,6 +658,12 @@ private:
 template <class T> struct TIsPODType<TSoftClassPtr<T> > { enum { Value = TIsPODType<FSoftObjectPtr>::Value }; };
 template <class T> struct TIsWeakPointerType<TSoftClassPtr<T> > { enum { Value = TIsWeakPointerType<FSoftObjectPtr>::Value }; };
 
+template <typename T>
+struct TCallTraits<TSoftClassPtr<T>> : public TCallTraitsBase<TSoftClassPtr<T>>
+{
+	using ConstPointerType = TSoftClassPtr<const T>;
+};
+
 /** Utility to create a TSoftObjectPtr without specifying the type */
 template <class T>
 TSoftClassPtr<std::remove_cv_t<T>> MakeSoftClassPtr(T* Object)
