[EnsureFix] RefreshPropertyObjects in SKismetInspector wasn't being garbage collected and the invalid objects were triggering an ensure in SKismetInspector::UpdateFromObjects.

diff --git a/Engine/Source/Editor/Kismet/Private/SKismetInspector.cpp b/Engine/Source/Editor/Kismet/Private/SKismetInspector.cpp
--- a/Engine/Source/Editor/Kismet/Private/SKismetInspector.cpp
+++ b/Engine/Source/Editor/Kismet/Private/SKismetInspector.cpp
@@ -166,6 +166,7 @@ void SKismetInspector::Tick( const FGeometry& AllottedGeometry, const double InC
 		else
 		{
 			FKismetSelectionInfo SelectionInfo;
+			RefreshPropertyObjects.Remove(nullptr);
 			UpdateFromObjects(RefreshPropertyObjects, SelectionInfo, RefreshOptions);
 			RefreshPropertyObjects.Empty();
 		}
@@ -284,6 +285,16 @@ const TArray< TWeakObjectPtr<UObject> >& SKismetInspector::GetSelectedObjects()
 	return SelectedObjects;
 }
 
+void SKismetInspector::AddReferencedObjects(FReferenceCollector& Collector)
+{
+	Collector.AddReferencedObjects(RefreshPropertyObjects);
+}
+
+FString SKismetInspector::GetReferencerName() const
+{
+	return TEXT("SKismetInspector");
+}
+
 FText SKismetInspector::GetContextualEditingWidgetTitle() const
 {
 	FText Title = PropertyViewTitle;
diff --git a/Engine/Source/Editor/Kismet/Public/SKismetInspector.h b/Engine/Source/Editor/Kismet/Public/SKismetInspector.h
--- a/Engine/Source/Editor/Kismet/Public/SKismetInspector.h
+++ b/Engine/Source/Editor/Kismet/Public/SKismetInspector.h
@@ -41,7 +41,7 @@ typedef TSet<class UObject*> FInspectorSelectionSet;
 // SKismetInspector
 
 /** Widget that shows properties and tools related to the selected node(s) */
-class KISMET_API SKismetInspector : public SCompoundWidget
+class KISMET_API SKismetInspector : public SCompoundWidget, public FGCObject
 {
 public:
 	SLATE_BEGIN_ARGS( SKismetInspector )
@@ -63,7 +63,7 @@ public:
 		SLATE_ARGUMENT( bool, ShowTitleArea)
 		SLATE_ARGUMENT( bool, ShowLocalVariables)
 	SLATE_END_ARGS()
-
+	
 	void Construct(const FArguments& InArgs);
 
 	/** Options for ShowDetails */
@@ -119,6 +119,9 @@ public:
 	/** returns the list of selected objects */
 	const TArray< TWeakObjectPtr<UObject> >& GetSelectedObjects() const;
 
+	virtual void AddReferencedObjects(FReferenceCollector& Collector) override;
+	virtual FString GetReferencerName() const override;
+
 protected:
 	/** Update the inspector window to show information on the supplied objects */
 	void UpdateFromObjects(const TArray<UObject*>& PropertyObjects, struct FKismetSelectionInfo& SelectionInfo, const FShowDetailsOptions& Options);
