Fix the Blueprint Namespace entry widget's tree view not working properly when only one root item matches the search box filter text.

diff --git a/Engine/Source/Editor/Kismet/Private/SBlueprintNamespaceEntry.cpp b/Engine/Source/Editor/Kismet/Private/SBlueprintNamespaceEntry.cpp
--- a/Engine/Source/Editor/Kismet/Private/SBlueprintNamespaceEntry.cpp
+++ b/Engine/Source/Editor/Kismet/Private/SBlueprintNamespaceEntry.cpp
@@ -225,6 +225,9 @@ TSharedRef<SWidget> SBlueprintNamespaceEntry::OnGetNamespaceTreeMenuContent()
 	// Sort the list alphabetically.
 	AllRegisteredPaths.Sort();
 
+	// Reset the search box so we don't reapply a previous menu's filter.
+	SearchBox.Reset();
+
 	// Build the namespace item tree from the filtered list of registered paths.
 	PopulateNamespaceTree();
 
@@ -284,7 +287,7 @@ TSharedRef<ITableRow> SBlueprintNamespaceEntry::OnGenerateRowForNamespaceTreeIte
 	// Check for an empty tree and add a single (disabled) item if found.
 	bool bIsEnabled = true;
 	FText ItemText;
-	if (ItemText.IsEmpty() && RootNodes.Num() == 1)
+	if (Item->NodePath.IsEmpty() && RootNodes.Num() == 1 && RootNodes[0] == Item)
 	{
 		bIsEnabled = false;
 		ItemText = LOCTEXT("BlueprintNamespaceList_NoItems", "No Matching Items");
