Add field notify checkbox to BP functions
- Generate a list of field notify properties to choose from (the list only shows when the selected property itself is field notify)

diff --git a/Engine/Source/Editor/EditorStyle/Private/StarshipStyle.cpp b/Engine/Source/Editor/EditorStyle/Private/StarshipStyle.cpp
--- a/Engine/Source/Editor/EditorStyle/Private/StarshipStyle.cpp
+++ b/Engine/Source/Editor/EditorStyle/Private/StarshipStyle.cpp
@@ -5049,6 +5049,9 @@ void FStarshipEditorStyle::FStyle::SetupPersonaStyle()
 		Set("Kismet.VariableList.ExposeForInstance", new CORE_IMAGE_BRUSH_SVG("Starship/Common/visible", Icon16x16));
 		Set("Kismet.VariableList.HideForInstance", new CORE_IMAGE_BRUSH_SVG("Starship/Common/hidden", Icon16x16));
 
+		Set("Kismet.VariableList.FieldNotify", new CORE_IMAGE_BRUSH_SVG("Starship/Common/fieldnotify_on", Icon16x16));
+		Set("Kismet.VariableList.NotFieldNotify", new CORE_IMAGE_BRUSH_SVG("Starship/Common/fieldnotify_off", Icon16x16));
+
 		Set( "Kismet.Explorer.Title", FTextBlockStyle(NormalText) .SetFont(DEFAULT_FONT( "BoldCondensedItalic", 11)));
 		Set( "Kismet.Explorer.SearchDepthFont", DEFAULT_FONT( "Bold", 14) );
 
diff --git a/Engine/Source/Editor/Kismet/Private/BlueprintDetailsCustomization.cpp b/Engine/Source/Editor/Kismet/Private/BlueprintDetailsCustomization.cpp
--- a/Engine/Source/Editor/Kismet/Private/BlueprintDetailsCustomization.cpp
+++ b/Engine/Source/Editor/Kismet/Private/BlueprintDetailsCustomization.cpp
@@ -99,6 +99,7 @@
 #include "PropertyHandle.h"
 #include "PropertyRestriction.h"
 #include "SBlueprintNamespaceEntry.h"
+#include "SFieldNotificationCheckList.h"
 #include "SGraphPin.h"
 #include "SKismetInspector.h"
 #include "SPinTypeSelector.h"
@@ -522,7 +523,7 @@ void FBlueprintVarActionDetails::CustomizeDetails( IDetailLayoutBuilder& DetailL
 		.ToolTip(ReadOnlyTooltip)
 	];
 	
-	if (BlueprintPtr && FBlueprintEditorUtils::ImplementsInterface(BlueprintPtr, true, UNotifyFieldValueChanged::StaticClass()))
+	if (BlueprintPtr && FBlueprintEditorUtils::ImplementsInterface(BlueprintPtr, true, UNotifyFieldValueChanged::StaticClass()) && !MyBlueprint.Pin()->SelectionAsDelegate())
 	{
 		// Show the flag if the class implement the interface but only allow the flag to be changed if the variable is defined in BP
 		const FText ToolTip = LOCTEXT("FieldNotifyToolTip", "Generate a field entry for the Field Notification system.");
@@ -538,11 +539,25 @@ void FBlueprintVarActionDetails::CustomizeDetails( IDetailLayoutBuilder& DetailL
 			]
 			.ValueContent()
 			[
-				SNew(SCheckBox)
-				.IsChecked(this, &FBlueprintVarActionDetails::OnFieldNotifyCheckboxState)
-				.OnCheckStateChanged(this, &FBlueprintVarActionDetails::OnFieldNotifyChanged)
-				.IsEnabled(GetPropertyOwnerBlueprint() && IsABlueprintVariable(VariableProperty) && IsAUserVariable(VariableProperty))
-				.ToolTip(FieldNotificationTooltip)
+				SNew(SHorizontalBox)
+				+ SHorizontalBox::Slot()
+				.AutoWidth()
+				[
+					SNew(SCheckBox)
+					.IsChecked(this, &FBlueprintVarActionDetails::OnFieldNotifyCheckboxState)
+					.OnCheckStateChanged(this, &FBlueprintVarActionDetails::OnFieldNotifyChanged)
+					.IsEnabled(GetPropertyOwnerBlueprint() && IsABlueprintVariable(VariableProperty) && IsAUserVariable(VariableProperty))
+					.ToolTip(FieldNotificationTooltip)
+				]
+				+ SHorizontalBox::Slot()
+				.AutoWidth()
+				[
+					SNew(UE::FieldNotification::SFieldNotificationCheckList)
+					.FieldName(CachedVariableName)
+					.BlueprintPtr(BlueprintPtr)
+					.Visibility(this, &FBlueprintVarActionDetails::GetFieldNotifyCheckboxListVisiblity)
+				]
+
 			];
 	}
 
@@ -1942,10 +1957,28 @@ void FBlueprintVarActionDetails::OnFieldNotifyChanged(ECheckBoxState InNewState)
 	}
 	else
 	{
+		FBlueprintEditorUtils::RemoveFieldNotifyFromAllMetadata(BlueprintObj, VarName);
 		FBlueprintEditorUtils::RemoveBlueprintVariableMetaData(GetBlueprintObj(), VarName, GetLocalVariableScope(CachedVariableProperty.Get()), FBlueprintMetadata::MD_FieldNotify);
 	}
 }
 
+EVisibility FBlueprintVarActionDetails::GetFieldNotifyCheckboxListVisiblity() const
+{
+	UBlueprint* const BlueprintObj = GetBlueprintObj();
+	const FName VarName = CachedVariableName;
+
+	if (BlueprintObj && !VarName.IsNone())
+	{
+		// Only show the list of checkboxes in the details panel when this variable is field notify.
+		const int32 VarIndex = FBlueprintEditorUtils::FindNewVariableIndex(BlueprintObj, VarName);
+		if (VarIndex != INDEX_NONE)
+		{
+			return BlueprintObj->NewVariables[VarIndex].HasMetaData(FBlueprintMetadata::MD_FieldNotify) ? EVisibility::Visible : EVisibility::Collapsed;
+		}
+	}
+	return EVisibility::Collapsed;
+}
+
 ECheckBoxState FBlueprintVarActionDetails::OnCreateWidgetCheckboxState() const
 {
 	FProperty* Property = CachedVariableProperty.Get();
@@ -4045,6 +4078,32 @@ void FBlueprintGraphActionDetails::CustomizeDetails( IDetailLayoutBuilder& Detai
 						.Font( IDetailLayoutBuilder::GetDetailFont() )
 					];
 			}
+			 
+			UBlueprint* BlueprintPtr = GetBlueprintObj();
+			
+			if (BlueprintPtr && IsFieldNotifyCheckVisible())
+			{
+				const FText ToolTip = LOCTEXT("FieldNotifyToolTip", "Generate a field entry for the Field Notification system.");
+				const FString DocLink = TEXT("Shared/Editors/BlueprintEditor/GraphDetails");
+				TSharedPtr<SToolTip> FieldNotificationTooltip = IDocumentation::Get()->CreateToolTip(LOCTEXT("FieldNotifyToolTip", "Generate a field entry for the Field Notification system."), NULL, DocLink, TEXT("FieldNotify"));
+
+				Category.AddCustomRow(LOCTEXT("IsFunctionFieldNotifyLabel", "Field Notify"))
+					.NameContent()
+					[
+						SNew(STextBlock)
+						.Text(LOCTEXT("IsFunctionFieldNotifyLabel", "Field Notify"))
+						.ToolTip(FieldNotificationTooltip)
+						.Font(IDetailLayoutBuilder::GetDetailFont())
+					]
+					.ValueContent()
+					[
+						SNew(SCheckBox)
+						.IsChecked(this, &FBlueprintGraphActionDetails::OnFieldNotifyCheckboxState)
+						.OnCheckStateChanged(this, &FBlueprintGraphActionDetails::OnFieldNotifyChanged)
+						.IsEnabled(this, &FBlueprintGraphActionDetails::GetIsFieldNotfyEnabled)
+						.ToolTip(FieldNotificationTooltip)
+					];
+			}
 
 			if (GetInstanceColorVisibility())
 			{
@@ -4556,6 +4615,81 @@ void FBlueprintGraphActionDetails::CustomizeDetails( IDetailLayoutBuilder& Detai
 }
 END_SLATE_FUNCTION_BUILD_OPTIMIZATION
 
+bool FBlueprintGraphActionDetails::IsFieldNotifyCheckVisible() const
+{
+	bool bSupportedType = false;
+	bool bIsEditable = false;
+	bool bImplementsFieldNotify = false;
+	UK2Node_EditablePinBase* FunctionEntryNode = FunctionEntryNodePtr.Get();
+	UFunction* Function = FindFunction();
+
+	if (FunctionEntryNode && Function)
+	{
+		UBlueprint* Blueprint = FunctionEntryNode->GetBlueprint();
+		const bool bIsInterface = FBlueprintEditorUtils::IsInterfaceBlueprint(Blueprint);
+		bImplementsFieldNotify = FBlueprintEditorUtils::ImplementsInterface(Blueprint, true, UNotifyFieldValueChanged::StaticClass());
+
+		bSupportedType = !bIsInterface && FunctionEntryNode->IsA<UK2Node_FunctionEntry>();
+		bIsEditable = FunctionEntryNode->IsEditable();
+	}
+	return bSupportedType && bIsEditable && bImplementsFieldNotify;
+}
+
+bool FBlueprintGraphActionDetails::GetIsFieldNotfyEnabled() const
+{
+	UK2Node_EditablePinBase* FunctionEntryNode = FunctionEntryNodePtr.Get();
+	UK2Node_EditablePinBase* FunctionResultNode = FunctionResultNodePtr.Get();
+
+	if (FunctionEntryNode && FunctionResultNode)
+	{
+		return GetIsConstFunction() == ECheckBoxState::Checked && GetIsPureFunction() == ECheckBoxState::Checked && FunctionEntryNode->GetAllPins().Num() == 1 && FunctionResultNode->GetAllPins().Num() == 2;
+	}
+	return false;
+}
+
+ECheckBoxState FBlueprintGraphActionDetails::OnFieldNotifyCheckboxState() const
+{
+	UBlueprint* const BlueprintObj = GetBlueprintObj();
+	const FName FuncName = FindFunction()->GetFName();
+
+	if (BlueprintObj && GetIsFieldNotfyEnabled())
+	{
+		if (!FuncName.IsNone())
+		{
+			return GetMetadataBlock()->HasMetaData(FBlueprintMetadata::MD_FieldNotify) ? ECheckBoxState::Checked : ECheckBoxState::Unchecked;
+		}
+		else if (BlueprintObj->GeneratedClass && BlueprintObj->GeneratedClass->ImplementsInterface(UNotifyFieldValueChanged::StaticClass()) && BlueprintObj->GeneratedClass->GetDefaultObject())
+		{
+			TScriptInterface<INotifyFieldValueChanged> DefaultObject = BlueprintObj->GeneratedClass->GetDefaultObject();
+			return DefaultObject->GetFieldNotificationDescriptor().GetField(BlueprintObj->GeneratedClass, FuncName).IsValid() ? ECheckBoxState::Checked : ECheckBoxState::Unchecked;
+		}
+	}
+
+	FBlueprintEditorUtils::RemoveFieldNotifyFromAllMetadata(BlueprintObj, FuncName);
+	GetMetadataBlock()->RemoveMetaData(FBlueprintMetadata::MD_FieldNotify);
+	return ECheckBoxState::Unchecked;
+}
+
+void FBlueprintGraphActionDetails::OnFieldNotifyChanged(ECheckBoxState InNewState)
+{
+	UBlueprint* const BlueprintObj = GetBlueprintObj();
+	FName FuncName = FindFunction()->GetFName();
+	const bool bFuncIsFieldNotify = InNewState == ECheckBoxState::Checked;
+
+	if (BlueprintObj)
+	{
+		if (bFuncIsFieldNotify)
+		{
+			GetMetadataBlock()->SetMetaData(FBlueprintMetadata::MD_FieldNotify, FString());
+		}
+		else
+		{
+			FBlueprintEditorUtils::RemoveFieldNotifyFromAllMetadata(BlueprintObj, FuncName);
+			GetMetadataBlock()->RemoveMetaData(FBlueprintMetadata::MD_FieldNotify);
+		}
+	}
+}
+
 TSharedRef<ITableRow> FBlueprintGraphActionDetails::OnGenerateReplicationComboWidget( TSharedPtr<FReplicationSpecifierLabel> InNetFlag, const TSharedRef<STableViewBase>& OwnerTable )
 {
 	return
diff --git a/Engine/Source/Editor/Kismet/Private/BlueprintDetailsCustomization.h b/Engine/Source/Editor/Kismet/Private/BlueprintDetailsCustomization.h
--- a/Engine/Source/Editor/Kismet/Private/BlueprintDetailsCustomization.h
+++ b/Engine/Source/Editor/Kismet/Private/BlueprintDetailsCustomization.h
@@ -226,6 +226,7 @@ private:
 	
 	ECheckBoxState OnFieldNotifyCheckboxState() const;
 	void OnFieldNotifyChanged(ECheckBoxState InNewState);
+	EVisibility GetFieldNotifyCheckboxListVisiblity() const;
 
 	ECheckBoxState OnCreateWidgetCheckboxState() const;
 	void OnCreateWidgetChanged(ECheckBoxState InNewState);
@@ -786,6 +787,11 @@ private:
 	void OnIsUnsafeDuringActorConstructionModified(const ECheckBoxState NewCheckedState);
 	ECheckBoxState GetIsUnsafeDuringActorConstruction() const;
 
+	bool IsFieldNotifyCheckVisible() const;
+	bool GetIsFieldNotfyEnabled() const;
+	ECheckBoxState OnFieldNotifyCheckboxState() const;
+	void OnFieldNotifyChanged(ECheckBoxState InNewState);
+
 	/** Determines if the selected event is identified as editor callable */
 	ECheckBoxState GetIsEditorCallableEvent() const;
 
diff --git a/Engine/Source/Editor/Kismet/Private/SFieldNotificationCheckList.cpp b/Engine/Source/Editor/Kismet/Private/SFieldNotificationCheckList.cpp
new file mode 100644
--- /dev/null
+++ b/Engine/Source/Editor/Kismet/Private/SFieldNotificationCheckList.cpp
@@ -0,0 +1,294 @@
+// Copyright Epic Games, Inc. All Rights Reserved.
+
+#include "SFieldNotificationCheckList.h"
+
+#include "Engine/BlueprintGeneratedClass.h"
+#include "INotifyFieldValueChanged.h"
+#include "Layout/Children.h"
+#include "K2Node_FunctionEntry.h"
+#include "Kismet2/BlueprintEditorUtils.h"
+#include "UObject/Class.h"
+#include "UObject/NameTypes.h"
+#include "UObject/ScriptInterface.h"
+#include "Widgets/Views/ITableRow.h"
+#include "Widgets/Input/SCheckBox.h"
+#include "Widgets/Input/SComboBox.h"
+#include "Widgets/Views/STableViewBase.h"
+#include "Widgets/Text/STextBlock.h"
+#include "Widgets/Views/SListView.h"
+
+#define LOCTEXT_NAMESPACE "SFieldNotificationCheckList"
+namespace UE::FieldNotification
+{
+
+void SFieldNotificationCheckList::Construct(const FArguments& InArgs)
+{
+	BlueprintPtr = InArgs._BlueprintPtr;
+	FieldName = InArgs._FieldName;
+
+	SAssignNew(ComboListView, SComboListType)
+			.ListItemsSource(&FieldNotificationIdsSource)
+			.SelectionMode(ESelectionMode::None)
+			.OnGenerateRow(this, &SFieldNotificationCheckList::GenerateMenuItemRow);
+
+	TSharedPtr<SWidget> ButtonContent = SNew(SImage)
+		.Image(FAppStyle::Get().GetBrush("Kismet.VariableList.FieldNotify"))
+		.Visibility(this, &SFieldNotificationCheckList::GetFieldNotifyIconVisibility);
+
+	ChildSlot
+	[
+		SNew(SComboButton)
+		.ToolTipText(LOCTEXT("FieldNotifyListTooltip", "Select which field in this blueprint class will be notified when the current property changes."))
+		.ContentPadding(FMargin(2.0f, 2.0f))
+		.OnMenuOpenChanged(this, &SFieldNotificationCheckList::OnMenuOpenChanged)
+		.IsFocusable(true)
+		.MenuContent()
+		[
+			ComboListView.ToSharedRef()
+		]
+		.ButtonContent()
+		[
+			ButtonContent.ToSharedRef()
+		]
+	];
+}
+
+TSharedRef<ITableRow> SFieldNotificationCheckList::GenerateMenuItemRow(TSharedPtr<FFieldNotificationId> InItem, const TSharedRef<STableViewBase>& OwnerTable)
+{
+	// Each row is a checkbox. Depending on whether we selected a function or a variable, we bind different delegates for the state of the checkbox.
+	TSharedPtr<SWidget> RowCheckBox =
+		SNew(SCheckBox)
+		.IsChecked(this, &SFieldNotificationCheckList::OnVariableCheckboxState, InItem->GetFieldName())
+		.OnCheckStateChanged(this, &SFieldNotificationCheckList::OnVariableCheckBoxChanged, InItem->GetFieldName())
+		.IsEnabled(this, &SFieldNotificationCheckList::IsCheckBoxEnabled, InItem->GetFieldName())
+		.Content()
+		[
+			SNew(STextBlock)
+			.Text(FText::FromName(InItem->GetFieldName()))
+		];
+
+	return SNew(SComboRow<TSharedPtr<FFieldNotificationId>>, OwnerTable)
+			[
+				SNew(SBox)
+				.Padding(10.0f, 5.0f)
+				[
+					RowCheckBox.ToSharedRef()
+				]
+			];
+}
+
+bool SFieldNotificationCheckList::IsCheckBoxEnabled(FName OtherName) const
+{
+	return OtherName != FieldName;
+}
+
+ECheckBoxState SFieldNotificationCheckList::OnVariableCheckboxState(FName OtherName) const
+{
+	if (OtherName == FieldName)
+	{
+		return ECheckBoxState::Checked;
+	}
+	UBlueprint* BlueprintObj = BlueprintPtr.Get();
+	if (BlueprintObj && !OtherName.IsNone())
+	{
+		const int32 VarIndex = FBlueprintEditorUtils::FindNewVariableIndex(BlueprintObj, FieldName);
+		if (VarIndex != INDEX_NONE)
+		{
+			// Parse the metadata of the variable by delimiter | and check if OtherName is in there.
+			FString FieldNotifyValues = BlueprintObj->NewVariables[VarIndex].GetMetaData(FBlueprintMetadata::MD_FieldNotify);
+			TArray<FString> ListOfFieldNotifies;
+			const TCHAR* Delimiter = TEXT("|");
+
+			FieldNotifyValues.ParseIntoArray(ListOfFieldNotifies, Delimiter);
+			return ListOfFieldNotifies.Contains(OtherName.ToString()) ? ECheckBoxState::Checked : ECheckBoxState::Unchecked;
+		}
+	}
+	return ECheckBoxState::Unchecked;
+}
+
+void SFieldNotificationCheckList::OnVariableCheckBoxChanged(ECheckBoxState InNewState, FName OtherName)
+{
+	UBlueprint* BlueprintObj = BlueprintPtr.Get();
+	if (BlueprintObj && !OtherName.IsNone())
+	{
+		const int32 VarIndex = FBlueprintEditorUtils::FindNewVariableIndex(BlueprintObj, FieldName);
+		if (VarIndex != INDEX_NONE)
+		{
+			FString FieldNotifyValues = BlueprintObj->NewVariables[VarIndex].GetMetaData(FBlueprintMetadata::MD_FieldNotify);
+			if (InNewState == ECheckBoxState::Checked)
+			{
+				const TCHAR* Delimiter = TEXT("|");
+				FieldNotifyValues = FieldNotifyValues.Len() > 0 ? FieldNotifyValues + Delimiter + OtherName.ToString() : OtherName.ToString();
+				BlueprintObj->NewVariables[VarIndex].SetMetaData(FBlueprintMetadata::MD_FieldNotify, FieldNotifyValues);
+			}
+			else
+			{
+				TArray<FString> ListOfFieldNotifies;
+				const TCHAR* Delimiter = TEXT("|");
+				FieldNotifyValues.ParseIntoArray(ListOfFieldNotifies, Delimiter);
+				ListOfFieldNotifies.Remove(OtherName.ToString());
+				if (ListOfFieldNotifies.Num() > 0)
+				{
+					BlueprintObj->NewVariables[VarIndex].SetMetaData(FBlueprintMetadata::MD_FieldNotify, FString::Join(ListOfFieldNotifies, Delimiter));
+				}
+				else
+				{
+					BlueprintObj->NewVariables[VarIndex].SetMetaData(FBlueprintMetadata::MD_FieldNotify, TEXT(""));
+				}
+			}
+		}
+	}
+}
+
+TArray<UK2Node_FunctionEntry*> SFieldNotificationCheckList::GetFieldNotifyFunctionEntryNodesInBlueprint()
+{
+	UBlueprint* BlueprintObj = BlueprintPtr.Get();
+	TArray<UK2Node_FunctionEntry*> FunctionEntryNodes = TArray<UK2Node_FunctionEntry*>();
+	if (BlueprintObj)
+	{
+		for (int32 i = 0; i < BlueprintObj->FunctionGraphs.Num(); ++i)
+		{
+			TObjectPtr<UEdGraph>& FunctionGraph = BlueprintObj->FunctionGraphs[i];
+			for (UEdGraphNode* Node : FunctionGraph->Nodes)
+			{
+				if (UK2Node_FunctionEntry* NodeFunctionEntry = Cast<UK2Node_FunctionEntry>(Node))
+				{
+					FunctionEntryNodes.Add(NodeFunctionEntry);
+					break;
+				}
+			}
+		}
+	}
+	return FunctionEntryNodes;
+}
+
+void SFieldNotificationCheckList::OnMenuOpenChanged(bool bOpen)
+{
+	FieldNotificationIdsSource.Reset();
+	UBlueprint* BlueprintObj = BlueprintPtr.Get();
+
+	if (BlueprintObj)
+	{
+		if (UClass* Class = BlueprintObj->GeneratedClass)
+		{
+			// Populate the source array when the menu opens.
+			// First, add all the fields from the Descriptor to the source array
+			TScriptInterface<INotifyFieldValueChanged> ScriptObject = Class->GetDefaultObject();
+			if (ScriptObject.GetInterface() && ScriptObject.GetObject())
+			{
+				SFieldNotificationCheckList* Self = this;
+				const UE::FieldNotification::IClassDescriptor& Descriptor = ScriptObject->GetFieldNotificationDescriptor();
+				Descriptor.ForEachField(Class, [Self, Class](const FFieldId& Id) ->bool
+					{
+						Self->FieldNotificationIdsSource.Add(MakeShared<FFieldNotificationId>(Id.GetName()));
+						return true;
+					});
+			}
+
+			// The class descriptor is updated on compile, so we may have some outdated properties that are not field notify anymore
+			// or new properties that are field notify. So we check with the blueprint variables to get the latest data.
+			for (int32 i = 0; i < BlueprintObj->NewVariables.Num(); ++i)
+			{
+				FBPVariableDescription& Variable = BlueprintObj->NewVariables[i];
+				bool bIsValidFieldNotify = Variable.HasMetaData(FBlueprintMetadata::MD_FieldNotify);
+				// Try to locate this variable name in the source array.
+				int32 FoundValueIndex = FieldNotificationIdsSource.IndexOfByPredicate([Variable](const TSharedPtr<FFieldNotificationId>& Other)
+					{
+						return Other->GetFieldName() == Variable.VarName;
+					});
+
+				// If this variable is not found in the source array but it is field notify, add it.
+				if (FoundValueIndex == INDEX_NONE && bIsValidFieldNotify)
+				{
+					FieldNotificationIdsSource.Add(MakeShared<FFieldNotificationId>(FName(Variable.VarName)));
+				}
+				// If we found this variable in the source array but it's not field notify anymore, remove it.
+				else if (FoundValueIndex != INDEX_NONE && !bIsValidFieldNotify)
+				{
+					FieldNotificationIdsSource.RemoveAtSwap(FoundValueIndex);
+				}
+			}
+
+			// Repeat the same logic above for functions. 
+			TArray<UK2Node_FunctionEntry*> NewFunctions = GetFieldNotifyFunctionEntryNodesInBlueprint();
+			for (int32 i = 0; i < NewFunctions.Num(); ++i)
+			{
+				UK2Node_FunctionEntry* EntryNode = NewFunctions[i];
+				bool bIsValidFieldNotify = EntryNode->MetaData.HasMetaData(FBlueprintMetadata::MD_FieldNotify);
+
+				int32 FoundValueIndex = FieldNotificationIdsSource.IndexOfByPredicate([EntryNode](const TSharedPtr<FFieldNotificationId>& Other)
+					{
+						return Other->GetFieldName() == EntryNode->FunctionReference.GetMemberName();
+					});
+				if (FoundValueIndex == INDEX_NONE && bIsValidFieldNotify)
+				{
+					FieldNotificationIdsSource.Add(MakeShared<FFieldNotificationId>(EntryNode->FunctionReference.GetMemberName()));
+				}
+				else if (FoundValueIndex != INDEX_NONE && !bIsValidFieldNotify)
+				{
+					FieldNotificationIdsSource.RemoveAtSwap(FoundValueIndex);
+				}
+			}
+			
+			// Get the list of field notifies from the last time we compiled.
+			// Detect if we have deleted any functions and variables since the last compile and update the source array.
+			TArray<FFieldNotificationId> LastFieldNotifyList = CastChecked<UBlueprintGeneratedClass>(Class)->FieldNotifies;
+
+			for (int32 i = 0; i < LastFieldNotifyList.Num(); ++i)
+			{
+				FFieldNotificationId Id = LastFieldNotifyList[i];
+				int32 FoundIndexInBlueprint = BlueprintObj->NewVariables.IndexOfByPredicate([Id](const FBPVariableDescription& Other)
+					{
+						return Other.VarName == Id.GetFieldName();
+					});
+				if (FoundIndexInBlueprint == INDEX_NONE)
+				{
+					FoundIndexInBlueprint = BlueprintObj->FunctionGraphs.IndexOfByPredicate([Id](const TObjectPtr<UEdGraph>& Other)
+						{
+							return Other->GetFName() == Id.GetFieldName();
+						});
+
+					if (FoundIndexInBlueprint == INDEX_NONE)
+					{
+						int32 FoundValueIndex = FieldNotificationIdsSource.IndexOfByPredicate([Id](const TSharedPtr<FFieldNotificationId>& Other)
+							{
+								return Other->GetFieldName() == Id.GetFieldName();
+							});
+
+						if (FoundValueIndex != INDEX_NONE)
+						{
+							FieldNotificationIdsSource.RemoveAtSwap(FoundValueIndex);
+						}
+					}
+				}
+			}
+
+			// Sort the source array alphabetically.
+			if (FieldNotificationIdsSource.Num())
+			{
+				FieldNotificationIdsSource.Sort([](const TSharedPtr<FFieldNotificationId>& A, const TSharedPtr<FFieldNotificationId>& B)
+				{
+					return A->GetFieldName().LexicalLess(B->GetFieldName());
+				});
+			}
+		}
+	}
+}
+
+EVisibility SFieldNotificationCheckList::GetFieldNotifyIconVisibility() const
+{
+	UBlueprint* BlueprintObj = BlueprintPtr.Get();
+	if (BlueprintObj)
+	{
+		const int32 VarIndex = FBlueprintEditorUtils::FindNewVariableIndex(BlueprintObj, FieldName);
+		if (VarIndex != INDEX_NONE)
+		{
+			FString FieldNotifyValues = BlueprintObj->NewVariables[VarIndex].GetMetaData(FBlueprintMetadata::MD_FieldNotify);
+			return FieldNotifyValues.Len() > 0 ? EVisibility::Visible : EVisibility::Collapsed;
+		}
+	}
+	return EVisibility::Collapsed;
+}
+
+}
+#undef LOCTEXT_NAMESPACE
diff --git a/Engine/Source/Editor/Kismet/Private/SFieldNotificationCheckList.h b/Engine/Source/Editor/Kismet/Private/SFieldNotificationCheckList.h
new file mode 100644
--- /dev/null
+++ b/Engine/Source/Editor/Kismet/Private/SFieldNotificationCheckList.h
@@ -0,0 +1,62 @@
+// Copyright Epic Games, Inc. All Rights Reserved.
+
+#pragma once
+
+#include "Containers/Array.h"
+#include "FieldNotificationId.h"
+#include "Misc/Attribute.h"
+#include "Templates/SharedPointer.h"
+#include "Widgets/DeclarativeSyntaxSupport.h"
+#include "Widgets/SCompoundWidget.h"
+
+class ITableRow;
+class STableViewBase;
+class SWidget;
+class UK2Node_FunctionEntry;
+enum class ECheckBoxState : uint8;
+template <typename ItemType> class SListView;
+
+namespace UE::FieldNotification
+{
+
+/**
+ * A widget in details panel which allows the user to choose which field notify variables/functions will be notified when the selected property changes.
+ */
+class SFieldNotificationCheckList : public SCompoundWidget
+{
+public:
+	typedef SListView< TSharedPtr<FFieldNotificationId> > SComboListType;
+
+	SLATE_BEGIN_ARGS(SFieldNotificationCheckList)
+	{}
+		SLATE_ARGUMENT(TWeakObjectPtr<UBlueprint>, BlueprintPtr)
+		SLATE_ARGUMENT(FName, FieldName)
+		SLATE_ARGUMENT(bool, IsFieldFunction)
+
+	SLATE_END_ARGS()
+
+	void Construct(const FArguments& InArgs);
+
+private:
+	EVisibility GetFieldNotifyIconVisibility() const;
+
+	TSharedRef<ITableRow> GenerateMenuItemRow(TSharedPtr<FFieldNotificationId> InItem, const TSharedRef<STableViewBase>& OwnerTable);
+	void OnMenuOpenChanged(bool bOpen);
+
+	bool IsCheckBoxEnabled(FName OtherName) const;
+	ECheckBoxState OnVariableCheckboxState(FName OtherName) const;
+	void OnVariableCheckBoxChanged(ECheckBoxState InNewState, FName OtherName);
+
+	/** Iterates through all graphs in the blueprint to find function graphs that are field notify and adds their entry node to OutArray. */
+	TArray<UK2Node_FunctionEntry*> GetFieldNotifyFunctionEntryNodesInBlueprint();
+
+private:
+	/** Array of all field notify variables and functions in the owning blueprint. */
+	TArray<TSharedPtr<FFieldNotificationId>> FieldNotificationIdsSource;
+	TWeakObjectPtr<UBlueprint> BlueprintPtr;
+	TSharedPtr< SComboListType > ComboListView;
+	/** The name of the selected field in the blueprint editor.*/
+	FName FieldName;
+};
+
+} // namespace
\ No newline at end of file
diff --git a/Engine/Source/Editor/KismetCompiler/Private/KismetCompiler.cpp b/Engine/Source/Editor/KismetCompiler/Private/KismetCompiler.cpp
--- a/Engine/Source/Editor/KismetCompiler/Private/KismetCompiler.cpp
+++ b/Engine/Source/Editor/KismetCompiler/Private/KismetCompiler.cpp
@@ -2547,6 +2547,14 @@ void FKismetCompilerContext::PrecompileFunction(FKismetFunctionContext& Context,
 			}
 		}
 
+		if (Context.EntryPoint && Context.Function)
+		{
+			if (Context.EntryPoint->MetaData.HasMetaData(FBlueprintMetadata::MD_FieldNotify) && NewClass->ImplementsInterface(UNotifyFieldValueChanged::StaticClass()))
+			{
+				ensure(!NewClass->FieldNotifies.Contains(FFieldNotificationId(Context.Function->GetFName())));
+				NewClass->FieldNotifies.Add(FFieldNotificationId(Context.Function->GetFName()));
+			}
+		}
 	}
 	else
 	{
diff --git a/Engine/Source/Editor/UnrealEd/Private/Kismet2/BlueprintEditorUtils.cpp b/Engine/Source/Editor/UnrealEd/Private/Kismet2/BlueprintEditorUtils.cpp
--- a/Engine/Source/Editor/UnrealEd/Private/Kismet2/BlueprintEditorUtils.cpp
+++ b/Engine/Source/Editor/UnrealEd/Private/Kismet2/BlueprintEditorUtils.cpp
@@ -134,6 +134,7 @@
 #include "Elements/Interfaces/TypedElementObjectInterface.h"
 #include "Misc/MessageDialog.h"
 #include "UObject/FastReferenceCollector.h"
+#include "INotifyFieldValueChanged.h"
 
 #define LOCTEXT_NAMESPACE "Blueprint"
 
@@ -2730,6 +2731,28 @@ void FBlueprintEditorUtils::RemoveGraphs( UBlueprint* Blueprint, const TArray<cl
 	FBlueprintEditorUtils::MarkBlueprintAsStructurallyModified(Blueprint);
 }
 
+namespace UE::Private
+{
+	// Given a function graph, check if it's field notify and remove its name from the meta data of all field notify variables.
+	static void RemoveFieldNotifyFunctionFromAllMetaData(UBlueprint* Blueprint, class UEdGraph* FunctionGraphToRemove)
+	{
+		if (Blueprint->FunctionGraphs.Contains(FunctionGraphToRemove) && FBlueprintEditorUtils::ImplementsInterface(Blueprint, true, UNotifyFieldValueChanged::StaticClass()))
+		{
+			for (UEdGraphNode* Node : FunctionGraphToRemove->Nodes)
+			{
+				if (UK2Node_FunctionEntry* NodeFunctionEntry = Cast<UK2Node_FunctionEntry>(Node))
+				{
+					if (NodeFunctionEntry->MetaData.HasMetaData(FBlueprintMetadata::MD_FieldNotify))
+					{
+						FBlueprintEditorUtils::RemoveFieldNotifyFromAllMetadata(Blueprint, FunctionGraphToRemove->GetFName());
+					}
+					break;
+				}
+			}
+		}
+	}
+}
+
 // Removes the supplied graph from the Blueprint.
 void FBlueprintEditorUtils::RemoveGraph(UBlueprint* Blueprint, class UEdGraph* GraphToRemove, EGraphRemoveFlags::Type Flags /*= Transient | Recompile */)
 {
@@ -2739,6 +2762,8 @@ void FBlueprintEditorUtils::RemoveGraph(UBlueprint* Blueprint, class UEdGraph* G
 	{
 		if (TestOuter == Blueprint)
 		{
+			UE::Private::RemoveFieldNotifyFunctionFromAllMetaData(Blueprint, GraphToRemove);
+
 			Blueprint->DelegateSignatureGraphs.Remove( GraphToRemove );
 			Blueprint->FunctionGraphs.Remove( GraphToRemove );
 			Blueprint->UbergraphPages.Remove( GraphToRemove );
@@ -4939,12 +4964,46 @@ void FBlueprintEditorUtils::RemoveMemberVariable(UBlueprint* Blueprint, const FN
 	const int32 VarIndex = FBlueprintEditorUtils::FindNewVariableIndex(Blueprint, VarName);
 	if (VarIndex != INDEX_NONE)
 	{
+		if ( Blueprint->NewVariables[VarIndex].HasMetaData(FBlueprintMetadata::MD_FieldNotify) && FBlueprintEditorUtils::ImplementsInterface(Blueprint, true, UNotifyFieldValueChanged::StaticClass()) )
+		{
+			RemoveFieldNotifyFromAllMetadata(Blueprint, VarName);
+		}
 		Blueprint->NewVariables.RemoveAt(VarIndex);
 		FBlueprintEditorUtils::RemoveVariableNodes(Blueprint, VarName);
 		FBlueprintEditorUtils::MarkBlueprintAsStructurallyModified(Blueprint);
 	}
 }
 
+// Removes the field name from the field notify metadata of all functions and variables
+void FBlueprintEditorUtils::RemoveFieldNotifyFromAllMetadata(UBlueprint* Blueprint, const FName FieldName)
+{
+	// Go through all field notify variables and remove FieldName from their metadata.
+	for (int32 i = 0; i < Blueprint->NewVariables.Num(); ++i)
+	{
+		FBPVariableDescription& Variable = Blueprint->NewVariables[i];
+		if (Variable.HasMetaData(FBlueprintMetadata::MD_FieldNotify))
+		{
+			FString FieldNotifyValues = Variable.GetMetaData(FBlueprintMetadata::MD_FieldNotify);
+			TArray<FString> ListOfFieldNotifies;
+			const TCHAR* Delimiter = TEXT("|");
+
+			FieldNotifyValues.ParseIntoArray(ListOfFieldNotifies, Delimiter);
+			if (ListOfFieldNotifies.Contains(FieldName.ToString()))
+			{
+				ListOfFieldNotifies.Remove(FieldName.ToString());
+				if (ListOfFieldNotifies.Num() > 0)
+				{
+					FBlueprintEditorUtils::SetBlueprintVariableMetaData(Blueprint, Variable.VarName, NULL, FBlueprintMetadata::MD_FieldNotify, FString::Join(ListOfFieldNotifies, Delimiter));
+				}
+				else
+				{
+					FBlueprintEditorUtils::SetBlueprintVariableMetaData(Blueprint, Variable.VarName, NULL, FBlueprintMetadata::MD_FieldNotify, TEXT(""));
+				}
+			}
+		}
+	}
+}
+
 void FBlueprintEditorUtils::BulkRemoveMemberVariables(UBlueprint* Blueprint, const TArray<FName>& VarNames)
 {
 	const FScopedTransaction Transaction( LOCTEXT("BulkRemoveMemberVariables", "Bulk Remove Member Variables") );
@@ -6068,6 +6127,7 @@ namespace UE::Blueprint::Private
 
 		return false;
 	}
+
 }
 
 bool FBlueprintEditorUtils::IsVariableUsed(const UBlueprint* Blueprint, const FName& VariableName, const UEdGraph* LocalGraphScope /* = nullptr */)
diff --git a/Engine/Source/Editor/UnrealEd/Public/Kismet2/BlueprintEditorUtils.h b/Engine/Source/Editor/UnrealEd/Public/Kismet2/BlueprintEditorUtils.h
--- a/Engine/Source/Editor/UnrealEd/Public/Kismet2/BlueprintEditorUtils.h
+++ b/Engine/Source/Editor/UnrealEd/Public/Kismet2/BlueprintEditorUtils.h
@@ -894,6 +894,13 @@ public:
 	 */
 	static void BulkRemoveMemberVariables(UBlueprint* Blueprint, const TArray<FName>& VarNames);
 
+	/**
+	 * Removes a field notify variable from the metadata of all other field notify variables and functions.
+	 *
+	 * @param	VarName	Name of the variable to be removed.
+	 */
+	static void RemoveFieldNotifyFromAllMetadata(UBlueprint* Blueprint, const FName VarName);
+
 	/**
 	 * Removes all unused member variables.
 	 *
diff --git a/Engine/Source/Editor/UnrealEd/UnrealEd.Build.cs b/Engine/Source/Editor/UnrealEd/UnrealEd.Build.cs
--- a/Engine/Source/Editor/UnrealEd/UnrealEd.Build.cs
+++ b/Engine/Source/Editor/UnrealEd/UnrealEd.Build.cs
@@ -65,6 +65,7 @@ public class UnrealEd : ModuleRules
 				"BlueprintGraph",
 				"HTTP",
 				"FunctionalTesting",
+				"AutomationController",
 				"Localization",
 				"AudioEditor",
 				"NetworkFileSystem",
@@ -84,6 +85,7 @@ public class UnrealEd : ModuleRules
 				"InterchangeEngine",
 				"DeveloperToolSettings",
 				"SubobjectDataInterface",
+				"SubobjectEditor",
 				"PhysicsUtilities",
 				"ToolWidgets",
 				"AssetTools", 
@@ -174,6 +176,7 @@ public class UnrealEd : ModuleRules
 				"BuildSettings",
 				"SubobjectEditor",
 				"HeadMountedDisplay",
+				"FieldNotification",
 			}
 		);
 
