UE-187439 - Variables dragged and dropped to the node with split struct pins are not added

CL 12611535 disabled drag and drop on nodes and their pins to prevent a reconstruction crash. That crash was actually due to drag and dropping a variable onto a *node* that had split pins. By doing that, the variable reference changed, which tripped up node reconstruction since it discovers an inconsistent pin count. However, it's actually safe to drag and drop a variable to a pin connection of a split node. This doesn't affect the underlying variable reference, or introduce any other behavior that might confuse node reconstruction.

The fix here simply removes the additional split pin criteria from DroppedOnPin.

diff --git a/Engine/Source/Editor/Kismet/Private/BPVariableDragDropAction.cpp b/Engine/Source/Editor/Kismet/Private/BPVariableDragDropAction.cpp
--- a/Engine/Source/Editor/Kismet/Private/BPVariableDragDropAction.cpp
+++ b/Engine/Source/Editor/Kismet/Private/BPVariableDragDropAction.cpp
@@ -307,7 +307,7 @@ FReply FKismetVariableDragDropAction::DroppedOnPin(FVector2D ScreenPosition, FVe
 			{
 				const bool bIsExecPin = Schema->IsExecPin(*TargetPin);
 
-				if (CanVariableBeDropped(VariableProperty, *TargetPin->GetOwningNode()->GetGraph()) && !NodeHasSplitPins(TargetPin->GetOwningNode()))
+				if (CanVariableBeDropped(VariableProperty, *TargetPin->GetOwningNode()->GetGraph()))
 				{
 					const bool bIsRead = (TargetPin->Direction == EGPD_Input) && !bIsExecPin;
 					const UBlueprint* Blueprint = FBlueprintEditorUtils::FindBlueprintForNode(TargetPin->GetOwningNode());
@@ -371,7 +371,7 @@ FReply FKismetVariableDragDropAction::DroppedOnNode(FVector2D ScreenPosition, FV
 	{
 		FProperty* VariableProperty = GetVariableProperty();
 
-		if(CanVariableBeDropped(VariableProperty, *TargetNode->GetGraph()) && !NodeHasSplitPins(TargetNode))
+		if (CanVariableBeDropped(VariableProperty, *TargetNode->GetGraph()) && !NodeHasSplitPins(TargetNode))
 		{
 			const FScopedTransaction Transaction(LOCTEXT("ReplacePinVariable", "Replace Pin Variable"));
 
@@ -406,8 +406,8 @@ FReply FKismetVariableDragDropAction::DroppedOnNode(FVector2D ScreenPosition, FV
 				Pin->PinName = VariableName;
 				Pin->PinType = NewPinType;
 
-				//break bad links
-				for(TArray<class UEdGraphPin*>::TIterator OtherPinIt(BadLinks);OtherPinIt;++OtherPinIt)
+				// break bad links
+				for (TArray<class UEdGraphPin*>::TIterator OtherPinIt(BadLinks);OtherPinIt;++OtherPinIt)
 				{
 					Pin->BreakLinkTo(*OtherPinIt);
 				}
