Fixed include path.

diff --git a/Engine/Source/Editor/Kismet/Public/BlueprintEditor.h b/Engine/Source/Editor/Kismet/Public/BlueprintEditor.h
--- a/Engine/Source/Editor/Kismet/Public/BlueprintEditor.h
+++ b/Engine/Source/Editor/Kismet/Public/BlueprintEditor.h
@@ -10,7 +10,6 @@
 #include "CoreMinimal.h"
 #include "CoreTypes.h"
 #include "Delegates/Delegate.h"
-#include "Developer/Merge/Public/Merge.h"
 #include "EdGraph/EdGraphNode.h"
 #include "EdGraph/EdGraphPin.h"
 #include "EdGraph/EdGraphSchema.h"
@@ -26,6 +25,7 @@
 #include "Layout/Visibility.h"
 #include "Math/Color.h"
 #include "Math/Vector2D.h"
+#include "Merge.h"
 #include "Misc/Guid.h"
 #include "Misc/NotifyHook.h"
 #include "Misc/Optional.h"
