Conform interfaces and functions before reconstructing nodes, this prevents a phantom compilation error when loading a blueprint after a parent blueprint has removed a function

diff --git a/Engine/Source/Editor/Kismet/Private/BlueprintCompilationManager.cpp b/Engine/Source/Editor/Kismet/Private/BlueprintCompilationManager.cpp
--- a/Engine/Source/Editor/Kismet/Private/BlueprintCompilationManager.cpp
+++ b/Engine/Source/Editor/Kismet/Private/BlueprintCompilationManager.cpp
@@ -139,6 +139,7 @@ struct FBlueprintCompilationManagerImpl : public FGCObject
 	static void ReinstanceBatch(TArray<FReinstancingJob>& Reinstancers, TMap< UClass*, UClass* >& InOutOldToNewClassMap, FUObjectSerializeContext* InLoadContext);
 	static UClass* FastGenerateSkeletonClass(UBlueprint* BP, FKismetCompilerContext& CompilerContext, bool bIsSkeletonOnly, TArray<FSkeletonFixupData>& OutSkeletonFixupData);
 	static bool IsQueuedForCompilation(UBlueprint* BP);
+	static void ConformToParentAndInterfaces(UBlueprint* BP);
 
 	// Declaration of archive to fix up bytecode references of blueprints that are actively compiled:
 	class FFixupBytecodeReferences : public FArchiveUObject
@@ -1217,6 +1218,8 @@ void FBlueprintCompilationManagerImpl::FlushCompilationQueueImpl(bool bSuppressB
 
 			UBlueprint* BP = CompilerData.BP;
 
+			ConformToParentAndInterfaces(BP);
+
 			// Some nodes are set up to do things during reconstruction only when this flag is NOT set.
 			if(BP->bIsRegeneratingOnLoad)
 			{
@@ -3260,6 +3263,26 @@ bool FBlueprintCompilationManagerImpl::IsQueuedForCompilation(UBlueprint* BP)
 	return BP->bQueuedForCompilation;
 }
 
+void FBlueprintCompilationManagerImpl::ConformToParentAndInterfaces(UBlueprint* BP)
+{
+	// If graphs are renamed the blueprint will be marked as not 'bCachedDependenciesUpToDate', but 
+	// because we're conforming an existing dependency we don't need to change bCachedDependenciesUpToDate:
+	TGuardValue<bool> LockDependenciesUpToDate(BP->bCachedDependenciesUpToDate, BP->bCachedDependenciesUpToDate);
+
+	// Make sure that this blueprint is up-to-date with regards to its parent functions
+	FBlueprintEditorUtils::ConformCallsToParentFunctions(BP);
+
+	// Conform implemented events here, to ensure we generate custom events if necessary after reparenting
+	FBlueprintEditorUtils::ConformImplementedEvents(BP);
+
+	// Conform implemented interfaces here, to ensure we generate all functions required by the interface as stubs
+	FBlueprintEditorUtils::ConformImplementedInterfaces(BP);
+
+	// Make sure we don't have any signature graphs with no corresponding variable - some assets have
+	// managed to get into this state - the UI does not provide a way to fix these objects manually
+	FBlueprintEditorUtils::ConformDelegateSignatureGraphs(BP);
+}
+
 // FFixupBytecodeReferences Implementation:
 FBlueprintCompilationManagerImpl::FFixupBytecodeReferences::FFixupBytecodeReferences(UObject* InObject)
 {
diff --git a/Engine/Source/Editor/KismetCompiler/Private/KismetCompiler.cpp b/Engine/Source/Editor/KismetCompiler/Private/KismetCompiler.cpp
--- a/Engine/Source/Editor/KismetCompiler/Private/KismetCompiler.cpp
+++ b/Engine/Source/Editor/KismetCompiler/Private/KismetCompiler.cpp
@@ -4802,25 +4802,6 @@ void FKismetCompilerContext::CompileClassLayout(EInternalCompilerFlags InternalF
 		}
 	}
 
-	{
-		// the following calls may mark the blueprint as dirty, but we know that these operations just cleaned up the BP 
-		// so dependencies can still be considered 'up to date'
-		TGuardValue<bool> LockDependenciesUpToDate(Blueprint->bCachedDependenciesUpToDate, Blueprint->bCachedDependenciesUpToDate);
-
-		// Make sure that this blueprint is up-to-date with regards to its parent functions
-		FBlueprintEditorUtils::ConformCallsToParentFunctions(Blueprint);
-
-		// Conform implemented events here, to ensure we generate custom events if necessary after reparenting
-		FBlueprintEditorUtils::ConformImplementedEvents(Blueprint);
-
-		// Conform implemented interfaces here, to ensure we generate all functions required by the interface as stubs
-		FBlueprintEditorUtils::ConformImplementedInterfaces(Blueprint);
-
-		// Make sure we don't have any signature graphs with no corresponding variable - some assets have
-		// managed to get into this state - the UI does not provide a way to fix these objects manually
-		FBlueprintEditorUtils::ConformDelegateSignatureGraphs(Blueprint);
-	}
-
 	IKismetCompilerInterface& KismetCompilerModule = FModuleManager::LoadModuleChecked<IKismetCompilerInterface>("KismetCompiler");
 	KismetCompilerModule.ValidateBPAndClassType(Blueprint, MessageLog);
 
