Defer validation of user-defined struct types in the pin type selector list until selection time for the member variable type widget in the user-defined struct asset editor.

Notes:
- This resolves a UX regression I introduced in Release-[REDACTED_PHONE]).
- Invocation of the pin type selector menu for member variables is now a bit faster as a result.

diff --git a/Engine/Source/Editor/Kismet/Private/UserDefinedStructureEditor.cpp b/Engine/Source/Editor/Kismet/Private/UserDefinedStructureEditor.cpp
--- a/Engine/Source/Editor/Kismet/Private/UserDefinedStructureEditor.cpp
+++ b/Engine/Source/Editor/Kismet/Private/UserDefinedStructureEditor.cpp
@@ -17,6 +17,7 @@
 #include "Framework/Docking/TabManager.h"
 #include "Framework/MultiBox/MultiBoxBuilder.h"
 #include "Framework/MultiBox/MultiBoxExtender.h"
+#include "Framework/Notifications/NotificationManager.h"
 #include "GenericPlatform/ICursor.h"
 #include "HAL/PlatformCrt.h"
 #include "HAL/PlatformMath.h"
@@ -64,6 +65,7 @@
 #include "Widgets/Input/SCheckBox.h"
 #include "Widgets/Input/SEditableTextBox.h"
 #include "Widgets/Layout/SBox.h"
+#include "Widgets/Notifications/SNotificationList.h"
 #include "Widgets/SBoxPanel.h"
 #include "Widgets/Text/STextBlock.h"
 #include "Widgets/Views/STableRow.h"
@@ -830,10 +832,18 @@ public:
 		auto StructureDetailsSP = StructureDetails.Pin();
 		if(StructureDetailsSP.IsValid())
 		{
-			FStructureEditorUtils::ChangeVariableType(StructureDetailsSP->GetUserDefinedStruct(), FieldGuid, PinType);
-			if (TSharedPtr<FUserDefinedStructureEditor> StructureEditorSP = StructureDetailsSP->GetStructureEditor().Pin())
+			if (FStructureEditorUtils::ChangeVariableType(StructureDetailsSP->GetUserDefinedStruct(), FieldGuid, PinType))
 			{
-				StructureEditorSP->SetInitialPinType(PinType);
+				if (TSharedPtr<FUserDefinedStructureEditor> StructureEditorSP = StructureDetailsSP->GetStructureEditor().Pin())
+				{
+					StructureEditorSP->SetInitialPinType(PinType);
+				}
+			}
+			else
+			{
+				FNotificationInfo NotificationInfo(LOCTEXT("VariableTypeChange_FailureNotification", "Variable type change failed (the selected type may not be compatible with this struct). See log for details."));
+				NotificationInfo.ExpireDuration = 5.0f;
+				FSlateNotificationManager::Get().AddNotification(NotificationInfo);
 			}
 		}
 	}
@@ -1012,32 +1022,6 @@ public:
 		return EVisibility::Collapsed;
 	}
 
-	void RemoveInvalidSubTypes(TSharedPtr<UEdGraphSchema_K2::FPinTypeTreeInfo> PinTypeNode, const UUserDefinedStruct* Parent) const
-	{
-		if (!PinTypeNode.IsValid() || !Parent)
-		{
-			return;
-		}
-
-		for (int32 ChildIndex = 0; ChildIndex < PinTypeNode->Children.Num();)
-		{
-			const auto Child = PinTypeNode->Children[ChildIndex];
-			if(Child.IsValid())
-			{
-				const bool bForceLoadSubCategoryObject = false;
-				const FEdGraphPinType& PinType = Child->GetPinType(bForceLoadSubCategoryObject);
-
-				const bool bCanCheckSubObjectWithoutLoading = PinType.PinSubCategoryObject.IsValid();
-				if(bCanCheckSubObjectWithoutLoading && !FStructureEditorUtils::CanHaveAMemberVariableOfType(Parent, PinType))
-				{
-					PinTypeNode->Children.RemoveAt(ChildIndex);
-					continue;
-				}
-			}
-			++ChildIndex;
-		}
-	}
-
 	void GetFilteredVariableTypeTree( TArray< TSharedPtr<UEdGraphSchema_K2::FPinTypeTreeInfo> >& TypeTree, ETypeTreeFilter TypeTreeFilter) const
 	{
 		auto K2Schema = GetDefault<UEdGraphSchema_K2>();
@@ -1045,12 +1029,6 @@ public:
 		if(StructureDetailsSP.IsValid() && K2Schema)
 		{
 			K2Schema->GetVariableTypeTree(TypeTree, TypeTreeFilter);
-			const auto Parent = StructureDetailsSP->GetUserDefinedStruct();
-			// THE TREE HAS ONLY 2 LEVELS
-			for (auto PinTypePtr : TypeTree)
-			{
-				RemoveInvalidSubTypes(PinTypePtr, Parent);
-			}
 		}
 	}
 
diff --git a/Engine/Source/Editor/UnrealEd/Private/Kismet2/StructureEditorUtils.cpp b/Engine/Source/Editor/UnrealEd/Private/Kismet2/StructureEditorUtils.cpp
--- a/Engine/Source/Editor/UnrealEd/Private/Kismet2/StructureEditorUtils.cpp
+++ b/Engine/Source/Editor/UnrealEd/Private/Kismet2/StructureEditorUtils.cpp
@@ -83,7 +83,7 @@ FStructureEditorUtils::EStructureError FStructureEditorUtils::IsStructureValid(c
 	{
 		if (OutMsg)
 		{
-			*OutMsg = FText::Format(LOCTEXT("StructureRecursionFmt", "Recursion: Struct cannot have itself as a member variable. Struct '{0}', recursive parent '{1}'"), 
+			*OutMsg = FText::Format(LOCTEXT("StructureRecursionFmt", "Recursion: Recursion: Struct cannot have itself or a nested struct member referencing itself as a member variable. Struct '{0}', recursive parent '{1}'"), 
 				 FText::FromString(Struct->GetFullName()), FText::FromString(RecursionParent->GetFullName())).ToString();
 		}
 		return EStructureError::Recursion;
