Don't mark level/OFPA actor packages dirty when changing subobject defaults in the Blueprint editor.

Also fixed to not mark dirty when changing defaults in the Blueprint editor's "data-only" view mode (this brings that mode to parity w/ the "full" view mode).

diff --git a/Engine/Source/Editor/Kismet/Private/BlueprintEditor.cpp b/Engine/Source/Editor/Kismet/Private/BlueprintEditor.cpp
--- a/Engine/Source/Editor/Kismet/Private/BlueprintEditor.cpp
+++ b/Engine/Source/Editor/Kismet/Private/BlueprintEditor.cpp
@@ -9585,19 +9585,43 @@ void FBlueprintEditor::OnPreObjectPropertyChanged(UObject* InObject, const FEdit
 {
 	if (GetDefault<UBlueprintEditorSettings>()->bDoNotMarkAllInstancesDirtyOnDefaultValueChange)
 	{
+		// Determine which inspector widget is currently in use.
+		TSharedPtr<SKismetInspector> CurrentInspectorWidget;
+		if (GetCurrentMode() == FBlueprintEditorApplicationModes::BlueprintDefaultsMode)
+		{
+			CurrentInspectorWidget = DefaultEditor;
+		}
+		else
+		{
+			CurrentInspectorWidget = Inspector;
+		}
+
 		// Note: While we could rely on our notify hook to determine whether an incoming change belongs to this
 		// context, NotifyPreChange() may end up getting called after this event is broadcast (e.g. from inside
 		// of a details customization). So it's not safe to assume they will get called in any specific order.
 		// It is, however, safe to assume this will get called for an archetype/CDO prior to change propagation.
-		if (Inspector.IsValid())
+		if (CurrentInspectorWidget.IsValid())
 		{
 			check(InObject);
 
+			auto IsSubobjectOfAnySelectedObject = [](const UObject* InObject, const TArray<TWeakObjectPtr<UObject>>& SelectedObjects)
+			{
+				for (const TWeakObjectPtr<UObject>& SelectedObject : SelectedObjects)
+				{
+					if (SelectedObject.IsValid() && InObject->IsInOuter(SelectedObject.Get()))
+					{
+						return true;
+					}
+				}
+
+				return false;
+			};
+
 			// If we're modifying an object that's selected into our property editor context (e.g. the Blueprint CDO, or an SCS
 			// component template), set up change propagation so that instances do not always mark their outer package as dirty.
 			bool bIsObjectSelectedForEditing = false;
-			const TArray<TWeakObjectPtr<UObject>> SelectedObjects = Inspector->GetSelectedObjects();
-			if (SelectedObjects.Contains(InObject))
+			const TArray<TWeakObjectPtr<UObject>>& SelectedObjects = CurrentInspectorWidget->GetSelectedObjects();
+			if (SelectedObjects.Contains(InObject) || IsSubobjectOfAnySelectedObject(InObject, SelectedObjects))
 			{
 				bIsObjectSelectedForEditing = true;
 			}
