PR #10526: fix a crash when diff an anim blueprint

diff --git a/Engine/Source/Editor/Kismet/Private/SBlueprintDiff.cpp b/Engine/Source/Editor/Kismet/Private/SBlueprintDiff.cpp
--- a/Engine/Source/Editor/Kismet/Private/SBlueprintDiff.cpp
+++ b/Engine/Source/Editor/Kismet/Private/SBlueprintDiff.cpp
@@ -737,7 +737,10 @@ bool FDiffPanel::CanCopyNodes() const
 
 void FDiffPanel::FocusDiff(UEdGraphPin& Pin)
 {
-	GraphEditor.Pin()->JumpToPin(&Pin);
+	if (GraphEditor.IsValid())
+	{
+		GraphEditor.Pin()->JumpToPin(&Pin);
+	}
 }
 
 void FDiffPanel::FocusDiff(UEdGraphNode& Node)
