Update UFunction::SuperStruct values in skeleton classes when relinking them during compilation. This prevents them from going stale and referencing UFunctions in a TRASHCLASS

diff --git a/Engine/Source/Editor/Kismet/Private/BlueprintCompilationManager.cpp b/Engine/Source/Editor/Kismet/Private/BlueprintCompilationManager.cpp
--- a/Engine/Source/Editor/Kismet/Private/BlueprintCompilationManager.cpp
+++ b/Engine/Source/Editor/Kismet/Private/BlueprintCompilationManager.cpp
@@ -141,6 +141,7 @@ struct FBlueprintCompilationManagerImpl : public FGCObject
 	static UClass* FastGenerateSkeletonClass(UBlueprint* BP, FKismetCompilerContext& CompilerContext, bool bIsSkeletonOnly, TArray<FSkeletonFixupData>& OutSkeletonFixupData);
 	static bool IsQueuedForCompilation(UBlueprint* BP);
 	static void ConformToParentAndInterfaces(UBlueprint* BP);
+	static void RelinkSkeleton(UClass* SkeletonToRelink);
 
 	// Declaration of archive to fix up bytecode references of blueprints that are actively compiled:
 	class FFixupBytecodeReferences : public FArchiveUObject
@@ -1085,19 +1086,7 @@ void FBlueprintCompilationManagerImpl::FlushCompilationQueueImpl(bool bSuppressB
 
 					// Just relink, note that UProperties that reference *other* types may be stale until
 					// we fixup below:
-					UClass* SkeletonToRelink = BP->SkeletonGeneratedClass;
-
-					// CDO needs to be moved aside already:
-					ensure(SkeletonToRelink->ClassDefaultObject == nullptr);
-					ensure(!SkeletonToRelink->GetSuperClass()->HasAnyClassFlags(CLASS_NewerVersionExists));
-					
-					SkeletonToRelink->ClassConstructor = nullptr;
-					SkeletonToRelink->ClassVTableHelperCtorCaller = nullptr;
-					SkeletonToRelink->CppClassStaticFunctions.Reset();
-					SkeletonToRelink->Bind();
-					SkeletonToRelink->ClearFunctionMapsCaches();
-					SkeletonToRelink->StaticLink(true);
-					SkeletonToRelink->GetDefaultObject()->SetFlags(RF_Transient);
+					RelinkSkeleton(BP->SkeletonGeneratedClass);
 				}
 
 				if(CompilerData.ShouldMarkUpToDateAfterSkeletonStage())
@@ -3328,6 +3317,32 @@ void FBlueprintCompilationManagerImpl::ConformToParentAndInterfaces(UBlueprint*
 	FBlueprintEditorUtils::ConformDelegateSignatureGraphs(BP);
 }
 
+void FBlueprintCompilationManagerImpl::RelinkSkeleton(UClass* SkeletonToRelink)
+{
+	// CDO needs to be moved aside already:
+	ensure(SkeletonToRelink->ClassDefaultObject == nullptr);
+	ensure(!SkeletonToRelink->GetSuperClass()->HasAnyClassFlags(CLASS_NewerVersionExists));
+
+	SkeletonToRelink->ClassConstructor = nullptr;
+	SkeletonToRelink->ClassVTableHelperCtorCaller = nullptr;
+	SkeletonToRelink->CppClassStaticFunctions.Reset();
+	SkeletonToRelink->Bind();
+	SkeletonToRelink->ClearFunctionMapsCaches();
+	SkeletonToRelink->StaticLink(true);
+	SkeletonToRelink->GetDefaultObject()->SetFlags(RF_Transient);
+
+	// Update UFunction SuperStruct pointers, which are typically set to their overridden function.
+	// For non-skeleton classes these are reassigned by the 'bytecode' recompile that we run on all
+	// referencing classes...
+	for (TFieldIterator<UFunction> FuncIter(SkeletonToRelink, EFieldIteratorFlags::ExcludeSuper); FuncIter; ++FuncIter)
+	{
+		if (UFunction* SuperFunction = SkeletonToRelink->GetSuperClass()->FindFunctionByName(FuncIter->GetFName()))
+		{
+			FuncIter->SetSuperStruct(SuperFunction);
+		}
+	}
+}
+
 // FFixupBytecodeReferences Implementation:
 FBlueprintCompilationManagerImpl::FFixupBytecodeReferences::FFixupBytecodeReferences(UObject* InObject)
 {
