[BugFix] Details panel order has unexpected category order changes when diffing old blueprints

turns out this is due to a nuance with UBlueprint::CategorySorting which ends up getting populated with significantly more properties if users view inherited variables within the MyBlueprint panel. From then on, the blueprint will get serialized with a different category sorting than previous versions. To fix, we just use the latest sorting on both left and right panels during diffs.

diff --git a/Engine/Source/Editor/Kismet/Private/SBlueprintDiff.cpp b/Engine/Source/Editor/Kismet/Private/SBlueprintDiff.cpp
--- a/Engine/Source/Editor/Kismet/Private/SBlueprintDiff.cpp
+++ b/Engine/Source/Editor/Kismet/Private/SBlueprintDiff.cpp
@@ -161,6 +161,25 @@ BEGIN_SLATE_FUNCTION_BUILD_OPTIMIZATION
 void SBlueprintDiff::Construct( const FArguments& InArgs)
 {
 	check(InArgs._BlueprintOld || InArgs._BlueprintNew);
+	
+	// make sure that both blueprints have the same category sorting so that they diff properly.
+	if (FPackageName::IsTempPackage(InArgs._BlueprintOld->GetPackage()->GetName()))
+	{
+		const_cast<UBlueprint*>(InArgs._BlueprintOld)->CategorySorting = InArgs._BlueprintNew->CategorySorting;
+	}
+	else if(FPackageName::IsTempPackage(InArgs._BlueprintNew->GetPackage()->GetName()))
+	{
+		const_cast<UBlueprint*>(InArgs._BlueprintNew)->CategorySorting = InArgs._BlueprintOld->CategorySorting;
+	}
+	else
+	{
+		// Neither New or Old BPs are temp so we need to scope this change to the lifetime of this diff window.
+		// when the window is closed, ScopedCategorySortChange will revert BlueprintOld->CategorySorting back to it's original state.
+		ScopedCategorySortChange.SetBlueprint(const_cast<UBlueprint*>(InArgs._BlueprintOld));
+		const_cast<UBlueprint*>(InArgs._BlueprintOld)->CategorySorting = InArgs._BlueprintNew->CategorySorting;
+	}
+	
+	
 	PanelOld.Blueprint = InArgs._BlueprintOld;
 	PanelNew.Blueprint = InArgs._BlueprintNew;
 	PanelOld.RevisionInfo = InArgs._OldRevision;
@@ -1292,6 +1311,20 @@ void SBlueprintDiff::UpdateTopSectionVisibility(const FName& InNewViewMode) cons
 	}
 }
 
+SBlueprintDiff::FScopedCategorySortChange::~FScopedCategorySortChange()
+{
+	if (Blueprint)
+	{
+		Blueprint->CategorySorting = Backup;
+	}
+}
+
+void SBlueprintDiff::FScopedCategorySortChange::SetBlueprint(UBlueprint* InBlueprint)
+{
+	Blueprint = InBlueprint;
+	Backup = Blueprint->CategorySorting;
+}
+
 void SBlueprintDiff::OnModeChanged(const FName& InNewViewMode) const
 {
 	UpdateTopSectionVisibility(InNewViewMode);
diff --git a/Engine/Source/Editor/Kismet/Public/SBlueprintDiff.h b/Engine/Source/Editor/Kismet/Public/SBlueprintDiff.h
--- a/Engine/Source/Editor/Kismet/Public/SBlueprintDiff.h
+++ b/Engine/Source/Editor/Kismet/Public/SBlueprintDiff.h
@@ -297,6 +297,18 @@ protected:
 	TWeakPtr<SWindow> WeakParentWindow;
 
 	FDelegateHandle AssetEditorCloseDelegate;
+
+	/** To make diffing more accurate and friendly, UBlueprint::CategorySorting gets modified. this will revert to the
+	 *  old version when the window closes */
+	class FScopedCategorySortChange
+	{
+	public:
+		~FScopedCategorySortChange();
+		void SetBlueprint(UBlueprint* Blueprint);
+	private:
+		UBlueprint* Blueprint = nullptr;
+		TArray<FName> Backup = {};
+	} ScopedCategorySortChange;
 };
 
 
