[StaticAnalysisFix] Fix possible nullptr derefrence

diff --git a/Engine/Source/Editor/Kismet/Private/SBlueprintDiff.cpp b/Engine/Source/Editor/Kismet/Private/SBlueprintDiff.cpp
--- a/Engine/Source/Editor/Kismet/Private/SBlueprintDiff.cpp
+++ b/Engine/Source/Editor/Kismet/Private/SBlueprintDiff.cpp
@@ -163,23 +163,25 @@ void SBlueprintDiff::Construct( const FArguments& InArgs)
 	check(InArgs._BlueprintOld || InArgs._BlueprintNew);
 	
 	// make sure that both blueprints have the same category sorting so that they diff properly.
-	if (FPackageName::IsTempPackage(InArgs._BlueprintOld->GetPackage()->GetName()))
+	if (InArgs._BlueprintOld && InArgs._BlueprintNew)
 	{
-		const_cast<UBlueprint*>(InArgs._BlueprintOld)->CategorySorting = InArgs._BlueprintNew->CategorySorting;
-	}
-	else if(FPackageName::IsTempPackage(InArgs._BlueprintNew->GetPackage()->GetName()))
-	{
-		const_cast<UBlueprint*>(InArgs._BlueprintNew)->CategorySorting = InArgs._BlueprintOld->CategorySorting;
-	}
-	else
-	{
-		// Neither New or Old BPs are temp so we need to scope this change to the lifetime of this diff window.
-		// when the window is closed, ScopedCategorySortChange will revert BlueprintOld->CategorySorting back to it's original state.
-		ScopedCategorySortChange.SetBlueprint(const_cast<UBlueprint*>(InArgs._BlueprintOld));
-		const_cast<UBlueprint*>(InArgs._BlueprintOld)->CategorySorting = InArgs._BlueprintNew->CategorySorting;
+		if (FPackageName::IsTempPackage(InArgs._BlueprintOld->GetPackage()->GetName()))
+        {
+        	const_cast<UBlueprint*>(InArgs._BlueprintOld)->CategorySorting = InArgs._BlueprintNew->CategorySorting;
+        }
+        else if(FPackageName::IsTempPackage(InArgs._BlueprintNew->GetPackage()->GetName()))
+        {
+        	const_cast<UBlueprint*>(InArgs._BlueprintNew)->CategorySorting = InArgs._BlueprintOld->CategorySorting;
+        }
+        else
+        {
+        	// Neither New or Old BPs are temp so we need to scope this change to the lifetime of this diff window.
+        	// when the window is closed, ScopedCategorySortChange will revert BlueprintOld->CategorySorting back to it's original state.
+        	ScopedCategorySortChange.SetBlueprint(const_cast<UBlueprint*>(InArgs._BlueprintOld));
+        	const_cast<UBlueprint*>(InArgs._BlueprintOld)->CategorySorting = InArgs._BlueprintNew->CategorySorting;
+        }
 	}
 	
-	
 	PanelOld.Blueprint = InArgs._BlueprintOld;
 	PanelNew.Blueprint = InArgs._BlueprintNew;
 	PanelOld.RevisionInfo = InArgs._OldRevision;
