Avoid warning about modifications on load for packages that are not yet saved

diff --git a/Engine/Source/Editor/Kismet/Private/BlueprintEditor.cpp b/Engine/Source/Editor/Kismet/Private/BlueprintEditor.cpp
--- a/Engine/Source/Editor/Kismet/Private/BlueprintEditor.cpp
+++ b/Engine/Source/Editor/Kismet/Private/BlueprintEditor.cpp
@@ -2014,7 +2014,8 @@ void FBlueprintEditor::CommonInitialization(const TArray<UBlueprint*>& InitBluep
 		// Update the blueprint if required
 		EBlueprintStatus OldStatus = InitBlueprint->Status;
 		EnsureBlueprintIsUpToDate(InitBlueprint);
-		bBlueprintModifiedOnOpen = (InitBlueprint->Status != OldStatus);
+		UPackage* BpPackage = InitBlueprint->GetOutermost();
+		bBlueprintModifiedOnOpen = (InitBlueprint->Status != OldStatus) && !BpPackage->HasAnyPackageFlags(PKG_NewlyCreated);
 
 		// Flag the blueprint as having been opened
 		InitBlueprint->bIsNewlyCreated = false;
