[Feature] made local assets editable in DataAsset diffs
[BugFix] invisible merge buttons were clickable in detail splitter when they shouldn't be
[BugFix] blueprint diffs weren't updating their detail panels when the objects are edited

diff --git a/Engine/Plugins/Editor/EngineAssetDefinitions/Source/Private/AssetDefinition_DataAsset.cpp b/Engine/Plugins/Editor/EngineAssetDefinitions/Source/Private/AssetDefinition_DataAsset.cpp
--- a/Engine/Plugins/Editor/EngineAssetDefinitions/Source/Private/AssetDefinition_DataAsset.cpp
+++ b/Engine/Plugins/Editor/EngineAssetDefinitions/Source/Private/AssetDefinition_DataAsset.cpp
@@ -179,7 +179,12 @@ EAssetCommandResult UAssetDefinition_DataAsset::PerformAssetDiff(const FAssetDif
 		return EAssetCommandResult::Unhandled;
 	}
 	
-	SDetailsDiff::CreateDiffWindow(DiffArgs.OldAsset, DiffArgs.NewAsset, DiffArgs.OldRevision, DiffArgs.NewRevision, UDataAsset::StaticClass());
+	const TSharedRef<SDetailsDiff> DetailsDiff = SDetailsDiff::CreateDiffWindow(DiffArgs.OldAsset, DiffArgs.NewAsset, DiffArgs.OldRevision, DiffArgs.NewRevision, UDataAsset::StaticClass());
+	// allow users to edit NewAsset if it's a local asset
+	if (!FPackageName::IsTempPackage(DiffArgs.NewAsset->GetPackage()->GetName()))
+	{
+		DetailsDiff->SetOutputObject(DiffArgs.NewAsset);
+	}
 	return EAssetCommandResult::Handled;
 }
 
diff --git a/Engine/Plugins/Experimental/Animation/CurveExpression/Source/Editor/Private/AssetDefinition_CurveExpressionsDataAsset.cpp b/Engine/Plugins/Experimental/Animation/CurveExpression/Source/Editor/Private/AssetDefinition_CurveExpressionsDataAsset.cpp
--- a/Engine/Plugins/Experimental/Animation/CurveExpression/Source/Editor/Private/AssetDefinition_CurveExpressionsDataAsset.cpp
+++ b/Engine/Plugins/Experimental/Animation/CurveExpression/Source/Editor/Private/AssetDefinition_CurveExpressionsDataAsset.cpp
@@ -12,8 +12,12 @@ EAssetCommandResult UAssetDefinition_CurveExpressionsDataAsset::PerformAssetDiff
 	{
 		return EAssetCommandResult::Unhandled;
 	}
-	
-	SDetailsDiff::CreateDiffWindow(DiffArgs.OldAsset, DiffArgs.NewAsset, DiffArgs.OldRevision, DiffArgs.NewRevision, UCurveExpressionsDataAsset::StaticClass());
+	const TSharedRef<SDetailsDiff> DetailsDiff = SDetailsDiff::CreateDiffWindow(DiffArgs.OldAsset, DiffArgs.NewAsset, DiffArgs.OldRevision, DiffArgs.NewRevision, UCurveExpressionsDataAsset::StaticClass());
+    // allow users to edit NewAsset if it's a local asset
+	if (!FPackageName::IsTempPackage(DiffArgs.NewAsset->GetPackage()->GetName()))
+	{
+		DetailsDiff->SetOutputObject(DiffArgs.NewAsset);
+	}
 	return EAssetCommandResult::Handled;
 }
 
diff --git a/Engine/Source/Developer/AssetTools/Private/AssetTypeActions/AssetTypeActions_DataAsset.cpp b/Engine/Source/Developer/AssetTools/Private/AssetTypeActions/AssetTypeActions_DataAsset.cpp
--- a/Engine/Source/Developer/AssetTools/Private/AssetTypeActions/AssetTypeActions_DataAsset.cpp
+++ b/Engine/Source/Developer/AssetTools/Private/AssetTypeActions/AssetTypeActions_DataAsset.cpp
@@ -73,7 +73,12 @@ void FAssetTypeActions_DataAsset::GetActions(const TArray<UObject*>& InObjects,
 
 void FAssetTypeActions_DataAsset::PerformAssetDiff(UObject* OldAsset, UObject* NewAsset, const FRevisionInfo& OldRevision, const FRevisionInfo& NewRevision) const
 {
-	SDetailsDiff::CreateDiffWindow(OldAsset, NewAsset, OldRevision, NewRevision, GetSupportedClass());
+	const TSharedRef<SDetailsDiff> DetailsDiff = SDetailsDiff::CreateDiffWindow(OldAsset, NewAsset, OldRevision, NewRevision, GetSupportedClass());
+	// allow users to edit NewAsset if it's a local asset
+	if (!FPackageName::IsTempPackage(NewAsset->GetPackage()->GetName()))
+	{
+		DetailsDiff->SetOutputObject(NewAsset);
+	}
 }
 
 void FAssetTypeActions_DataAsset::ExecuteChangeDataAssetClass(TArray<TWeakObjectPtr<UDataAsset>> InDataAssets)
diff --git a/Engine/Source/Editor/Kismet/Private/DiffControl.cpp b/Engine/Source/Editor/Kismet/Private/DiffControl.cpp
--- a/Engine/Source/Editor/Kismet/Private/DiffControl.cpp
+++ b/Engine/Source/Editor/Kismet/Private/DiffControl.cpp
@@ -420,6 +420,15 @@ TSharedRef<IDetailsView> FDetailsDiffControl::GetDetailsWidget(const UObject* Ob
 	return DetailsDiffs[Object].DetailsWidget();
 }
 
+TSharedPtr<IDetailsView> FDetailsDiffControl::TryGetDetailsWidget(const UObject* Object) const
+{
+	if (const FDetailsDiff* Found = DetailsDiffs.Find(Object))
+	{
+		return Found->DetailsWidget();
+	}
+	return nullptr;
+}
+
 TSharedPtr<FAsyncDetailViewDiff> FDetailsDiffControl::GetDifferencesWithLeft(const UObject* Object) const
 {
 	return PropertyTreeDifferences[Object].Left;
@@ -430,6 +439,11 @@ TSharedPtr<FAsyncDetailViewDiff> FDetailsDiffControl::GetDifferencesWithRight(co
 	return PropertyTreeDifferences[Object].Right;
 }
 
+int32 FDetailsDiffControl::IndexOfObject(const UObject* Object) const
+{
+	return ObjectDisplayOrder.IndexOfByKey(Object);
+}
+
 void FDetailsDiffControl::OnSelectDiffEntry(FPropertySoftPath PropertyName)
 {
 	SelectionCallback.ExecuteIfBound();
diff --git a/Engine/Source/Editor/Kismet/Private/DiffControl.h b/Engine/Source/Editor/Kismet/Private/DiffControl.h
--- a/Engine/Source/Editor/Kismet/Private/DiffControl.h
+++ b/Engine/Source/Editor/Kismet/Private/DiffControl.h
@@ -104,8 +104,10 @@ public:
 
 	TSharedRef<IDetailsView> InsertObject(const UObject* Object, bool bScrollbarOnLeft = false, int32 Index = INDEX_NONE);
 	TSharedRef<IDetailsView> GetDetailsWidget(const UObject* Object) const;
+	TSharedPtr<IDetailsView> TryGetDetailsWidget(const UObject* Object) const;
 	TSharedPtr<FAsyncDetailViewDiff> GetDifferencesWithLeft(const UObject* Object) const;
 	TSharedPtr<FAsyncDetailViewDiff> GetDifferencesWithRight(const UObject* Object) const;
+	int32 IndexOfObject(const UObject* Object) const;
 
 	// to support comment posting, set this to the tree view that contains the comments.
 	virtual void EnableComments(TWeakPtr<STreeView<TSharedPtr<FBlueprintDifferenceTreeEntry>>> TreeView) override;
diff --git a/Engine/Source/Editor/Kismet/Private/SBlueprintDiff.cpp b/Engine/Source/Editor/Kismet/Private/SBlueprintDiff.cpp
--- a/Engine/Source/Editor/Kismet/Private/SBlueprintDiff.cpp
+++ b/Engine/Source/Editor/Kismet/Private/SBlueprintDiff.cpp
@@ -392,6 +392,14 @@ SBlueprintDiff::~SBlueprintDiff()
 	}
 }
 
+void SBlueprintDiff::Tick(const FGeometry& AllottedGeometry, const double InCurrentTime, const float InDeltaTime)
+{
+	if (const TSharedPtr<IDiffControl> DiffControl = ModePanels[CurrentMode].DiffControl)
+	{
+		DiffControl->Tick();
+	}
+}
+
 void SBlueprintDiff::OnCloseAssetEditor(UObject* Asset, EAssetEditorCloseReason CloseReason)
 {
 	if (PanelOld.Blueprint == Asset || PanelNew.Blueprint == Asset || CloseReason == EAssetEditorCloseReason::CloseAllAssetEditors)
diff --git a/Engine/Source/Editor/Kismet/Private/SDetailsDiff.cpp b/Engine/Source/Editor/Kismet/Private/SDetailsDiff.cpp
--- a/Engine/Source/Editor/Kismet/Private/SDetailsDiff.cpp
+++ b/Engine/Source/Editor/Kismet/Private/SDetailsDiff.cpp
@@ -424,27 +424,37 @@ SDetailsDiff::FDiffControl SDetailsDiff::GenerateDetailsPanel(const TFunction<co
 	const TWeakPtr<FDetailsDiffControl> WeakDiffControl = NewDiffControl;
 	OnOutputObjectSetEvent.AddLambda([WeakSplitter, WeakDiffControl, this, Redirector]()
 	{
-		const UObject* CenterObject = Redirector ? Redirector(OutputObject) : OutputObject;
+		const UObject* OutputDetailObject = Redirector ? Redirector(OutputObject) : OutputObject;
 		
 		const TSharedPtr<SDetailsSplitter> Splitter = WeakSplitter.Pin();
 		const TSharedPtr<FDetailsDiffControl> DiffControl = WeakDiffControl.Pin();
 		if (Splitter && DiffControl)
 		{
-			TSharedRef<IDetailsView> DetailsView = DiffControl->InsertObject(CenterObject, false, 1);
+			// if output object is already in panel, don't insert a new one
+			TSharedPtr<IDetailsView> DetailsView = DiffControl->TryGetDetailsWidget(OutputDetailObject);
+			if (DetailsView)
+			{
+				// update readonly status in splitter so that property merge buttons appear
+				const int32 Index = DiffControl->IndexOfObject(OutputDetailObject);
+				Splitter->GetPanel(Index).IsReadonly = false;
+			}
+			else
+			{
+				DetailsView = DiffControl->InsertObject(OutputDetailObject, false, 1);
+				// insert the output object as a central panel
+				Splitter->AddSlot(
+					SDetailsSplitter::Slot()
+						.DetailsView(DetailsView)
+						.Value(0.5f)
+						.IsReadonly(false)
+						.DifferencesWithRightPanel(DiffControl.ToSharedRef(), &FDetailsDiffControl::GetDifferencesWithRight, (const UObject*)OutputDetailObject),
+					1 // insert between left and right panel (index 1)
+				);
+			}
 
 			// allow user to edit the output panel
 			DetailsView->SetIsPropertyEditingEnabledDelegate(FIsPropertyEditingEnabled::CreateStatic([]{return true; }));
 			
-			// insert the output object as a central panel
-			Splitter->AddSlot(
-				SDetailsSplitter::Slot()
-					.DetailsView(DiffControl->GetDetailsWidget(CenterObject))
-					.Value(0.5f)
-					.IsReadonly(false)
-					.DifferencesWithRightPanel(DiffControl.ToSharedRef(), &FDetailsDiffControl::GetDifferencesWithRight, (const UObject*)CenterObject),
-				1 // insert between left and right panel (index 1)
-			);
-
 			// highlight merge conflicts
 			TMap<FString, TMap<FPropertySoftPath, FLinearColor>> Highlights;
 			for (auto& [ObjectPath, Properties] : MergeConflicts)
diff --git a/Engine/Source/Editor/Kismet/Public/SBlueprintDiff.h b/Engine/Source/Editor/Kismet/Public/SBlueprintDiff.h
--- a/Engine/Source/Editor/Kismet/Public/SBlueprintDiff.h
+++ b/Engine/Source/Editor/Kismet/Public/SBlueprintDiff.h
@@ -149,6 +149,8 @@ public:
 	void Construct(const FArguments& InArgs);
 	virtual ~SBlueprintDiff();
 
+	virtual void Tick(const FGeometry& AllottedGeometry, const double InCurrentTime, const float InDeltaTime) override;
+
 	/** Called when a new Graph is clicked on by user */
 	void OnGraphChanged(FGraphToDiff* Diff);
 
diff --git a/Engine/Source/Editor/PropertyEditor/Private/SDetailsSplitter.cpp b/Engine/Source/Editor/PropertyEditor/Private/SDetailsSplitter.cpp
--- a/Engine/Source/Editor/PropertyEditor/Private/SDetailsSplitter.cpp
+++ b/Engine/Source/Editor/PropertyEditor/Private/SDetailsSplitter.cpp
@@ -345,6 +345,11 @@ void SDetailsSplitter::AddSlot(const FSlot::FSlotArguments& SlotArgs, int32 Inde
 	}, Index);
 }
 
+SDetailsSplitter::FPanel& SDetailsSplitter::GetPanel(int32 Index)
+{
+	return Panels[Index];
+}
+
 void SDetailsSplitter::AddHighlights(const TMap<FString, TMap<FPropertySoftPath, FLinearColor>>& Highlights)
 {
 	CustomHighlights = Highlights;
@@ -500,7 +505,7 @@ FReply SDetailsSplitter::OnMouseMove(const FGeometry& MyGeometry, const FPointer
 		
 		const int32 RightIndex = LeftIndex + 1;
 		const FPanel& RightPanel = Panels[RightIndex];
-		if (LeftPanel.IsReadonly.Get(false) && RightPanel.IsReadonly.Get(false))
+		if (LeftPanel.IsReadonly.Get(true) && RightPanel.IsReadonly.Get(true))
 		{
 			continue;
 		}
@@ -534,7 +539,7 @@ FReply SDetailsSplitter::OnMouseMove(const FGeometry& MyGeometry, const FPointer
 					return ETreeTraverseControl::Continue;
 				}
 
-				if (LeftPropertyRect.IsValid() && !RightPanel.IsReadonly.Get(false))
+				if (LeftPropertyRect.IsValid() && !RightPanel.IsReadonly.Get(true))
 				{
 					const FSlateRect CopyButtonZoneLeftToRight = FSlateRect(
 						LeftPropertyRect.Right,
@@ -558,7 +563,7 @@ FReply SDetailsSplitter::OnMouseMove(const FGeometry& MyGeometry, const FPointer
 					}
 				}
 				
-				if (RightPropertyRect.IsValid() && !LeftPanel.IsReadonly.Get(false))
+				if (RightPropertyRect.IsValid() && !LeftPanel.IsReadonly.Get(true))
 				{
 					const FSlateRect CopyButtonZoneRightToLeft = FSlateRect(
 						RightPropertyRect.Left - 15.f,
diff --git a/Engine/Source/Editor/PropertyEditor/Public/SDetailsSplitter.h b/Engine/Source/Editor/PropertyEditor/Public/SDetailsSplitter.h
--- a/Engine/Source/Editor/PropertyEditor/Public/SDetailsSplitter.h
+++ b/Engine/Source/Editor/PropertyEditor/Public/SDetailsSplitter.h
@@ -18,10 +18,16 @@ public:
 			/** When the RuleSize is set to FractionOfParent, the size of the slot is the Value percentage of its parent size. */
 			SLATE_ATTRIBUTE(float, Value)
 			SLATE_ARGUMENT(TSharedPtr<IDetailsView>, DetailsView)
-			SLATE_ATTRIBUTE(bool, IsReadonly)
+			SLATE_ATTRIBUTE(bool, IsReadonly) // default true
 			SLATE_ATTRIBUTE(TSharedPtr<FAsyncDetailViewDiff>, DifferencesWithRightPanel)
 		SLATE_SLOT_END_ARGS()
 	};
+	struct FPanel
+	{
+		TSharedPtr<IDetailsView> DetailsView;
+		TAttribute<bool> IsReadonly;
+		TAttribute<TSharedPtr<FAsyncDetailViewDiff>> DiffRight;
+	};
 	static FSlot::FSlotArguments Slot();
 	
 	SLATE_BEGIN_ARGS(SDetailsSplitter)
@@ -40,18 +46,12 @@ public:
 	virtual FReply OnMouseButtonDown(const FGeometry& MyGeometry, const FPointerEvent& MouseEvent) override;
 	
 	void AddSlot(const FSlot::FSlotArguments& SlotArgs, int32 Index = INDEX_NONE);
+	FPanel& GetPanel(int32 Index);
 
 	void AddHighlights(const TMap<FString, TMap<FPropertySoftPath, FLinearColor>>& Highlights);
 
 private:
 	
-	struct FPanel
-	{
-		TSharedPtr<IDetailsView> DetailsView;
-		TAttribute<bool> IsReadonly;
-		TAttribute<TSharedPtr<FAsyncDetailViewDiff>> DiffRight;
-	};
-	
 	enum class EPropertyCopyDirection
 	{
 		Copy_None,
