[BugFix] Review comments are being archived automatically due to undefined behavior
FReviewComment.bIsClosed has no default initialization so it's being assigned to garbage (oops)

diff --git a/Engine/Source/Editor/Kismet/Public/ReviewComments.h b/Engine/Source/Editor/Kismet/Public/ReviewComments.h
--- a/Engine/Source/Editor/Kismet/Public/ReviewComments.h
+++ b/Engine/Source/Editor/Kismet/Public/ReviewComments.h
@@ -106,7 +106,7 @@ public:
 	FReviewCommentContext Context;
 	
 	// flags
-	bool bIsClosed : 1;
+	bool bIsClosed : 1 = false;
 
 	FString ToJson() const;
 	TSharedPtr<FJsonObject> ToJsonObject() const;
