Allow you to delete functions/variables in the "My Blueprint" window with backspace as well as the delete key. This makes it more consistent with the rest of the editor on all platforms

diff --git a/Engine/Source/Editor/Kismet/Private/SMyBlueprint.cpp b/Engine/Source/Editor/Kismet/Private/SMyBlueprint.cpp
--- a/Engine/Source/Editor/Kismet/Private/SMyBlueprint.cpp
+++ b/Engine/Source/Editor/Kismet/Private/SMyBlueprint.cpp
@@ -138,7 +138,7 @@ void FMyBlueprintCommands::RegisterCommands()
 	UI_COMMAND( FocusNode, "Focus", "Focuses on the associated node", EUserInterfaceActionType::Button, FInputChord() );
 	UI_COMMAND( FocusNodeInNewTab, "Focus in New Tab", "Focuses on the associated node in a new tab", EUserInterfaceActionType::Button, FInputChord() );
 	UI_COMMAND( ImplementFunction, "Implement event", "Implements this overridable function as a new event.", EUserInterfaceActionType::Button, FInputChord() );
-	UI_COMMAND( DeleteEntry, "Delete", "Deletes this function or variable from this blueprint.", EUserInterfaceActionType::Button, FInputChord(EKeys::Platform_Delete));
+	UI_COMMAND( DeleteEntry, "Delete", "Deletes this function or variable from this blueprint.", EUserInterfaceActionType::Button, FInputChord(EKeys::Delete), FInputChord(EKeys::BackSpace));
 	UI_COMMAND( PasteVariable, "Paste Variable", "Pastes the variable to this blueprint.", EUserInterfaceActionType::Button, FInputChord());
 	UI_COMMAND( PasteLocalVariable, "Paste Local Variable", "Pastes the variable to this scope.", EUserInterfaceActionType::Button, FInputChord());
 	UI_COMMAND( PasteFunction, "Paste Function", "Pastes the function to this blueprint.", EUserInterfaceActionType::Button, FInputChord());
