Guard usage of FBPVariableDescription::GetMetaData

diff --git a/Engine/Source/Editor/Kismet/Private/SFieldNotificationCheckList.cpp b/Engine/Source/Editor/Kismet/Private/SFieldNotificationCheckList.cpp
--- a/Engine/Source/Editor/Kismet/Private/SFieldNotificationCheckList.cpp
+++ b/Engine/Source/Editor/Kismet/Private/SFieldNotificationCheckList.cpp
@@ -95,7 +95,8 @@ ECheckBoxState SFieldNotificationCheckList::OnVariableCheckboxState(FName OtherN
 		if (VarIndex != INDEX_NONE)
 		{
 			// Parse the metadata of the variable by delimiter | and check if OtherName is in there.
-			FString FieldNotifyValues = BlueprintObj->NewVariables[VarIndex].GetMetaData(FBlueprintMetadata::MD_FieldNotify);
+			const FBPVariableDescription& VariableDescription = BlueprintObj->NewVariables[VarIndex];
+			FString FieldNotifyValues = VariableDescription.HasMetaData(FBlueprintMetadata::MD_FieldNotify) ? VariableDescription.GetMetaData(FBlueprintMetadata::MD_FieldNotify) : FString();
 			TArray<FString> ListOfFieldNotifies;
 			const TCHAR* Delimiter = TEXT("|");
 
@@ -114,7 +115,8 @@ void SFieldNotificationCheckList::OnVariableCheckBoxChanged(ECheckBoxState InNew
 		const int32 VarIndex = FBlueprintEditorUtils::FindNewVariableIndex(BlueprintObj, FieldName);
 		if (VarIndex != INDEX_NONE)
 		{
-			FString FieldNotifyValues = BlueprintObj->NewVariables[VarIndex].GetMetaData(FBlueprintMetadata::MD_FieldNotify);
+			const FBPVariableDescription& VariableDescription = BlueprintObj->NewVariables[VarIndex];
+			FString FieldNotifyValues = VariableDescription.HasMetaData(FBlueprintMetadata::MD_FieldNotify) ? VariableDescription.GetMetaData(FBlueprintMetadata::MD_FieldNotify) : FString();
 			if (InNewState == ECheckBoxState::Checked)
 			{
 				const TCHAR* Delimiter = TEXT("|");
@@ -283,7 +285,8 @@ EVisibility SFieldNotificationCheckList::GetFieldNotifyIconVisibility() const
 		const int32 VarIndex = FBlueprintEditorUtils::FindNewVariableIndex(BlueprintObj, FieldName);
 		if (VarIndex != INDEX_NONE)
 		{
-			FString FieldNotifyValues = BlueprintObj->NewVariables[VarIndex].GetMetaData(FBlueprintMetadata::MD_FieldNotify);
+			const FBPVariableDescription& VariableDescription = BlueprintObj->NewVariables[VarIndex];
+			FString FieldNotifyValues = VariableDescription.HasMetaData(FBlueprintMetadata::MD_FieldNotify) ? VariableDescription.GetMetaData(FBlueprintMetadata::MD_FieldNotify) : FString();
 			return FieldNotifyValues.Len() > 0 ? EVisibility::Visible : EVisibility::Collapsed;
 		}
 	}
