Significantly improve FindInBlueprint performance and memory usage on startup before a search is requested:
Remove a redundant string copy inside FImaginaryBlueprint and defer reading the AssetData tag when bDisableImmediateAssetDiscovery is true
Add HasEncodedValue/ClearEncodedValue wrappers to FSearchData

diff --git a/Engine/Source/Editor/Kismet/Private/FindInBlueprintManager.cpp b/Engine/Source/Editor/Kismet/Private/FindInBlueprintManager.cpp
--- a/Engine/Source/Editor/Kismet/Private/FindInBlueprintManager.cpp
+++ b/Engine/Source/Editor/Kismet/Private/FindInBlueprintManager.cpp
@@ -1412,11 +1412,12 @@ public:
 				CSV_CUSTOM_STAT(FindInBlueprint, IndexedAssetCountThisFrame, 1, ECsvCustomStatOp::Accumulate);
 
 				FSoftObjectPath AssetPath = AssetPathsToIndex[ArrayIdx];
-				FSearchData SearchData = FFindInBlueprintSearchManager::Get().GetSearchDataForAssetPath(AssetPath);
+				FFindInBlueprintSearchManager& FindManager = FFindInBlueprintSearchManager::Get();
+				FSearchData SearchData = FindManager.GetSearchDataForAssetPath(AssetPath);
 				if (SearchData.IsValid() && !SearchData.IsMarkedForDeletion() && !SearchData.IsIndexingCompleted())
 				{
 					// Generate the metadata tag value if it was not previously cached or loaded.
-					if (SearchData.Value.Len() == 0)
+					if (!SearchData.HasEncodedValue())
 					{
 						// This must be done on the main thread, so enqueue it and continue.
 						Controller->AddAssetPathToGatherQueue(AssetPath);
@@ -1426,12 +1427,8 @@ public:
 						if (bEnableFullIndexingPass)
 						{
 							// Unpack the metadata tag and rebuild the index for this asset.
-							const FString AssetPathAsString = AssetPath.ToString();
-							if (FiBSerializationHelpers::ValidateSearchDataVersionInfo(AssetPathAsString, SearchData.Value, SearchData.VersionInfo))
+							if (FindManager.ProcessEncodedValueForUnloadedBlueprint(SearchData))
 							{
-								SearchData.ImaginaryBlueprint = MakeShareable(new FImaginaryBlueprint(FPaths::GetBaseFilename(AssetPathAsString), AssetPathAsString, SearchData.ParentClass, SearchData.Interfaces, SearchData.Value, SearchData.VersionInfo));
-								SearchData.Value.Empty();
-
 								// Build the full index using a BFS traversal.
 								TArray<FImaginaryFiBDataSharedPtr> IndexNodes = { SearchData.ImaginaryBlueprint };
 								while (IndexNodes.Num() > 0)
@@ -1455,7 +1452,7 @@ public:
 							SearchData.StateFlags |= ESearchDataStateFlags::IsIndexed;
 
 							// Update this entry in the search database (thread-safe).
-							FFindInBlueprintSearchManager::Get().ApplySearchDataToDatabase(MoveTemp(SearchData));
+							FindManager.ApplySearchDataToDatabase(MoveTemp(SearchData));
 						}
 
 						// Signal that indexing has been completed for this asset path.
@@ -2165,17 +2162,14 @@ void FFindInBlueprintSearchManager::OnAssetAdded(const FAssetData& InAssetData)
 		return;
 	}
 
-	if (InAssetData.IsAssetLoaded())
+	if (UObject* AssetObject = InAssetData.FastGetAsset(false))
 	{
-		if (UObject* AssetObject = InAssetData.GetAsset())
+		if (ensureMsgf(AssetObject->IsA(AssetClass), TEXT("AssetClass (%s) matched handler, but does not match actual object type (%s) for asset: %s."), *AssetClass->GetName(), *AssetObject->GetClass()->GetName(), *AssetObject->GetPathName()))
 		{
-			if (ensureMsgf(AssetObject->IsA(AssetClass), TEXT("AssetClass (%s) matched handler, but does not match actual object type (%s) for asset: %s."), *AssetClass->GetName(), *AssetObject->GetClass()->GetName(), *AssetObject->GetPathName()))
+			UBlueprint* Blueprint = Handler->RetrieveBlueprint(AssetObject);
+			if (Blueprint)
 			{
-				UBlueprint* Blueprint = Handler->RetrieveBlueprint(AssetObject);
-				if (Blueprint)
-				{
-					AddOrUpdateBlueprintSearchMetadata(Blueprint);
-				}
+				AddOrUpdateBlueprintSearchMetadata(Blueprint);
 			}
 		}
 	}
@@ -2191,14 +2185,23 @@ void FFindInBlueprintSearchManager::AddUnloadedBlueprintSearchMetadata(const FAs
 	FAssetDataTagMapSharedView::FFindTagResult Result = InAssetData.TagsAndValues.FindTag(FBlueprintTags::FindInBlueprintsData);
 	if (Result.IsSet())
 	{
-		const FString& FiBVersionedSearchData = Result.GetValue();
-		if (FiBVersionedSearchData.Len() == 0)
+		if (bDisableImmediateAssetDiscovery)
 		{
-			UnindexedAssets.Add(InAssetData.GetSoftObjectPath());
+			// If the versioned key is set at all, we assume it is valid and will parse it later
+			ExtractUnloadedFiBData(InAssetData, nullptr, FBlueprintTags::FindInBlueprintsData, EFiBVersion::FIB_VER_NONE);
 		}
 		else
 		{
-			ExtractUnloadedFiBData(InAssetData, FiBVersionedSearchData, EFiBVersion::FIB_VER_NONE);
+			// Extract it now
+			FString FiBVersionedSearchData = Result.GetValue();
+			if (FiBVersionedSearchData.Len() == 0)
+			{
+				UnindexedAssets.Add(InAssetData.GetSoftObjectPath());
+			}
+			else
+			{
+				ExtractUnloadedFiBData(InAssetData, &FiBVersionedSearchData, NAME_None, EFiBVersion::FIB_VER_NONE);
+			}
 		}
 	}
 	else
@@ -2207,18 +2210,57 @@ void FFindInBlueprintSearchManager::AddUnloadedBlueprintSearchMetadata(const FAs
 		FAssetDataTagMapSharedView::FFindTagResult ResultLegacy = InAssetData.TagsAndValues.FindTag(FBlueprintTags::UnversionedFindInBlueprintsData);
 		if (ResultLegacy.IsSet())
 		{
-			ExtractUnloadedFiBData(InAssetData, ResultLegacy.GetValue(), EFiBVersion::FIB_VER_BASE);
+			if (bDisableImmediateAssetDiscovery)
+			{
+				ExtractUnloadedFiBData(InAssetData, nullptr, FBlueprintTags::UnversionedFindInBlueprintsData, EFiBVersion::FIB_VER_BASE);
+			}
+			else
+			{
+				FString FiBUnversionedSearchData = Result.GetValue();
+				ExtractUnloadedFiBData(InAssetData, &FiBUnversionedSearchData, NAME_None, EFiBVersion::FIB_VER_BASE);
+			}
 		}
 		// The asset has no FiB data, keep track of it so we can inform the user
 		else
 		{
 			UnindexedAssets.Add(InAssetData.GetSoftObjectPath());
 		}
+	}
+}
+
+bool FFindInBlueprintSearchManager::ProcessEncodedValueForUnloadedBlueprint(FSearchData& SearchData)
+{
+	FString TempEncodedString;
+	if (!SearchData.AssetKeyForValue.IsNone() && SearchData.Value.IsEmpty())
+	{
+		// Get the string from the asset registry now
+		FAssetData AssetData;
+		AssetData = AssetRegistryModule->Get().GetAssetByObjectPath(SearchData.AssetPath);
+		if (ensure(AssetData.IsValid()))
+		{
+			FAssetDataTagMapSharedView::FFindTagResult Result = AssetData.TagsAndValues.FindTag(SearchData.AssetKeyForValue);
+			if (Result.IsSet())
+			{
+				// This makes a large string copy because the version in the asset data cannot be treated as an FString
+				TempEncodedString = Result.GetValue();
+			}
+		}
+	}
+	const FString& EncodedString = !TempEncodedString.IsEmpty() ? TempEncodedString : SearchData.Value;
+	const FString AssetPath = SearchData.AssetPath.ToString();
+	if (FiBSerializationHelpers::ValidateSearchDataVersionInfo(AssetPath, EncodedString, SearchData.VersionInfo))
+	{
+		// Parse the data into json and then clear the memory
+		SearchData.ImaginaryBlueprint = MakeShareable(new FImaginaryBlueprint(FPaths::GetBaseFilename(AssetPath), AssetPath, SearchData.ParentClass, SearchData.Interfaces, EncodedString, SearchData.VersionInfo));
+		SearchData.ClearEncodedValue();
 
+		return true;
 	}
+
+	return false;
 }
 
-void FFindInBlueprintSearchManager::ExtractUnloadedFiBData(const FAssetData& InAssetData, const FString& InFiBData, EFiBVersion InFiBDataVersion)
+void FFindInBlueprintSearchManager::ExtractUnloadedFiBData(const FAssetData& InAssetData, FString* InFiBData, FName InKeyForFiBData, EFiBVersion InFiBDataVersion)
 {
 	CSV_SCOPED_TIMING_STAT(FindInBlueprint, ExtractUnloadedFiBData);
 	CSV_CUSTOM_STAT(FindInBlueprint, ExtractUnloadedCountThisFrame, 1, ECsvCustomStatOp::Accumulate);
@@ -2281,13 +2323,22 @@ void FFindInBlueprintSearchManager::ExtractUnloadedFiBData(const FAssetData& InA
 		}
 	}
 
-	NewSearchData.Value = InFiBData;
+	if (InFiBData)
+	{
+		NewSearchData.Value = MoveTemp(*InFiBData);
+	}
+	else
+	{
+		NewSearchData.Value.Empty();
+	}
+
+	NewSearchData.AssetKeyForValue = InKeyForFiBData;
 
 	// This will be set to 'None' if the data is versioned. Deserialization of the actual version from the tag value is deferred until later.
 	NewSearchData.VersionInfo.FiBDataVersion = InFiBDataVersion;
 
 	// In these modes, or if there is no tag data, no additional indexing work is deferred for unloaded assets.
-	if (!bDisableDeferredIndexing && !bDisableThreadedIndexing && NewSearchData.Value.Len() > 0)
+	if (!bDisableDeferredIndexing && !bDisableThreadedIndexing && NewSearchData.HasEncodedValue())
 	{
 		// Add it to the list of assets that require a full index rebuild from the metadata. This work will not block the main thread and is decoupled from the search thread.
 		PendingAssets.Add(NewSearchData.AssetPath);
@@ -2579,7 +2630,7 @@ void FFindInBlueprintSearchManager::AddOrUpdateBlueprintSearchMetadata(UBlueprin
 	if (SearchData.IsValid())
 	{
 		// Clear any previously-gathered data.
-		SearchData.Value.Empty();
+		SearchData.ClearEncodedValue();
 
 		// Update version info stored in database. This indicates which format to use when regenerating the tag value.
 		SearchData.VersionInfo = FSearchDataVersionInfo::Current;
@@ -2817,13 +2868,11 @@ bool FFindInBlueprintSearchManager::ContinueSearchQuery(const FStreamSearch* InS
 			SearchData = GetSearchDataForAssetPath(SearchData.AssetPath);
 
 			// If there is FiB data, parse it into an ImaginaryBlueprint
-			if (SearchData.IsValid() && SearchData.Value.Len() > 0)
+			if (SearchData.IsValid() && SearchData.HasEncodedValue())
 			{
-				const FString AssetPath = SearchData.AssetPath.ToString();
-				if (FiBSerializationHelpers::ValidateSearchDataVersionInfo(AssetPath, SearchData.Value, SearchData.VersionInfo))
+				if (ProcessEncodedValueForUnloadedBlueprint(SearchData))
 				{
-					SearchData.ImaginaryBlueprint = MakeShareable(new FImaginaryBlueprint(FPaths::GetBaseFilename(SearchData.AssetPath.ToString()), SearchData.AssetPath.ToString(), SearchData.ParentClass, SearchData.Interfaces, SearchData.Value, SearchData.VersionInfo));
-					SearchData.Value.Empty();
+					check(SearchData.ImaginaryBlueprint.IsValid());
 
 					// In the case of parallel global searches, two search threads may be looking at the same entry. Thus, we only allow one search thread to be actively parsing JSON nodes.
 					SearchData.ImaginaryBlueprint->EnableInterlockedParsing();
@@ -3438,7 +3487,7 @@ bool FFindInBlueprintSearchManager::IsAsyncSearchQueryInProgress() const
 	return ActiveSearchQueries.Num() > 0;
 }
 
-TSharedPtr< FJsonObject > FFindInBlueprintSearchManager::ConvertJsonStringToObject(FSearchDataVersionInfo InVersionInfo, FString InJsonString, TMap<int32, FText>& OutFTextLookupTable)
+TSharedPtr< FJsonObject > FFindInBlueprintSearchManager::ConvertJsonStringToObject(FSearchDataVersionInfo InVersionInfo, const FString& InJsonString, TMap<int32, FText>& OutFTextLookupTable)
 {
 	/** The searchable data is more complicated than a Json string, the Json being the main searchable body that is parsed. Below is a diagram of the full data:
 	 *  | int32 "Version" | int32 "Size" | TMap "Lookup Table" | Json String |
diff --git a/Engine/Source/Editor/Kismet/Private/ImaginaryBlueprintData.cpp b/Engine/Source/Editor/Kismet/Private/ImaginaryBlueprintData.cpp
--- a/Engine/Source/Editor/Kismet/Private/ImaginaryBlueprintData.cpp
+++ b/Engine/Source/Editor/Kismet/Private/ImaginaryBlueprintData.cpp
@@ -544,12 +544,11 @@ void FCategorySectionHelper::DumpParsedObject_Internal(FArchive& Ar) const
 //////////////////////////////////////////
 // FImaginaryBlueprint
 
-FImaginaryBlueprint::FImaginaryBlueprint(FString InBlueprintName, FString InBlueprintPath, FString InBlueprintParentClass, TArray<FString>& InInterfaces, FString InUnparsedStringData, FSearchDataVersionInfo InVersionInfo)
+FImaginaryBlueprint::FImaginaryBlueprint(const FString& InBlueprintName, const FString& InBlueprintPath, const FString& InBlueprintParentClass, const TArray<FString>& InInterfaces, const FString& InUnparsedStringData, FSearchDataVersionInfo InVersionInfo)
 	: FImaginaryFiBData(nullptr)
 	, BlueprintPath(InBlueprintPath)
-	, UnparsedStringData(InUnparsedStringData)
 {
-	ParseToJson(InVersionInfo);
+	ParseToJson(InVersionInfo, InUnparsedStringData);
 	LookupTablePtr = &LookupTable;
 	ParsedTagsAndValues.Add(FindInBlueprintsHelpers::FSimpleFTextKeyStorage(FFindInBlueprintSearchTags::FiB_Name), FSearchableValueInfo(FFindInBlueprintSearchTags::FiB_Name, FText::FromString(InBlueprintName), ESearchableValueStatus::ExplicitySearchable));
 	ParsedTagsAndValues.Add(FindInBlueprintsHelpers::FSimpleFTextKeyStorage(FFindInBlueprintSearchTags::FiB_Path), FSearchableValueInfo(FFindInBlueprintSearchTags::FiB_Path, FText::FromString(InBlueprintPath), ESearchableValueStatus::ExplicitySearchable));
@@ -558,7 +557,7 @@ FImaginaryBlueprint::FImaginaryBlueprint(FString InBlueprintName, FString InBlue
 	TSharedPtr< FCategorySectionHelper, ESPMode::ThreadSafe > InterfaceCategory = MakeShareable(new FCategorySectionHelper(nullptr, &LookupTable, FFindInBlueprintSearchTags::FiB_Interfaces, true));
 	for( int32 InterfaceIdx = 0; InterfaceIdx < InInterfaces.Num(); ++InterfaceIdx)
 	{
-		FString& Interface = InInterfaces[InterfaceIdx];
+		const FString& Interface = InInterfaces[InterfaceIdx];
 		FText Key = FText::FromString(FString::FromInt(InterfaceIdx));
 		FSearchableValueInfo Value(Key, FText::FromString(Interface), ESearchableValueStatus::ExplicitySearchable);
 		InterfaceCategory->AddKeyValuePair(FFindInBlueprintSearchTags::FiB_Interfaces, Value);
@@ -608,7 +607,7 @@ bool FImaginaryBlueprint::CanCallFilter(ESearchQueryFilter InSearchQueryFilter)
 		FImaginaryFiBData::CanCallFilter(InSearchQueryFilter);
 }
 
-void FImaginaryBlueprint::ParseToJson(FSearchDataVersionInfo InVersionInfo)
+void FImaginaryBlueprint::ParseToJson(FSearchDataVersionInfo InVersionInfo, const FString& UnparsedStringData)
 {
 	UnparsedJsonObject = FFindInBlueprintSearchManager::ConvertJsonStringToObject(InVersionInfo, UnparsedStringData, LookupTable);
 }
diff --git a/Engine/Source/Editor/Kismet/Public/FindInBlueprintManager.h b/Engine/Source/Editor/Kismet/Public/FindInBlueprintManager.h
--- a/Engine/Source/Editor/Kismet/Public/FindInBlueprintManager.h
+++ b/Engine/Source/Editor/Kismet/Public/FindInBlueprintManager.h
@@ -192,9 +192,12 @@ struct FSearchData
 	/** The full asset path this search data is associated with of the form /Game/Path/To/Package.Package */
 	FSoftObjectPath AssetPath;
 
-	/** Search data block for the Blueprint */
+	/** Encoded search data block for the Blueprint, this will not always be set if it's already been parsed */
 	FString Value;
 
+	/** Key to use to look up the encoded search data from an FAssetData, if this is set Value will probably be empty */
+	FName AssetKeyForValue;
+
 	/** Parent Class */
 	FString ParentClass;
 
@@ -216,11 +219,25 @@ struct FSearchData
 	{
 	}
 
+	/** True if this represents a valid asset */
 	bool IsValid() const
 	{
 		return !AssetPath.IsNull();
 	}
 
+	/** True if this has an encoded value that has yet to be parsed */
+	bool HasEncodedValue() const
+	{
+		return !Value.IsEmpty() || !AssetKeyForValue.IsNone();
+	}
+
+	/** Clear the encoded value after parsing or getting new data */
+	void ClearEncodedValue()
+	{
+		Value.Reset();
+		AssetKeyForValue = NAME_None;
+	}
+
 	bool IsIndexingCompleted() const
 	{
 		return EnumHasAllFlags(StateFlags, ESearchDataStateFlags::IsIndexed);
@@ -581,6 +598,9 @@ public:
 	 */
 	FSearchData QuerySingleBlueprint(UBlueprint* InBlueprint, bool bInRebuildSearchData);
 
+	/** Processes the encoded string value in the SearchData into the intermediate format, return true if string and version were valid */
+	bool ProcessEncodedValueForUnloadedBlueprint(FSearchData& SearchData);
+
 	/** Returns the number of unindexed Blueprints, either due to not having been indexed before, or AR data being out-of-date */
 	int32 GetNumberUnindexedAssets() const;
 
@@ -684,7 +704,7 @@ public:
 	static FString ConvertFTextToHexString(FText InValue);
 
 	/** Given a fully constructed Find-in-Blueprint FString of searchable data, will parse and construct a JsonObject */
-	static TSharedPtr< class FJsonObject > ConvertJsonStringToObject(FSearchDataVersionInfo InVersionInfo, FString InJsonString, TMap<int32, FText>& OutFTextLookupTable);
+	static TSharedPtr< class FJsonObject > ConvertJsonStringToObject(FSearchDataVersionInfo InVersionInfo, const FString& InJsonString, TMap<int32, FText>& OutFTextLookupTable);
 
 	/** Generates a human-readable search index for the given Blueprint (for debugging purposes) */
 	static FString GenerateSearchIndexForDebugging(UBlueprint* InBlueprint);
@@ -744,7 +764,7 @@ private:
 	void AddUnloadedBlueprintSearchMetadata(const FAssetData& InAssetData);
 
 	/** Begins the process of extracting FiB data from an unloaded asset */
-	void ExtractUnloadedFiBData(const FAssetData& InAssetData, const FString& InFiBData, EFiBVersion InFiBDataVersion);
+	void ExtractUnloadedFiBData(const FAssetData& InAssetData, FString* InFiBData, FName InKeyForFiBData, EFiBVersion InFiBDataVersion);
 
 	/** Determines the global find results tab label */
 	FText GetGlobalFindResultsTabLabel(int32 TabIdx);
diff --git a/Engine/Source/Editor/Kismet/Public/ImaginaryBlueprintData.h b/Engine/Source/Editor/Kismet/Public/ImaginaryBlueprintData.h
--- a/Engine/Source/Editor/Kismet/Public/ImaginaryBlueprintData.h
+++ b/Engine/Source/Editor/Kismet/Public/ImaginaryBlueprintData.h
@@ -339,7 +339,7 @@ protected:
 class FImaginaryBlueprint : public FImaginaryFiBData
 {
 public:
-	FImaginaryBlueprint(FString InBlueprintName, FString InBlueprintPath, FString InBlueprintParentClass, TArray<FString>& InInterfaces, FString InUnparsedStringData, FSearchDataVersionInfo InVersionInfo);
+	FImaginaryBlueprint(const FString& InBlueprintName, const FString& InBlueprintPath, const FString& InBlueprintParentClass, const TArray<FString>& InInterfaces, const FString& InUnparsedStringData, FSearchDataVersionInfo InVersionInfo);
 
 	/** FImaginaryFiBData Interface */
 	virtual bool IsCompatibleWithFilter(ESearchQueryFilter InSearchQueryFilter) const override;
@@ -361,15 +361,12 @@ protected:
 	void ParseComponents(TSharedPtr< FJsonObject > InJsonObject, TArray<FImaginaryFiBDataSharedPtr>& OutParsedChildData);
 
 	/** Parses a raw string of Json to a Json object hierarchy */
-	void ParseToJson(FSearchDataVersionInfo InVersionInfo);
+	void ParseToJson(FSearchDataVersionInfo InVersionInfoconst, const FString& UnparsedStringData);
 
 protected:
 	/** The path for this Blueprint */
 	FString BlueprintPath;
 
-	/** The raw Json string yet to be parsed */
-	FString UnparsedStringData;
-
 	/** Lookup table used as a compression tool for the FTexts stored in the Json object */
 	TMap<int32, FText> LookupTable;
 };
