[CrashFix] fixed missing nullptr check in SDetailsDiff::GenerateDifferencesList

diff --git a/Engine/Source/Editor/Kismet/Private/SDetailsDiff.cpp b/Engine/Source/Editor/Kismet/Private/SDetailsDiff.cpp
--- a/Engine/Source/Editor/Kismet/Private/SDetailsDiff.cpp
+++ b/Engine/Source/Editor/Kismet/Private/SDetailsDiff.cpp
@@ -378,7 +378,7 @@ void SDetailsDiff::GenerateDifferencesList()
 	};
 	
 	TFunction<const UObject*(const UObject*)> Redirector;
-	if (PanelOld.Object->IsA<UBlueprint>() && PanelNew.Object->IsA<UBlueprint>())
+	if ((!PanelOld.Object || PanelOld.Object->IsA<UBlueprint>()) && (!PanelNew.Object || PanelNew.Object->IsA<UBlueprint>()))
 	{
 		// Blueprints diff their GeneratedClass CDO in the details panel instead
 		Redirector = GetBlueprintCDO;
