FieldNotification: Add the FieldNotify for function on the skeletal class. Previously, only the generated class had them.
#prefligth 64c9a98e9f17c746d5c78a5e

diff --git a/Engine/Source/Editor/Kismet/Private/BlueprintCompilationManager.cpp b/Engine/Source/Editor/Kismet/Private/BlueprintCompilationManager.cpp
--- a/Engine/Source/Editor/Kismet/Private/BlueprintCompilationManager.cpp
+++ b/Engine/Source/Editor/Kismet/Private/BlueprintCompilationManager.cpp
@@ -21,6 +21,7 @@
 #include "FindInBlueprintManager.h"
 #include "HAL/LowLevelMemTracker.h"
 #include "IMessageLogListing.h"
+#include "INotifyFieldValueChanged.h"
 #include "K2Node_CreateDelegate.h"
 #include "K2Node_CustomEvent.h"
 #include "K2Node_FunctionEntry.h"
@@ -3000,6 +3001,12 @@ UClass* FBlueprintCompilationManagerImpl::FastGenerateSkeletonClass(UBlueprint*
 					CompilerContext.SetCalculatedMetaDataAndFlags(NewFunction, EntryNode, Schema);
 				}
 
+				if (EntryNode->MetaData.HasMetaData(FBlueprintMetadata::MD_FieldNotify) && Ret->ImplementsInterface(UNotifyFieldValueChanged::StaticClass()))
+				{
+					ensure(!Ret->FieldNotifies.Contains(FFieldNotificationId(NewFunction->GetFName())));
+					Ret->FieldNotifies.Add(FFieldNotificationId(NewFunction->GetFName()));
+				}
+
 				if (BP->bIsRegeneratingOnLoad)
 				{
 					// Ensure that the function's variable cache is up-to-date after property creation. Note that
