Fix bad test introduced with CL26807302
#rnx

diff --git a/Engine/Source/Editor/Kismet/Private/BlueprintCompilationManager.cpp b/Engine/Source/Editor/Kismet/Private/BlueprintCompilationManager.cpp
--- a/Engine/Source/Editor/Kismet/Private/BlueprintCompilationManager.cpp
+++ b/Engine/Source/Editor/Kismet/Private/BlueprintCompilationManager.cpp
@@ -2999,12 +2999,12 @@ UClass* FBlueprintCompilationManagerImpl::FastGenerateSkeletonClass(UBlueprint*
 					}
 
 					CompilerContext.SetCalculatedMetaDataAndFlags(NewFunction, EntryNode, Schema);
-				}
 
-				if (EntryNode->MetaData.HasMetaData(FBlueprintMetadata::MD_FieldNotify) && Ret->ImplementsInterface(UNotifyFieldValueChanged::StaticClass()))
-				{
-					ensure(!Ret->FieldNotifies.Contains(FFieldNotificationId(NewFunction->GetFName())));
-					Ret->FieldNotifies.Add(FFieldNotificationId(NewFunction->GetFName()));
+					if (EntryNode->MetaData.HasMetaData(FBlueprintMetadata::MD_FieldNotify) && Ret->ImplementsInterface(UNotifyFieldValueChanged::StaticClass()))
+					{
+						ensure(!Ret->FieldNotifies.Contains(FFieldNotificationId(NewFunction->GetFName())));
+						Ret->FieldNotifies.Add(FFieldNotificationId(NewFunction->GetFName()));
+					}
 				}
 
 				if (BP->bIsRegeneratingOnLoad)
