Make the ProcessClick behavior in the SCS Editor Viewport Client handle the ComponentVisManager in the same order as the Level Editor Viewport, making it easier to add custom tools and extensions to your editor. This will result in the right click behavior being evaluated in the same order as the level editor for more consistent options.

diff --git a/Engine/Source/Editor/Kismet/Private/SCSEditorViewportClient.cpp b/Engine/Source/Editor/Kismet/Private/SCSEditorViewportClient.cpp
--- a/Engine/Source/Editor/Kismet/Private/SCSEditorViewportClient.cpp
+++ b/Engine/Source/Editor/Kismet/Private/SCSEditorViewportClient.cpp
@@ -403,6 +403,10 @@ void FSCSEditorViewportClient::ProcessClick(class FSceneView& View, class HHitPr
 			Viewport->InvalidateHitProxy();
 			return;
 		}
+		else if(GUnrealEd->ComponentVisManager.HandleClick(this, HitProxy, Click))
+		{
+			// Component Vis Manager handled this click, no need to do anything
+		}
 		else if (HitProxy->IsA(HActor::StaticGetType()))
 		{
 			HActor* ActorProxy = (HActor*)HitProxy;
@@ -452,8 +456,6 @@ void FSCSEditorViewportClient::ProcessClick(class FSceneView& View, class HHitPr
 			return;
 		}
 	}
-	
-	GUnrealEd->ComponentVisManager.HandleClick(this, HitProxy, Click);
 }
 
 struct FTemplateComponentMoved
