Ensure the active component visualizer is cleared when the associated component is no longer selected

diff --git a/Engine/Source/Editor/Kismet/Private/BlueprintEditor.cpp b/Engine/Source/Editor/Kismet/Private/BlueprintEditor.cpp
--- a/Engine/Source/Editor/Kismet/Private/BlueprintEditor.cpp
+++ b/Engine/Source/Editor/Kismet/Private/BlueprintEditor.cpp
@@ -1195,6 +1195,28 @@ UObject* FBlueprintEditor::GetSubobjectEditorObjectContext() const
 
 void FBlueprintEditor::OnSelectionUpdated(const TArray<TSharedPtr<FSubobjectEditorTreeNode>>& SelectedNodes)
 {
+	// Check whether the active component visualizer is relevant for the selected components ...
+	if (const TSharedPtr<FComponentVisualizer> ComponentVisualizer = GUnrealEd->ComponentVisManager.GetActiveComponentVis())
+	{
+		bool bClearActiveVisualizer = true;
+		for (const FSubobjectEditorTreeNodePtrType SelectedNode : SelectedNodes)
+		{
+			const FSubobjectData* const Data = SelectedNode->GetDataSource();
+			const UActorComponent* const Component = Data ? Data->FindComponentInstanceInActor(GetPreviewActor()) : nullptr;
+			if (Component != nullptr && Component->IsRegistered() && Component == ComponentVisualizer->GetEditedComponent())
+			{
+				bClearActiveVisualizer = false;
+				break;
+			}
+		}
+
+		// If the relevant component for the active visualizer is no longer selected, clear the active visualizer.
+		if (bClearActiveVisualizer)
+		{
+			GUnrealEd->ComponentVisManager.ClearActiveComponentVis();
+		}
+	}
+
 	if (SubobjectViewport.IsValid())
 	{
 		SubobjectViewport->OnComponentSelectionChanged();
diff --git a/Engine/Source/Editor/UnrealEd/Private/ComponentVisualizerManager.cpp b/Engine/Source/Editor/UnrealEd/Private/ComponentVisualizerManager.cpp
--- a/Engine/Source/Editor/UnrealEd/Private/ComponentVisualizerManager.cpp
+++ b/Engine/Source/Editor/UnrealEd/Private/ComponentVisualizerManager.cpp
@@ -89,6 +89,11 @@ bool FComponentVisualizerManager::HandleProxyForComponentVis(FEditorViewportClie
 	return false;
 }
 
+TSharedPtr<FComponentVisualizer> FComponentVisualizerManager::GetActiveComponentVis()
+{
+	return EditedVisualizerPtr.Pin();
+}
+
 bool FComponentVisualizerManager::SetActiveComponentVis(FEditorViewportClient* InViewportClient, TSharedPtr<FComponentVisualizer>& InVisualizer)
 {
 	if (InViewportClient && InVisualizer.IsValid())
diff --git a/Engine/Source/Editor/UnrealEd/Public/ComponentVisualizerManager.h b/Engine/Source/Editor/UnrealEd/Public/ComponentVisualizerManager.h
--- a/Engine/Source/Editor/UnrealEd/Public/ComponentVisualizerManager.h
+++ b/Engine/Source/Editor/UnrealEd/Public/ComponentVisualizerManager.h
@@ -22,6 +22,9 @@ public:
 	/** Activate a component visualizer given a clicked proxy */
 	UNREALED_API bool HandleProxyForComponentVis(FEditorViewportClient* InViewportClient, HHitProxy *HitProxy, const FViewportClick &Click);
 
+	/** Get the active component visualizer. */
+	UNREALED_API TSharedPtr<FComponentVisualizer> GetActiveComponentVis();
+
 	/** Activate a component visualizer given the specific visualizer */
 	UNREALED_API bool SetActiveComponentVis(FEditorViewportClient* InViewportClient, TSharedPtr<FComponentVisualizer>& InVisualizer);
 
