CIS fix

diff --git a/Engine/Source/Editor/Kismet/Private/BlueprintEditor.cpp b/Engine/Source/Editor/Kismet/Private/BlueprintEditor.cpp
--- a/Engine/Source/Editor/Kismet/Private/BlueprintEditor.cpp
+++ b/Engine/Source/Editor/Kismet/Private/BlueprintEditor.cpp
@@ -1199,7 +1199,7 @@ void FBlueprintEditor::OnSelectionUpdated(const TArray<TSharedPtr<FSubobjectEdit
 	if (const TSharedPtr<FComponentVisualizer> ComponentVisualizer = GUnrealEd->ComponentVisManager.GetActiveComponentVis())
 	{
 		bool bClearActiveVisualizer = true;
-		for (const FSubobjectEditorTreeNodePtrType SelectedNode : SelectedNodes)
+		for (const FSubobjectEditorTreeNodePtrType& SelectedNode : SelectedNodes)
 		{
 			const FSubobjectData* const Data = SelectedNode->GetDataSource();
 			const UActorComponent* const Component = Data ? Data->FindComponentInstanceInActor(GetPreviewActor()) : nullptr;
