Fix crash when deleting blueprint functions in a widget blueprint.
Add a guard for FindFunction() to make sure it's not null.

diff --git a/Engine/Source/Editor/Kismet/Private/BlueprintDetailsCustomization.cpp b/Engine/Source/Editor/Kismet/Private/BlueprintDetailsCustomization.cpp
--- a/Engine/Source/Editor/Kismet/Private/BlueprintDetailsCustomization.cpp
+++ b/Engine/Source/Editor/Kismet/Private/BlueprintDetailsCustomization.cpp
@@ -4650,42 +4650,49 @@ bool FBlueprintGraphActionDetails::GetIsFieldNotfyEnabled() const
 ECheckBoxState FBlueprintGraphActionDetails::OnFieldNotifyCheckboxState() const
 {
 	UBlueprint* const BlueprintObj = GetBlueprintObj();
-	const FName FuncName = FindFunction()->GetFName();
 
-	if (BlueprintObj && GetIsFieldNotfyEnabled())
+	if (UFunction* Function = FindFunction())
 	{
-		if (!FuncName.IsNone())
-		{
-			return GetMetadataBlock()->HasMetaData(FBlueprintMetadata::MD_FieldNotify) ? ECheckBoxState::Checked : ECheckBoxState::Unchecked;
-		}
-		else if (BlueprintObj->GeneratedClass && BlueprintObj->GeneratedClass->ImplementsInterface(UNotifyFieldValueChanged::StaticClass()) && BlueprintObj->GeneratedClass->GetDefaultObject())
+		const FName FuncName = Function->GetFName();
+
+		if (BlueprintObj && GetIsFieldNotfyEnabled())
 		{
-			TScriptInterface<INotifyFieldValueChanged> DefaultObject = BlueprintObj->GeneratedClass->GetDefaultObject();
-			return DefaultObject->GetFieldNotificationDescriptor().GetField(BlueprintObj->GeneratedClass, FuncName).IsValid() ? ECheckBoxState::Checked : ECheckBoxState::Unchecked;
+			if (!FuncName.IsNone())
+			{
+				return GetMetadataBlock()->HasMetaData(FBlueprintMetadata::MD_FieldNotify) ? ECheckBoxState::Checked : ECheckBoxState::Unchecked;
+			}
+			else if (BlueprintObj->GeneratedClass && BlueprintObj->GeneratedClass->ImplementsInterface(UNotifyFieldValueChanged::StaticClass()) && BlueprintObj->GeneratedClass->GetDefaultObject())
+			{
+				TScriptInterface<INotifyFieldValueChanged> DefaultObject = BlueprintObj->GeneratedClass->GetDefaultObject();
+				return DefaultObject->GetFieldNotificationDescriptor().GetField(BlueprintObj->GeneratedClass, FuncName).IsValid() ? ECheckBoxState::Checked : ECheckBoxState::Unchecked;
+			}
 		}
-	}
 
-	FBlueprintEditorUtils::RemoveFieldNotifyFromAllMetadata(BlueprintObj, FuncName);
-	GetMetadataBlock()->RemoveMetaData(FBlueprintMetadata::MD_FieldNotify);
+		FBlueprintEditorUtils::RemoveFieldNotifyFromAllMetadata(BlueprintObj, FuncName);
+		GetMetadataBlock()->RemoveMetaData(FBlueprintMetadata::MD_FieldNotify);
+	}
 	return ECheckBoxState::Unchecked;
 }
 
 void FBlueprintGraphActionDetails::OnFieldNotifyChanged(ECheckBoxState InNewState)
 {
 	UBlueprint* const BlueprintObj = GetBlueprintObj();
-	FName FuncName = FindFunction()->GetFName();
-	const bool bFuncIsFieldNotify = InNewState == ECheckBoxState::Checked;
-
-	if (BlueprintObj)
+	if (UFunction* Function = FindFunction())
 	{
-		if (bFuncIsFieldNotify)
-		{
-			GetMetadataBlock()->SetMetaData(FBlueprintMetadata::MD_FieldNotify, FString());
-		}
-		else
+		const FName FuncName = Function->GetFName();
+		const bool bFuncIsFieldNotify = InNewState == ECheckBoxState::Checked;
+
+		if (BlueprintObj)
 		{
-			FBlueprintEditorUtils::RemoveFieldNotifyFromAllMetadata(BlueprintObj, FuncName);
-			GetMetadataBlock()->RemoveMetaData(FBlueprintMetadata::MD_FieldNotify);
+			if (bFuncIsFieldNotify)
+			{
+				GetMetadataBlock()->SetMetaData(FBlueprintMetadata::MD_FieldNotify, FString());
+			}
+			else
+			{
+				FBlueprintEditorUtils::RemoveFieldNotifyFromAllMetadata(BlueprintObj, FuncName);
+				GetMetadataBlock()->RemoveMetaData(FBlueprintMetadata::MD_FieldNotify);
+			}
 		}
 	}
 }
