Fix logic detecting blueprint subclasses in the diff panel, this re-enables logic to show differences in widget blueprints

diff --git a/Engine/Source/Editor/Kismet/Private/SBlueprintDiff.cpp b/Engine/Source/Editor/Kismet/Private/SBlueprintDiff.cpp
--- a/Engine/Source/Editor/Kismet/Private/SBlueprintDiff.cpp
+++ b/Engine/Source/Editor/Kismet/Private/SBlueprintDiff.cpp
@@ -953,11 +953,11 @@ void SBlueprintDiff::GenerateDifferencesList()
 	bool bIsSpecialized = false;
 	if (PanelOld.Blueprint)
 	{
-		bIsSpecialized = PanelOld.Blueprint->GetClass() == UBlueprint::StaticClass();
+		bIsSpecialized = (PanelOld.Blueprint->GetClass() != UBlueprint::StaticClass());
 	}
 	if (PanelNew.Blueprint)
 	{
-		bIsSpecialized |= PanelNew.Blueprint->GetClass() == UBlueprint::StaticClass();
+		bIsSpecialized |= (PanelNew.Blueprint->GetClass() != UBlueprint::StaticClass());
 	}
 	
 	if (bIsSpecialized)
