Further improve error message when a blueprint attempts to compile while other blueprints are compiling, QueuedRequests has already been emptied

diff --git a/Engine/Source/Editor/Kismet/Private/BlueprintCompilationManager.cpp b/Engine/Source/Editor/Kismet/Private/BlueprintCompilationManager.cpp
--- a/Engine/Source/Editor/Kismet/Private/BlueprintCompilationManager.cpp
+++ b/Engine/Source/Editor/Kismet/Private/BlueprintCompilationManager.cpp
@@ -241,9 +241,13 @@ void FBlueprintCompilationManagerImpl::QueueForCompilation(const FBPCompileReque
 		if(GCompilingBlueprint)
 		{
 			FString CurrentlyCompiling;
-			for (const FBPCompileRequestInternal& CompilerData : QueuedRequests)
+			for (const TPair<TObjectPtr<UClass>, TObjectPtr<UClass>>& CompilerData : ClassesToReinstance)
 			{
-				CurrentlyCompiling += CompilerData.UserData.BPToCompile->GetName() + TEXT(" ");
+				if (!CompilerData.Value)
+				{
+					continue;
+				}
+				CurrentlyCompiling += CompilerData.Value->GetName() + TEXT(" ");
 			}
 			ensureMsgf(false, 
 				TEXT("Attempting to enqueue %s for compile while compiling: %s"), 
