Fix composite graphs not being displayed when diffing

diff --git a/Engine/Source/Editor/Kismet/Private/SBlueprintDiff.cpp b/Engine/Source/Editor/Kismet/Private/SBlueprintDiff.cpp
--- a/Engine/Source/Editor/Kismet/Private/SBlueprintDiff.cpp
+++ b/Engine/Source/Editor/Kismet/Private/SBlueprintDiff.cpp
@@ -563,11 +563,7 @@ FGraphToDiff* SBlueprintDiff::FindGraphToDiffEntry(const FString& GraphPath)
 
 void SBlueprintDiff::FocusOnGraphRevisions( FGraphToDiff* Diff )
 {
-	UEdGraph* Graph = Diff->GetGraphOld() ? Diff->GetGraphOld() : Diff->GetGraphNew();
-
-	FString GraphPath = FGraphDiffControl::GetGraphPath(Graph);
-
-	HandleGraphChanged(GraphPath);
+	HandleGraphChanged(Diff);
 
 	ResetGraphEditors();
 }
@@ -809,48 +805,35 @@ FDiffPanel& SBlueprintDiff::GetDiffPanelForNode(UEdGraphNode& Node)
 	return Default;
 }
 
-void SBlueprintDiff::HandleGraphChanged( const FString& GraphPath )
+void SBlueprintDiff::HandleGraphChanged(FGraphToDiff* Diff)
 {
+	check(Diff);
 	SetCurrentMode(GraphMode);
-	
-	UEdGraph* GraphOld = nullptr;
-	UEdGraph* GraphNew = nullptr;
-	TSharedPtr<TArray<FDiffSingleResult>> DiffResults;
-	int32 RealDifferencesStartIndex = INDEX_NONE;
-	for (const TSharedPtr<FGraphToDiff>& GraphToDiff : Graphs)
-	{
-		UEdGraph* NewGraph = GraphToDiff->GetGraphNew();
-		UEdGraph* OldGraph = GraphToDiff->GetGraphOld();
-		const FString OtherGraphPath = NewGraph ? FGraphDiffControl::GetGraphPath(NewGraph) : FGraphDiffControl::GetGraphPath(OldGraph);
-		if (GraphPath.Equals(OtherGraphPath))
-		{
-			GraphNew = NewGraph;
-			GraphOld = OldGraph;
-			DiffResults = GraphToDiff->FoundDiffs;
-			RealDifferencesStartIndex = GraphToDiff->RealDifferencesStartIndex;
-			break;
-		}
-	}
-	
+
+	UEdGraph* GraphOld = Diff->GetGraphOld();
+	UEdGraph* GraphNew = Diff->GetGraphNew();
+	TSharedPtr<TArray<FDiffSingleResult>> DiffResults = Diff->FoundDiffs;
+	int32 RealDifferencesStartIndex = Diff->RealDifferencesStartIndex;
+
 	const TAttribute<int32> FocusedDiffResult = TAttribute<int32>::CreateLambda(
-        [this, RealDifferencesStartIndex]()
-        {
-        	int32 FocusedDiffResult = INDEX_NONE;
-        	if (RealDifferencesStartIndex != INDEX_NONE)
-        	{
+		[this, RealDifferencesStartIndex]()
+		{
+			int32 FocusedDiffResult = INDEX_NONE;
+			if (RealDifferencesStartIndex != INDEX_NONE)
+			{
 				FocusedDiffResult = DiffTreeView::CurrentDifference(DifferencesTreeView.ToSharedRef(), RealDifferences) - RealDifferencesStartIndex;
 			}
-        	
+
 			// find selected index in all the graphs, and subtract the index of the first entry in this graph
 			return FocusedDiffResult;
-        });
+		});
 
 	// only regenerate PanelOld if the old graph has changed
 	if (!PanelOld.GraphEditor.IsValid() || GraphOld != PanelOld.GraphEditor.Pin()->GetCurrentGraph())
 	{
 		PanelOld.GeneratePanel(GraphOld, DiffResults, FocusedDiffResult);
 	}
-	
+
 	// only regenerate PanelNew if the old graph has changed
 	if (!PanelNew.GraphEditor.IsValid() || GraphNew != PanelNew.GraphEditor.Pin()->GetCurrentGraph())
 	{
diff --git a/Engine/Source/Editor/Kismet/Public/SBlueprintDiff.h b/Engine/Source/Editor/Kismet/Public/SBlueprintDiff.h
--- a/Engine/Source/Editor/Kismet/Public/SBlueprintDiff.h
+++ b/Engine/Source/Editor/Kismet/Public/SBlueprintDiff.h
@@ -217,7 +217,7 @@ protected:
 	FDiffPanel& GetDiffPanelForNode(UEdGraphNode& Node);
 
 	/** Event handler that updates the graph view when user selects a new graph */
-	void HandleGraphChanged( const FString& GraphPath );
+	void HandleGraphChanged(FGraphToDiff* Diff);
 	
 	/** Function used to generate the list of differences and the widgets needed to calculate that list */
 	void GenerateDifferencesList();
