Change to not replace not only CDOs references but any of there archetypes

diff --git a/Engine/Source/Editor/Kismet/Private/BlueprintCompilationManager.cpp b/Engine/Source/Editor/Kismet/Private/BlueprintCompilationManager.cpp
--- a/Engine/Source/Editor/Kismet/Private/BlueprintCompilationManager.cpp
+++ b/Engine/Source/Editor/Kismet/Private/BlueprintCompilationManager.cpp
@@ -127,10 +127,10 @@ struct FBlueprintCompilationManagerImpl : public FGCObject
 
 	void QueueForCompilation(const FBPCompileRequestInternal& CompileJob);
 	void CompileSynchronouslyImpl(const FBPCompileRequestInternal& Request);
-	void FlushCompilationQueueImpl(bool bSuppressBroadcastCompiled, TArray<UBlueprint*>* BlueprintsCompiled, TArray<UBlueprint*>* BlueprintsCompiledOrSkeletonCompiled, FUObjectSerializeContext* InLoadContext);
+	void FlushCompilationQueueImpl(bool bSuppressBroadcastCompiled, TArray<UBlueprint*>* BlueprintsCompiled, TArray<UBlueprint*>* BlueprintsCompiledOrSkeletonCompiled, FUObjectSerializeContext* InLoadContext, TMap<UClass*, TMap<UObject*, UObject*>>* OldToNewTemplates = nullptr);
 	void FixupDelegateProperties(const TArray<FCompilerData>& CurrentlyCompilingBPs);
 	void ProcessExtensions(const TArray<FCompilerData>& InCurrentlyCompilingBPs);
-	void FlushReinstancingQueueImpl(bool bFindAndReplaceCDOReferences = false);
+	void FlushReinstancingQueueImpl(bool bFindAndReplaceCDOReferences = false, TMap<UClass*, TMap<UObject*, UObject*>>* OldToNewTemplates = nullptr);
 	bool HasBlueprintsToCompile() const;
 	bool IsGeneratedClassLayoutReady() const;
 	void GetDefaultValue(const UClass* ForClass, const FProperty* Property, FString& OutDefaultValueAsString) const;
@@ -138,7 +138,7 @@ struct FBlueprintCompilationManagerImpl : public FGCObject
 
 	static void ReparentHierarchies(const TMap<UClass*, UClass*>& OldClassToNewClass, EReparentClassOptions Options);
 	static void BuildDSOMap(UObject* OldObject, UObject* NewObject, TMap<UObject*, UObject*>& OutOldToNewDSO);
-	static void ReinstanceBatch(TArray<FReinstancingJob>& Reinstancers, TMap< UClass*, UClass* >& InOutOldToNewClassMap, FUObjectSerializeContext* InLoadContext);
+	static void ReinstanceBatch(TArray<FReinstancingJob>& Reinstancers, TMap< UClass*, UClass* >& InOutOldToNewClassMap, FUObjectSerializeContext* InLoadContext, TMap<UClass*, TMap<UObject*, UObject*>>* OldToNewTemplates = nullptr);
 	static UClass* FastGenerateSkeletonClass(UBlueprint* BP, FKismetCompilerContext& CompilerContext, bool bIsSkeletonOnly, TArray<FSkeletonFixupData>& OutSkeletonFixupData);
 	static bool IsQueuedForCompilation(UBlueprint* BP);
 	static void ConformToParentAndInterfaces(UBlueprint* BP);
@@ -302,10 +302,11 @@ void FBlueprintCompilationManagerImpl::CompileSynchronouslyImpl(const FBPCompile
 	// We suppress normal compilation broadcasts because the old code path 
 	// did this after GC and we want to match the old behavior:
 	const bool bSuppressBroadcastCompiled = true;
+	TMap<UClass*, TMap<UObject*, UObject*>> OldToNewTemplates;
 	TArray<UBlueprint*> CompiledBlueprints;
 	TArray<UBlueprint*> SkeletonCompiledBlueprints;
-	FlushCompilationQueueImpl(bSuppressBroadcastCompiled, &CompiledBlueprints, &SkeletonCompiledBlueprints, nullptr);
-	FlushReinstancingQueueImpl(bFindAndReplaceCDOReferences);
+	FlushCompilationQueueImpl(bSuppressBroadcastCompiled, &CompiledBlueprints, &SkeletonCompiledBlueprints, nullptr, bFindAndReplaceCDOReferences ? &OldToNewTemplates : nullptr);
+	FlushReinstancingQueueImpl(bFindAndReplaceCDOReferences, bFindAndReplaceCDOReferences ? &OldToNewTemplates : nullptr);
 	
 	if( Request.UserData.ClientResultsLog && Request.UserData.ClientResultsLog->bLogDetailedResults)
 	{
@@ -601,7 +602,7 @@ namespace UE::Kismet::BlueprintCompilationManager::Private
 	}
 }
 
-void FBlueprintCompilationManagerImpl::FlushCompilationQueueImpl(bool bSuppressBroadcastCompiled, TArray<UBlueprint*>* BlueprintsCompiled, TArray<UBlueprint*>* BlueprintsCompiledOrSkeletonCompiled, FUObjectSerializeContext* InLoadContext)
+void FBlueprintCompilationManagerImpl::FlushCompilationQueueImpl(bool bSuppressBroadcastCompiled, TArray<UBlueprint*>* BlueprintsCompiled, TArray<UBlueprint*>* BlueprintsCompiledOrSkeletonCompiled, FUObjectSerializeContext* InLoadContext, TMap<UClass*, TMap<UObject*, UObject*>>* OldToNewTemplates /* = nullptr*/)
 {
 	TRACE_CPUPROFILER_EVENT_SCOPE(FBlueprintCompilationManager::FlushCompilationQueue);
 	DECLARE_SCOPE_HIERARCHICAL_COUNTER_FUNC()
@@ -1593,7 +1594,7 @@ void FBlueprintCompilationManagerImpl::FlushCompilationQueueImpl(bool bSuppressB
 			}
 
 			FScopedDurationTimer ReinstTimer(GTimeReinstancing);
-			ReinstanceBatch(Reinstancers, MutableView(ClassesToReinstance), InLoadContext);
+			ReinstanceBatch(Reinstancers, MutableView(ClassesToReinstance), InLoadContext, OldToNewTemplates);
 
 			// We purposefully do not remove the OldCDOs yet, need to keep them in memory past first GC
 		}
@@ -1926,7 +1927,7 @@ void FBlueprintCompilationManagerImpl::ProcessExtensions(const TArray<FCompilerD
 	}
 }
 
-void FBlueprintCompilationManagerImpl::FlushReinstancingQueueImpl(bool bFindAndReplaceCDOReferences)
+void FBlueprintCompilationManagerImpl::FlushReinstancingQueueImpl(bool bFindAndReplaceCDOReferences, TMap<UClass*, TMap<UObject*, UObject*>>* OldToNewTemplates /* = nullptr*/)
 {
 	DECLARE_SCOPE_HIERARCHICAL_COUNTER_FUNC()
 
@@ -1957,6 +1958,7 @@ void FBlueprintCompilationManagerImpl::FlushReinstancingQueueImpl(bool bFindAndR
 		FReplaceInstancesOfClassParameters Options;
 		Options.bArchetypesAreUpToDate = true;
 		Options.bReplaceReferencesToOldCDOs = bFindAndReplaceCDOReferences;
+		Options.OldToNewTemplates = OldToNewTemplates;
 		FBlueprintCompileReinstancer::BatchReplaceInstancesOfClass(ClassesToReinstanceOwned, Options);
 
 		// Special case when we run on ALT, we want to cleanup all classes flagged for reinstanciation right away.
@@ -2238,7 +2240,7 @@ void FBlueprintCompilationManagerImpl::BuildDSOMap(UObject* OldObject, UObject*
 	}
 }
 
-void FBlueprintCompilationManagerImpl::ReinstanceBatch(TArray<FReinstancingJob>& Reinstancers, TMap< UClass*, UClass* >& InOutOldToNewClassMap, FUObjectSerializeContext* InLoadContext)
+void FBlueprintCompilationManagerImpl::ReinstanceBatch(TArray<FReinstancingJob>& Reinstancers, TMap< UClass*, UClass* >& InOutOldToNewClassMap, FUObjectSerializeContext* InLoadContext, TMap<UClass*, TMap<UObject*, UObject*>>* OldToNewTemplates /* = nullptr*/)
 {
 	TGuardValue<bool> ReinstancingGuard(GIsReinstancing, true);
 
@@ -2443,6 +2445,11 @@ void FBlueprintCompilationManagerImpl::ReinstanceBatch(TArray<FReinstancingJob>&
 					FBlueprintCompileReinstancer::CopyPropertiesForUnrelatedObjects(Pair.Key, Pair.Value, /*bClearExternalReferences*/true, bUseDeltaSerialization, /*bOnlyHandleDirectSubObjects*/true, &OldToNewInstanceMap);
 				}
 
+				if (OldToNewTemplates && !OldToNewInstanceMap.IsEmpty())
+				{
+					OldToNewTemplates->FindOrAdd(ReinstancingJob.OldToNew.Key).Append(OldToNewInstanceMap);
+				}
+
 				if (ReinstancingJob.Compiler.IsValid())
 				{
 					ReinstancingJob.Compiler->PropagateValuesToCDO(NewCDO, OldCDO);
@@ -2599,38 +2606,6 @@ void FBlueprintCompilationManagerImpl::ReinstanceBatch(TArray<FReinstancingJob>&
 				const EObjectFlags FlagMask = RF_Public | RF_ArchetypeObject | RF_Transactional | RF_Transient | RF_TextExportTransient | RF_InheritableComponentTemplate | RF_Standalone; //TODO: what about RF_RootSet?
 				UObject* NewArchetype = NewObject<UObject>(OriginalOuter, NewClass, OriginalName, OriginalFlags & FlagMask);
 
-				// grab the old archetype's subobjects:
-				{
-					TArray<UObject*> OldSubobjects;
-					GetObjectsWithOuter( Archetype, OldSubobjects, false );
-
-					for(UObject* Subobject : OldSubobjects )
-					{
-						if(Subobject->HasAnyFlags(RF_DefaultSubObject))
-						{
-							// CPFUO handles DSOs:
-							continue;
-						}
-
-						// Even if the object created the DSO itself, we want to carry over the existing one from the old archetype, so rename the newly constructed one out of the way
-						if (UObject* ExistingObject = static_cast<UObject*>(FindObjectWithOuter(NewArchetype, nullptr, Subobject->GetFName())))
-						{
-							UClass* ExistingObjectClass = ExistingObject->GetClass();
-							UObject* TransientOuterForRename = GetTransientOuterForRename(ExistingObjectClass);
-							ExistingObject->Rename(*MakeUniqueObjectName(TransientOuterForRename, ExistingObjectClass).ToString(), TransientOuterForRename, REN_DoNotDirty | REN_DontCreateRedirectors | REN_ForceNoResetLoaders);
-						}
-
-						// Non DSO subobject - just reuse the subobject:
-						Subobject->Rename(
-							nullptr, 
-							// destination:
-							NewArchetype, 
-							// Rename options:
-							REN_DoNotDirty | REN_DontCreateRedirectors | REN_ForceNoResetLoaders 
-						);
-					}
-				}
-
 				OldArchetypeToNewArchetype.Add(Archetype, NewArchetype);
 
 				// also map old *default* subobjects to new default subobjects:
@@ -2639,9 +2614,6 @@ void FBlueprintCompilationManagerImpl::ReinstanceBatch(TArray<FReinstancingJob>&
 				ArchetypeReferencers.Add(NewArchetype);
 
 				FLinkerLoad::PRIVATE_PatchNewObjectIntoExport(Archetype, NewArchetype);
-
-				Archetype->RemoveFromRoot();
-				Archetype->MarkAsGarbage();
 			}
 		}
 	}
@@ -2662,11 +2634,27 @@ void FBlueprintCompilationManagerImpl::ReinstanceBatch(TArray<FReinstancingJob>&
 				if(NewInstance && *NewInstance)
 				{
 					// The new object hierarchy has been created, all of the old instances are in the transient package and new
-					// ones have taken their place. Referenc members will mostly be pointing at *old* instances, and will get fixed
+					// ones have taken their place. Reference members will mostly be pointing at *old* instances, and will get fixed
 					// up below:			
 					const bool bUseDeltaSerialization = ReinstancingJob.Reinstancer.IsValid() ? ReinstancingJob.Reinstancer->bUseDeltaSerializationToCopyProperties : false;
 
-					FBlueprintCompileReinstancer::CopyPropertiesForUnrelatedObjects(OldInstance, *NewInstance, false, bUseDeltaSerialization);
+					TMap<UObject*, UObject*> CreatedInstanceMap;
+					FBlueprintCompileReinstancer::PreCreateSubObjectsForReinstantiation(InOutOldToNewClassMap, OldInstance, *NewInstance, CreatedInstanceMap);
+
+					// We only need to copy properties of the pre-created instances, the rest of the default sub object is done inside the UEditorEngine::CopyPropertiesForUnrelatedObjects
+					TMap<UObject*, UObject*> OldToNewInstanceMap(CreatedInstanceMap);
+					for (const auto& Pair : CreatedInstanceMap)
+					{
+						FBlueprintCompileReinstancer::CopyPropertiesForUnrelatedObjects(OldInstance, *NewInstance, /*bClearExternalReferences*/true, bUseDeltaSerialization, /*bOnlyHandleDirectSubObjects*/true, &OldToNewInstanceMap);
+					}
+
+					if (OldToNewTemplates && !OldToNewInstanceMap.IsEmpty())
+					{
+						OldToNewTemplates->FindOrAdd(OldClass).Append(OldToNewInstanceMap);
+					}
+
+					OldInstance->RemoveFromRoot();
+					OldInstance->MarkAsGarbage();
 				}
 			}
 		}
diff --git a/Engine/Source/Editor/UnrealEd/Private/Kismet2/KismetReinstanceUtilities.cpp b/Engine/Source/Editor/UnrealEd/Private/Kismet2/KismetReinstanceUtilities.cpp
--- a/Engine/Source/Editor/UnrealEd/Private/Kismet2/KismetReinstanceUtilities.cpp
+++ b/Engine/Source/Editor/UnrealEd/Private/Kismet2/KismetReinstanceUtilities.cpp
@@ -121,17 +121,29 @@ struct FReplaceReferenceHelper
 			*NamesOfClasses, *NamesOfObjects);
 	}
 
-	static void IncludeCDO(UClass* OldClass, UClass* NewClass, TMap<UObject*, UObject*> &OldToNewInstanceMap, TArray<UObject*> &SourceObjects, UObject* OriginalCDO)
+	static void IncludeCDO(UClass* OldClass, UClass* NewClass, TMap<UObject*, UObject*> &OldToNewInstanceMap, TArray<UObject*> &SourceObjects, UObject* OriginalCDO, TMap<UClass*, TMap<UObject*, UObject*>>* OldToNewTemplates = nullptr)
 	{
 		UObject* OldCDO = OldClass->GetDefaultObject();
 		UObject* NewCDO = NewClass->GetDefaultObject();
 
-		// Add the old->new CDO mapping into the fixup map
-		OldToNewInstanceMap.Add(OldCDO, NewCDO);
-		// Add in the old CDO to this pass, so CDO references are fixed up
-		SourceObjects.Add(OldCDO);
+		if (const TMap<UObject*, UObject*>* OldToNewTemplateMapping = OldToNewTemplates ? OldToNewTemplates->Find(OldClass) : nullptr)
+		{
+			OldToNewInstanceMap.Append(*OldToNewTemplateMapping);
+
+			TArray<UObject*> SourceTemplateObjects;
+			OldToNewTemplateMapping->GenerateKeyArray(SourceTemplateObjects);
+			SourceObjects.Append(SourceTemplateObjects);
+		}
+		else
+		{
+			// Add the old->new CDO mapping into the fixup map
+			OldToNewInstanceMap.Add(OldCDO, NewCDO);
+			// Add in the old CDO to this pass, so CDO references are fixed up
+			SourceObjects.Add(OldCDO);
+		}
+
 
-		if (OriginalCDO)
+		if (OriginalCDO && OriginalCDO != OldCDO)
 		{
 			OldToNewInstanceMap.Add(OriginalCDO, NewCDO);
 			SourceObjects.Add(OriginalCDO);
@@ -2610,7 +2622,7 @@ void FBlueprintCompileReinstancer::ReplaceInstancesOfClass_Inner(const TMap<UCla
 			check(OldClass && NewClass);
 			check(OldClass != NewClass || IsReloadActive());
 
-			FReplaceReferenceHelper::IncludeCDO(OldClass, NewClass, OldToNewInstanceMap, SourceObjects, InOriginalCDO);
+			FReplaceReferenceHelper::IncludeCDO(OldClass, NewClass, OldToNewInstanceMap, SourceObjects, InOriginalCDO, Params.OldToNewTemplates);
 
 			if (bClassObjectReplaced)
 			{
diff --git a/Engine/Source/Editor/UnrealEd/Public/Kismet2/KismetReinstanceUtilities.h b/Engine/Source/Editor/UnrealEd/Public/Kismet2/KismetReinstanceUtilities.h
--- a/Engine/Source/Editor/UnrealEd/Public/Kismet2/KismetReinstanceUtilities.h
+++ b/Engine/Source/Editor/UnrealEd/Public/Kismet2/KismetReinstanceUtilities.h
@@ -55,7 +55,10 @@ struct FReplaceInstancesOfClassParameters
 	UClass* NewClass = nullptr;
 
 	/** OriginalCDO, use if OldClass->ClassDefaultObject has been overwritten (non-batch only, legacy) */
-	UObject* OriginalCDO = nullptr; 
+	UObject* OriginalCDO = nullptr;
+
+	/* Mapping of all the replaced CDOs and archetypes*/
+	TMap<UClass*, TMap<UObject*, UObject*>>* OldToNewTemplates = nullptr;
 
 	/** Set of objects that should not have their references updated if they refer to instances that are replaced */
 	TSet<UObject*>* ObjectsThatShouldUseOldStuff = nullptr;
@@ -285,6 +288,7 @@ protected:
 	/** Determine whether reinstancing actors should preserve the root component of the new actor */
 	virtual bool ShouldPreserveRootComponentOfReinstancedActor() const { return true; }
 
+public:
 	/**
 	* Attempts to copy as many properties as possible from the old object to the new. 
 	* Use during BP compilation to copy properties from the old CDO to the new one.
@@ -296,7 +300,6 @@ protected:
 	*/
 	static UNREALED_API void CopyPropertiesForUnrelatedObjects(UObject* OldObject, UObject* NewObject, bool bClearExternalReferences, bool bForceDeltaSerialization = false, bool bOnlyHandleDirectSubObjects = false, TMap<UObject*, UObject*>* OldToNewInstanceMap =nullptr);
 
-public:
 	/**
 	 * This method will pre-create all non-default sub object needed for a re-instantiation, what is left is to CopyPropertiesForUnrelatedObjects on the created instances map to finish the re-instancing
 	 * If the re-instancing is done in a big batch and part to the sub object might already be re-instantiated, you will need to provide those via the OldToNewInstanceMap
