[CIS FIX] Remapped "Add Selected Actor Reference(s)" to ctrl-R since there was a collision between it and "add reroute node" and it's the less versatile of the two

diff --git a/Engine/Source/Editor/Kismet/Private/BlueprintEditorCommands.cpp b/Engine/Source/Editor/Kismet/Private/BlueprintEditorCommands.cpp
--- a/Engine/Source/Editor/Kismet/Private/BlueprintEditorCommands.cpp
+++ b/Engine/Source/Editor/Kismet/Private/BlueprintEditorCommands.cpp
@@ -417,7 +417,7 @@ void FBlueprintSpawnNodeCommands::RegisterCommands()
 	}
 
 	TSharedPtr<FNodeSpawnInfo> AddActorRefAction = MakeShareable(new FActorRefSpawnInfo);
-	UI_COMMAND(AddActorRefAction->CommandInfo, "Add Selected Actor Reference(s)", "Spawns node(s) which reference the currently selected actor(s) in the level.", EUserInterfaceActionType::Button, FInputChord(EKeys::R));
+	UI_COMMAND(AddActorRefAction->CommandInfo, "Add Selected Actor Reference(s)", "Spawns node(s) which reference the currently selected actor(s) in the level.", EUserInterfaceActionType::Button, FInputChord(EModifierKey::Control, EKeys::R));
 	NodeCommands.Add(AddActorRefAction);
 }
 
