Use ESlowTaskVisibility::Important so that Blueprint function and macro library loading is actually displayed usefully to the user

diff --git a/Engine/Source/Editor/Kismet/Private/BlueprintEditor.cpp b/Engine/Source/Editor/Kismet/Private/BlueprintEditor.cpp
--- a/Engine/Source/Editor/Kismet/Private/BlueprintEditor.cpp
+++ b/Engine/Source/Editor/Kismet/Private/BlueprintEditor.cpp
@@ -2080,13 +2080,25 @@ void FBlueprintEditor::LoadLibrariesFromAssetRegistry()
 			TArray<FAssetData> AssetData;
 			AssetRegistryModule.Get().GetAssetsByClass(UBlueprint::StaticClass()->GetClassPathName(), AssetData);
 
-			GWarn->BeginSlowTask(LOCTEXT("LoadingBlueprintAssetData", "Loading Blueprint Asset Data"), true);
-
 			const FName BPTypeName(GET_MEMBER_NAME_STRING_CHECKED(UBlueprint, BlueprintType));
-			const FName BPNamespaceName(GET_MEMBER_NAME_STRING_CHECKED(UBlueprint, BlueprintNamespace));
-			const FString BPMacroTypeStr(TEXT("BPTYPE_MacroLibrary"));
-			const FString BPFunctionTypeStr(TEXT("BPTYPE_FunctionLibrary"));
+			TArray<const FAssetData*> RelevantAssets;
+			const TCHAR* BPMacroTypeStr = TEXT("BPTYPE_MacroLibrary");
+			const TCHAR* BPFunctionTypeStr = TEXT("BPTYPE_FunctionLibrary");
+			for (const FAssetData& AssetEntry : AssetData)
+			{
+				const FString AssetBPType = AssetEntry.GetTagValueRef<FString>(BPTypeName);
+
+				// Only check for Blueprint Macros & Functions in the asset data for loading
+				if ((AssetBPType == BPMacroTypeStr) || (AssetBPType == BPFunctionTypeStr))
+				{
+					RelevantAssets.Add(&AssetEntry);
+				}
+			}
 
+			FScopedSlowTask LoadingMacrosAndFunctions(RelevantAssets.Num(), LOCTEXT("LoadingBlueprintAssetData", "Loading Blueprint Asset Data"));
+			LoadingMacrosAndFunctions.Visibility = ESlowTaskVisibility::Important; // this function can be very slow, users will benefit from our messages
+			const FName BPNamespaceName(GET_MEMBER_NAME_STRING_CHECKED(UBlueprint, BlueprintNamespace));
+			
 			struct FExpensiveObjectRecord
 			{
 				FExpensiveObjectRecord() : Seconds(0.0) {}
@@ -2101,54 +2113,45 @@ void FBlueprintEditor::LoadLibrariesFromAssetRegistry()
 			int32 NumLibariesLoaded = 0;
 
 			const double StartTimeAll = FPlatformTime::Seconds();
-			int32 AssetIndexBeingProcessed = 0;
-			for (const FAssetData& AssetEntry : AssetData)
+			for (const FAssetData* AssetEntryPtr : RelevantAssets)
 			{
-				const FString AssetBPType = AssetEntry.GetTagValueRef<FString>(BPTypeName);
-
-				// Only check for Blueprint Macros & Functions in the asset data for loading
-				if ((AssetBPType == BPMacroTypeStr) || (AssetBPType == BPFunctionTypeStr))
-				{
-					const FString BlueprintPath = AssetEntry.ToSoftObjectPath().ToString();
+				const FAssetData& AssetEntry = *AssetEntryPtr;
+				const FString BlueprintPath = AssetEntry.GetSoftObjectPath().ToString();
 
-					// See if this passes the namespace check
-					bool bAllowLoadBP = !ImportedNamespaceHelper.IsValid() || ImportedNamespaceHelper->IsImportedAsset(AssetEntry);
+				// See if this passes the namespace check
+				bool bAllowLoadBP = !ImportedNamespaceHelper.IsValid() || ImportedNamespaceHelper->IsImportedAsset(AssetEntry);
 					
-					// For blueprints inside developers folder, only allow the ones inside current user's developers folder.
-					if (bAllowLoadBP)
+				// For blueprints inside developers folder, only allow the ones inside current user's developers folder.
+				if (bAllowLoadBP)
+				{
+					if (BlueprintPath.StartsWith(DeveloperPath) && 
+						!BlueprintPath.StartsWith(UserDeveloperPath))
 					{
-						if (BlueprintPath.StartsWith(DeveloperPath))
-						{
-							if (!BlueprintPath.StartsWith(UserDeveloperPath))
-							{
-								bAllowLoadBP = false;
-							}
-						}
+						bAllowLoadBP = false;
 					}
+				}
 
-					if (bAllowLoadBP)
-					{
-						GWarn->StatusUpdate(AssetIndexBeingProcessed, AssetData.Num(), FText::FromName(AssetEntry.AssetName));
+				if (bAllowLoadBP)
+				{
+					LoadingMacrosAndFunctions.EnterProgressFrame(1.f, 
+						FText::Format(LOCTEXT("LoadingFuncMacroLib", "Loading Function or Macro library: {0}"), FText::FromName(AssetEntry.AssetName)));
 
-						++NumLibariesLoaded;
-						const double StartTime = FPlatformTime::Seconds();
+					++NumLibariesLoaded;
+					const double StartTime = FPlatformTime::Seconds();
 
-						// Load the blueprint
-						UBlueprint* BlueprintLibPtr = LoadObject<UBlueprint>(nullptr, *BlueprintPath, nullptr, 0, nullptr);
-						if (BlueprintLibPtr)
-						{
-							StandardLibraries.AddUnique(BlueprintLibPtr);
-						}
+					// Load the blueprint
+					UBlueprint* BlueprintLibPtr = LoadObject<UBlueprint>(nullptr, *BlueprintPath, nullptr, 0, nullptr);
+					if (BlueprintLibPtr)
+					{
+						StandardLibraries.AddUnique(BlueprintLibPtr);
+					}
 
-						const double ElapsedTime = FPlatformTime::Seconds() - StartTime;
-						if (ElapsedTime > MinSecondsToReportExpensiveObject)
-						{
-							ExpensiveObjects.Add(FExpensiveObjectRecord(ElapsedTime, AssetEntry.PackageName));
-						}
+					const double ElapsedTime = FPlatformTime::Seconds() - StartTime;
+					if (ElapsedTime > MinSecondsToReportExpensiveObject)
+					{
+						ExpensiveObjects.Add(FExpensiveObjectRecord(ElapsedTime, AssetEntry.PackageName));
 					}
 				}
-
-				++AssetIndexBeingProcessed;
 			}
 
 			if (ExpensiveObjects.Num() > 0)
@@ -2164,8 +2167,6 @@ void FBlueprintEditor::LoadLibrariesFromAssetRegistry()
 					UE_LOG(LogBlueprintEditor, Log, TEXT("Perf: %.1f seconds loading: %s"), ExpensiveObjectRecord.Seconds, *ExpensiveObjectRecord.Path.ToString());
 				}
 			}
-
-			GWarn->EndSlowTask();
 		}
 	}
 }
