Add support for units for blueprint variables (ints and floats)

diff --git a/Engine/Source/Editor/Kismet/Private/BlueprintDetailsCustomization.cpp b/Engine/Source/Editor/Kismet/Private/BlueprintDetailsCustomization.cpp
--- a/Engine/Source/Editor/Kismet/Private/BlueprintDetailsCustomization.cpp
+++ b/Engine/Source/Editor/Kismet/Private/BlueprintDetailsCustomization.cpp
@@ -84,6 +84,7 @@
 #include "Layout/WidgetPath.h"
 #include "Logging/LogCategory.h"
 #include "Logging/LogMacros.h"
+#include "Math/UnitConversion.h"
 #include "Math/Vector2D.h"
 #include "Math/Vector4.h"
 #include "Misc/Attribute.h"
@@ -485,7 +486,7 @@ void FBlueprintVarActionDetails::CustomizeDetails( IDetailLayoutBuilder& DetailL
 	TSharedPtr<SToolTip> EditableTooltip = IDocumentation::Get()->CreateToolTip(LOCTEXT("VarEditableTooltip", "Whether this variable is publicly editable on instances of this Blueprint."), NULL, DocLink, TEXT("Editable"));
 
 	Category.AddCustomRow( LOCTEXT("IsVariableEditableLabel", "Instance Editable") )
-	.Visibility(TAttribute<EVisibility>(this, &FBlueprintVarActionDetails::ShowEditableCheckboxVisibilty))
+	.Visibility(TAttribute<EVisibility>(this, &FBlueprintVarActionDetails::ShowEditableCheckboxVisibility))
 	.NameContent()
 	[
 		SNew(STextBlock)
@@ -505,7 +506,7 @@ void FBlueprintVarActionDetails::CustomizeDetails( IDetailLayoutBuilder& DetailL
 	TSharedPtr<SToolTip> ReadOnlyTooltip = IDocumentation::Get()->CreateToolTip(LOCTEXT("VarReadOnlyTooltip", "Whether this variable can be set by Blueprint nodes or if it is read-only."), NULL, DocLink, TEXT("ReadOnly"));
 
 	Category.AddCustomRow(LOCTEXT("IsVariableReadOnlyLabel", "Blueprint Read Only"))
-	.Visibility(TAttribute<EVisibility>(this, &FBlueprintVarActionDetails::ShowReadOnlyCheckboxVisibilty))
+	.Visibility(TAttribute<EVisibility>(this, &FBlueprintVarActionDetails::ShowReadOnlyCheckboxVisibility))
 	.NameContent()
 	[
 		SNew(STextBlock)
@@ -554,7 +555,7 @@ void FBlueprintVarActionDetails::CustomizeDetails( IDetailLayoutBuilder& DetailL
 					SNew(UE::FieldNotification::SFieldNotificationCheckList)
 					.FieldName(CachedVariableName)
 					.BlueprintPtr(BlueprintPtr)
-					.Visibility(this, &FBlueprintVarActionDetails::GetFieldNotifyCheckboxListVisiblity)
+					.Visibility(this, &FBlueprintVarActionDetails::GetFieldNotifyCheckboxListVisibility)
 				]
 
 			];
@@ -834,6 +835,34 @@ void FBlueprintVarActionDetails::CustomizeDetails( IDetailLayoutBuilder& DetailL
 			.Font(DetailFontInfo)
 		]
 	];
+	
+	TSharedPtr<SToolTip> UnitsTooltip = IDocumentation::Get()->CreateToolTip(LOCTEXT("VarUnitsTooltip", "Units of this variable."), NULL, DocLink, TEXT("Units"));
+
+	UnitsOptions.Empty();
+	UnitsOptions.Add(MakeShareable(new FString("None")));
+	for (const TCHAR* UnitsName : FUnitConversion::GetSupportedUnits())
+	{
+		UnitsOptions.Add(MakeShareable(new FString(UnitsName)));
+	}
+	
+	Category.AddCustomRow(LOCTEXT("VariableUnitsLabel", "Units"))
+	.Visibility(TAttribute<EVisibility>(this, &FBlueprintVarActionDetails::GetVariableUnitsVisibility))
+	.NameContent()
+	[
+		SNew(STextBlock)
+		.Text(LOCTEXT("VariableUnitsLabel", "Units"))
+		.ToolTip(UnitsTooltip)
+		.Font(IDetailLayoutBuilder::GetDetailFont())
+	]
+	.ValueContent()
+	[
+		SNew(STextComboBox)
+			.OptionsSource( &UnitsOptions )
+			.InitiallySelectedItem(GetVariableUnits())
+			.OnSelectionChanged( this, &FBlueprintVarActionDetails::OnVariableUnitsChanged )
+			.ToolTip(UnitsTooltip)
+			.Font( DetailFontInfo )
+	];
 
 	TSharedPtr<SToolTip> BitmaskTooltip = IDocumentation::Get()->CreateToolTip(LOCTEXT("VarBitmaskTooltip", "Whether or not to treat this variable as a bitmask."), nullptr, DocLink, TEXT("Bitmask"));
 
@@ -1851,7 +1880,7 @@ void FBlueprintVarActionDetails::OnCategorySelectionChanged( TSharedPtr<FText> P
 	}
 }
 
-EVisibility FBlueprintVarActionDetails::ShowEditableCheckboxVisibilty() const
+EVisibility FBlueprintVarActionDetails::ShowEditableCheckboxVisibility() const
 {
 	FProperty* VariableProperty = CachedVariableProperty.Get();
 	if (VariableProperty && GetPropertyOwnerBlueprint())
@@ -1885,7 +1914,7 @@ void FBlueprintVarActionDetails::OnEditableChanged(ECheckBoxState InNewState)
 	FBlueprintEditorUtils::SetBlueprintOnlyEditableFlag(BlueprintObj, VarName, !bVariableIsExposed);
 }
 
-EVisibility FBlueprintVarActionDetails::ShowReadOnlyCheckboxVisibilty() const
+EVisibility FBlueprintVarActionDetails::ShowReadOnlyCheckboxVisibility() const
 {
 	FProperty* VariableProperty = CachedVariableProperty.Get();
 	if (VariableProperty && GetPropertyOwnerBlueprint())
@@ -1919,6 +1948,63 @@ void FBlueprintVarActionDetails::OnReadyOnlyChanged(ECheckBoxState InNewState)
 	FBlueprintEditorUtils::SetBlueprintPropertyReadOnlyFlag(BlueprintObj, VarName, bVariableIsReadOnly);
 }
 
+EVisibility FBlueprintVarActionDetails::GetVariableUnitsVisibility() const
+{
+	FProperty* VariableProperty = CachedVariableProperty.Get();
+	if (VariableProperty)
+	{
+		const bool bIsInteger = VariableProperty->IsA(FIntProperty::StaticClass()) || VariableProperty->IsA(FInt64Property::StaticClass());
+		const bool bIsReal = VariableProperty->IsA(FFloatProperty::StaticClass()) || VariableProperty->IsA(FDoubleProperty::StaticClass());
+
+		if (IsABlueprintVariable(VariableProperty) && !IsALocalVariable(VariableProperty) && (bIsInteger || bIsReal))
+		{
+			return EVisibility::Visible;
+		}
+	}
+	return EVisibility::Hidden;
+}
+
+TSharedPtr<FString> FBlueprintVarActionDetails::GetVariableUnits() const
+{
+	if (CachedVariableName != NAME_None)
+	{
+		if (const UBlueprint* BlueprintObj = GetPropertyOwnerBlueprint() )
+		{
+			FString Result;
+			if (FBlueprintEditorUtils::GetBlueprintVariableMetaData(BlueprintObj, CachedVariableName, GetLocalVariableScope(CachedVariableProperty.Get()), "ForceUnits", /*out*/ Result))
+			{
+				for (const TSharedPtr<FString>& UnitOption : UnitsOptions)
+				{
+					if (*UnitOption == Result)
+					{
+						return UnitOption;
+					}
+				}
+			}
+		}
+	}
+	// Return none;
+	return UnitsOptions.IsEmpty() ? MakeShareable(new FString("None")) : UnitsOptions[0];
+}
+
+void FBlueprintVarActionDetails::OnVariableUnitsChanged(TSharedPtr<FString> UnitsSelected, ESelectInfo::Type SelectInfo)
+{
+	if (CachedVariableName != NAME_None)
+	{
+		if ( UBlueprint* BlueprintObj = GetPropertyOwnerBlueprint() )
+		{
+			if (UnitsSelected && !UnitsOptions.IsEmpty() && UnitsSelected != UnitsOptions[0] )
+			{
+				FBlueprintEditorUtils::SetBlueprintVariableMetaData(BlueprintObj, CachedVariableName, GetLocalVariableScope(CachedVariableProperty.Get()), "ForceUnits", *UnitsSelected);
+			}
+			else
+			{
+				FBlueprintEditorUtils::RemoveBlueprintVariableMetaData(BlueprintObj, CachedVariableName, GetLocalVariableScope(CachedVariableProperty.Get()), "ForceUnits");
+			}
+		}
+	}
+}
+
 ECheckBoxState FBlueprintVarActionDetails::OnFieldNotifyCheckboxState() const
 {
 	UBlueprint* const BlueprintObj = GetBlueprintObj();
@@ -1961,7 +2047,7 @@ void FBlueprintVarActionDetails::OnFieldNotifyChanged(ECheckBoxState InNewState)
 	}
 }
 
-EVisibility FBlueprintVarActionDetails::GetFieldNotifyCheckboxListVisiblity() const
+EVisibility FBlueprintVarActionDetails::GetFieldNotifyCheckboxListVisibility() const
 {
 	UBlueprint* const BlueprintObj = GetBlueprintObj();
 	const FName VarName = CachedVariableName;
diff --git a/Engine/Source/Editor/Kismet/Private/BlueprintDetailsCustomization.h b/Engine/Source/Editor/Kismet/Private/BlueprintDetailsCustomization.h
--- a/Engine/Source/Editor/Kismet/Private/BlueprintDetailsCustomization.h
+++ b/Engine/Source/Editor/Kismet/Private/BlueprintDetailsCustomization.h
@@ -216,17 +216,21 @@ private:
 	TSharedRef< ITableRow > MakeCategoryViewWidget( TSharedPtr<FText> Item, const TSharedRef< STableViewBase >& OwnerTable );
 	void OnCategorySelectionChanged( TSharedPtr<FText> ProposedSelection, ESelectInfo::Type /*SelectInfo*/ );
 	
-	EVisibility ShowEditableCheckboxVisibilty() const;
+	EVisibility ShowEditableCheckboxVisibility() const;
 	ECheckBoxState OnEditableCheckboxState() const;
 	void OnEditableChanged(ECheckBoxState InNewState);
 
-	EVisibility ShowReadOnlyCheckboxVisibilty() const;
+	EVisibility ShowReadOnlyCheckboxVisibility() const;
 	ECheckBoxState OnReadyOnlyCheckboxState() const;
 	void OnReadyOnlyChanged(ECheckBoxState InNewState);
+
+	EVisibility GetVariableUnitsVisibility() const;
+	TSharedPtr<FString> GetVariableUnits() const;
+	void OnVariableUnitsChanged(TSharedPtr<FString> ItemSelected, ESelectInfo::Type SelectInfo);
 	
 	ECheckBoxState OnFieldNotifyCheckboxState() const;
 	void OnFieldNotifyChanged(ECheckBoxState InNewState);
-	EVisibility GetFieldNotifyCheckboxListVisiblity() const;
+	EVisibility GetFieldNotifyCheckboxListVisibility() const;
 
 	ECheckBoxState OnCreateWidgetCheckboxState() const;
 	void OnCreateWidgetChanged(ECheckBoxState InNewState);
@@ -329,6 +333,9 @@ private:
 	/** Array of replication options for our combo text box */
 	TArray<TSharedPtr<FString>> ReplicationOptions;
 
+	/** Array of units options for our combo text box */
+	TArray<TSharedPtr<FString>> UnitsOptions;
+
 	/** Array of enum type names for integers used as bitmasks */
 	TArray<TSharedPtr<FTopLevelAssetPath>> BitmaskEnumTypePaths;
 
diff --git a/Engine/Source/Runtime/Core/Private/Math/UnitConversion.cpp b/Engine/Source/Runtime/Core/Private/Math/UnitConversion.cpp
--- a/Engine/Source/Runtime/Core/Private/Math/UnitConversion.cpp
+++ b/Engine/Source/Runtime/Core/Private/Math/UnitConversion.cpp
@@ -27,7 +27,7 @@ struct FParseCandidate
 };
 
 //TODO: Add note about updating
-FParseCandidate ParseCandidates[] = {
+constexpr FParseCandidate ParseCandidates[] = {
 	
 	{ TEXT("Micrometers"),			EUnit::Micrometers },			{ TEXT("um"),		EUnit::Micrometers }, 			{ TEXT("\u00B5m"),	EUnit::Micrometers },
 	{ TEXT("Millimeters"),			EUnit::Millimeters },			{ TEXT("mm"),		EUnit::Millimeters },
@@ -109,7 +109,7 @@ FParseCandidate ParseCandidates[] = {
 
 	{ TEXT("Percent"),				EUnit::Percentage },			{ TEXT("%"),	EUnit::Percentage },
 
-	{ TEXT("times"),				EUnit::Multiplier },			{ TEXT("x"),	EUnit::Multiplier },			{ TEXT("multiplier"),		EUnit::Multiplier },
+	{ TEXT("Times"),				EUnit::Multiplier },			{ TEXT("x"),	EUnit::Multiplier },			{ TEXT("Multiplier"),		EUnit::Multiplier },
 
 	{ TEXT("Pascals"),				EUnit::Pascals },				{ TEXT("Pa"),	EUnit::Pascals},
 	{ TEXT("KiloPascals"),			EUnit::KiloPascals},			{ TEXT("kPa"),	EUnit::KiloPascals},
@@ -118,7 +118,7 @@ FParseCandidate ParseCandidates[] = {
 };
 
 /** Static array of display strings that directly map to EUnit enumerations */
-const TCHAR* const DisplayStrings[] = {
+constexpr const TCHAR* const DisplayStrings[] = {
 	TEXT("\u00B5m"),			TEXT("mm"),					TEXT("cm"),					TEXT("m"),					TEXT("km"),
 	TEXT("in"),					TEXT("ft"),					TEXT("yd"),					TEXT("mi"),
 	TEXT("ly"),
@@ -157,7 +157,100 @@ const TCHAR* const DisplayStrings[] = {
 	TEXT("Pa"), TEXT("kPa"), TEXT("MPa"), TEXT("GPa"),
 };
 
-const EUnitType UnitTypes[] = {
+constexpr const TCHAR* const SupportedUnitsStrings[] = {
+	TEXT("Micrometers"),
+	TEXT("Millimeters"),		
+	TEXT("Centimeters"),		
+	TEXT("Meters"),				
+	TEXT("Kilometers"),			
+	TEXT("Inches"),				
+	TEXT("Feet"),				
+	TEXT("Yards"),				
+	TEXT("Miles"),				
+	TEXT("Lightyears"),			
+
+	TEXT("Degrees"),				
+	TEXT("Radians"),				
+		
+	TEXT("CentimetersPerSecond"),	
+	TEXT("MetersPerSecond"),		
+	TEXT("KilometersPerHour"),	
+	TEXT("MilesPerHour"),			
+
+	TEXT("DegreesPerSecond"),		
+	TEXT("RadiansPerSecond"),		
+		
+	TEXT("Celsius"),			
+	TEXT("Farenheit"),			
+	TEXT("Kelvin"),				
+					
+	TEXT("Micrograms"),			
+	TEXT("Milligrams"),			
+	TEXT("Grams"),				
+	TEXT("Kilograms"),			
+	TEXT("MetricTons"),			
+	TEXT("Ounces"),				
+	TEXT("Pounds"),				
+	TEXT("Stones"),				
+
+	TEXT("GramsPerCubicCentimeter"),	
+	TEXT("GramsPerCubicMeter"),			
+	TEXT("KilogramsPerCubicCentimeter"),
+	TEXT("KilogramsPerCubicMeter"),		
+
+	TEXT("Newtons"),			
+	TEXT("PoundsForce"),		
+	TEXT("KilogramsForce"),		
+	TEXT("KilogramsCentimetersPerSecondSquared"),	
+
+	TEXT("NewtonMeters"),			
+	TEXT("KilogramsCentimetersSquaredPerSecondSquared"),
+
+	TEXT("Hertz"),				
+	TEXT("Kilohertz"),			
+	TEXT("Megahertz"),			
+	TEXT("Gigahertz"),			
+	TEXT("RevolutionsPerMinute")
+
+	TEXT("Bytes"),				
+	TEXT("Kilobytes"),		
+	TEXT("Megabytes"),		
+	TEXT("Gigabytes"),		
+	TEXT("Terabytes"),		
+
+	TEXT("Lumens"),				
+	TEXT("Candela"),			
+	TEXT("Lux"),				
+	TEXT("CandelaPerMeterSquared"), 
+	TEXT("EV"),				
+
+	TEXT("Nanoseconds"),			
+	TEXT("Microseconds"),			
+	TEXT("Milliseconds"),			
+	TEXT("Seconds"),				
+	TEXT("Minutes"),		
+	TEXT("Hours"),			
+	TEXT("Days"),			
+	TEXT("Months"),			
+	TEXT("Years"),			
+
+	TEXT("ppi"),			
+
+	TEXT("Percent"),
+
+	TEXT("Times"),
+	TEXT("Multiplier"),	
+
+	TEXT("Pascals"),		
+	TEXT("KiloPascals"),	
+	TEXT("MegaPascals"),	
+	TEXT("GigaPascals")
+};
+
+static_assert(UE_ARRAY_COUNT(DisplayStrings) == UE_ARRAY_COUNT(SupportedUnitsStrings));
+static_assert(UE_ARRAY_COUNT(DisplayStrings) == (uint32)EUnit::Unspecified);
+
+constexpr EUnitType UnitTypes[] = {
 	EUnitType::Distance,	EUnitType::Distance,	EUnitType::Distance,	EUnitType::Distance,	EUnitType::Distance,
 	EUnitType::Distance,	EUnitType::Distance,	EUnitType::Distance,	EUnitType::Distance,
 	EUnitType::Distance,
@@ -194,8 +287,6 @@ const EUnitType UnitTypes[] = {
 	EUnitType::Stress, EUnitType::Stress, EUnitType::Stress, EUnitType::Stress,
 };
 
-
-
 DEFINE_EXPRESSION_NODE_TYPE(FNumericUnit<double>, 0x3C138BC9, 0x71314F0B, 0xBB469BF7, 0xED47D147)
 
 struct FUnitExpressionParser
@@ -580,6 +671,11 @@ TOptional<EUnit> FUnitConversion::UnitFromString(const TCHAR* UnitString)
 	return TOptional<EUnit>();
 }
 
+TConstArrayView<const TCHAR*> FUnitConversion::GetSupportedUnits()
+{
+	return TConstArrayView<const TCHAR*>(SupportedUnitsStrings,  UE_ARRAY_COUNT(SupportedUnitsStrings));
+}
+
 namespace UnitConversion
 {
 
diff --git a/Engine/Source/Runtime/Core/Public/Math/UnitConversion.h b/Engine/Source/Runtime/Core/Public/Math/UnitConversion.h
--- a/Engine/Source/Runtime/Core/Public/Math/UnitConversion.h
+++ b/Engine/Source/Runtime/Core/Public/Math/UnitConversion.h
@@ -135,6 +135,9 @@ struct FUnitConversion
 	/** Helper function to find a unit from a string (name or display string) */
 	static CORE_API TOptional<EUnit> UnitFromString(const TCHAR* UnitString);
 
+	/** Helper function to get all supported units */
+	static CORE_API TConstArrayView<const TCHAR*> GetSupportedUnits();
+
 public:
 
 	/** Convert the specified number from one unit to another. Does nothing if the units are incompatible. */
