Expose blueprint reparenting flags publicly and add option for CDO reference replacement
* Deprecated previous function signature and redirected to new signature

diff --git a/Engine/Source/Editor/Kismet/Private/BlueprintCompilationManager.cpp b/Engine/Source/Editor/Kismet/Private/BlueprintCompilationManager.cpp
--- a/Engine/Source/Editor/Kismet/Private/BlueprintCompilationManager.cpp
+++ b/Engine/Source/Editor/Kismet/Private/BlueprintCompilationManager.cpp
@@ -106,14 +106,6 @@ struct FBPCompileRequestInternal
 	FBPCompileRequest UserData;
 };
 
-enum class EReparentClassOptions
-{
-	None = 0x0,
-
-	ReplaceReferencesToOldClasses = 0x1,
-};
-ENUM_CLASS_FLAGS(EReparentClassOptions)
-
 struct FBlueprintCompilationManagerImpl : public FGCObject
 {
 	FBlueprintCompilationManagerImpl();
@@ -186,9 +178,9 @@ struct FBlueprintCompilationManagerImpl : public FGCObject
 // free function that we use to cross a module boundary (from CoreUObject to here)
 void FlushReinstancingQueueImplWrapper();
 void MoveSkelCDOAside(UClass* Class, TMap<UClass*, UClass*>& OldToNewMap);
-void ReparentHierarchiesWrapper(const TMap<UClass*, UClass*>& OldToNewMap)
+void ReparentHierarchiesWrapper(const TMap<UClass*, UClass*>& OldToNewMap, EReparentClassOptions Flags)
 {
-	FBlueprintCompilationManagerImpl::ReparentHierarchies(OldToNewMap, EReparentClassOptions::ReplaceReferencesToOldClasses);
+	FBlueprintCompilationManagerImpl::ReparentHierarchies(OldToNewMap, Flags);
 }
 
 FBlueprintCompilationManagerImpl::FBlueprintCompilationManagerImpl()
@@ -2068,6 +2060,7 @@ void FBlueprintCompilationManagerImpl::VerifyNoQueuedRequests(const TArray<FComp
 void FBlueprintCompilationManagerImpl::ReparentHierarchies(const TMap<UClass*, UClass*>& OldToNewClasses, EReparentClassOptions Options)
 {
 	const bool bReplaceReferencesToOldClasses = (Options & EReparentClassOptions::ReplaceReferencesToOldClasses) != EReparentClassOptions::None;
+	const bool bReplaceReferencesToOldCDOs = (Options & EReparentClassOptions::ReplaceCDOReferences) != EReparentClassOptions::None;
 
 	// something has decided to replace instances of a class. We need to update all the children of those types:
 	TArray< UClass* > ClassesOrdered;
@@ -2193,6 +2186,7 @@ void FBlueprintCompilationManagerImpl::ReparentHierarchies(const TMap<UClass*, U
 	FReplaceInstancesOfClassParameters BatchOptions;
 	BatchOptions.bArchetypesAreUpToDate = true;
 	BatchOptions.bReplaceReferencesToOldClasses = bReplaceReferencesToOldClasses;
+	BatchOptions.bReplaceReferencesToOldCDOs = bReplaceReferencesToOldCDOs;
 
 	// Make sure we don't replace old instances that are in the *callers* old to new TMap!
 	TSet<UObject*> OldObjects;
diff --git a/Engine/Source/Runtime/CoreUObject/Private/Blueprint/BlueprintSupport.cpp b/Engine/Source/Runtime/CoreUObject/Private/Blueprint/BlueprintSupport.cpp
--- a/Engine/Source/Runtime/CoreUObject/Private/Blueprint/BlueprintSupport.cpp
+++ b/Engine/Source/Runtime/CoreUObject/Private/Blueprint/BlueprintSupport.cpp
@@ -139,10 +139,15 @@ void FBlueprintSupport::FlushReinstancingQueue()
 }
 
 void FBlueprintSupport::ReparentHierarchies(const TMap<UClass*, UClass*>& OldClassToNewClass)
+{
+	ReparentHierarchies(OldClassToNewClass, EReparentClassOptions::ReplaceReferencesToOldClasses);
+}
+
+void FBlueprintSupport::ReparentHierarchies(const TMap<UClass*, UClass*>& OldClassToNewClass, EReparentClassOptions Flags)
 {
 	if (ClassReparentingFPtr)
 	{
-		(*ClassReparentingFPtr)(OldClassToNewClass);
+		(*ClassReparentingFPtr)(OldClassToNewClass, Flags);
 	}
 }
 
diff --git a/Engine/Source/Runtime/CoreUObject/Public/Blueprint/BlueprintSupport.h b/Engine/Source/Runtime/CoreUObject/Public/Blueprint/BlueprintSupport.h
--- a/Engine/Source/Runtime/CoreUObject/Public/Blueprint/BlueprintSupport.h
+++ b/Engine/Source/Runtime/CoreUObject/Public/Blueprint/BlueprintSupport.h
@@ -8,6 +8,7 @@
 #include "CoreTypes.h"
 #include "HAL/ThreadSingleton.h"
 #include "Internationalization/Text.h"
+#include "Misc/EnumRange.h"
 #include "UObject/NameTypes.h"
 #include "UObject/UObjectGlobals.h"
 
@@ -78,8 +79,17 @@ struct FBlueprintWarningDeclaration
 	FText WarningDescription;
 };
 
+enum class EReparentClassOptions
+{
+	None = 0x0,
+
+	ReplaceReferencesToOldClasses = 0x1,
+	ReplaceCDOReferences = 0x2,
+};
+ENUM_CLASS_FLAGS(EReparentClassOptions)
+
 typedef void (*FFlushReinstancingQueueFPtr)();
-typedef void (*FClassReparentingFPtr)(const TMap<UClass*, UClass*>&);
+typedef void (*FClassReparentingFPtr)(const TMap<UClass*, UClass*>& /*OldClassToNewClass*/, EReparentClassOptions);
 
 /** 
  * This set of functions contains blueprint related UObject functionality.
@@ -109,7 +119,9 @@ struct FBlueprintSupport
 	static void FlushReinstancingQueue();
 	COREUOBJECT_API static void SetFlushReinstancingQueueFPtr(FFlushReinstancingQueueFPtr Ptr);	
 
+	UE_DEPRECATED(5.4, "Use ReparentHierarchies(const TMap<UClass*, UClass*>&, EReparentClassOptions). Previous behaviour preserved with Flags=EReparentClassOptions::ReplaceReferencesToOldClasses")
 	COREUOBJECT_API static void ReparentHierarchies(const TMap<UClass*, UClass*>& OldClassToNewClass);
+	COREUOBJECT_API static void ReparentHierarchies(const TMap<UClass*, UClass*>& OldClassToNewClass, EReparentClassOptions Flags);
 	COREUOBJECT_API static void SetClassReparentingFPtr(FClassReparentingFPtr Ptr);
 
 	/** Tells if the specified object is one of the many flavors of FLinkerPlaceholderBase that we have. */
