Replace references to CDOs (and hopefully other archetype objects) when using ReparentHierarchies to replace part or all of a UClass hierarchy, this change will not impact blueprint compile on load or compilation

diff --git a/Engine/Source/Editor/Kismet/Private/BlueprintCompilationManager.cpp b/Engine/Source/Editor/Kismet/Private/BlueprintCompilationManager.cpp
--- a/Engine/Source/Editor/Kismet/Private/BlueprintCompilationManager.cpp
+++ b/Engine/Source/Editor/Kismet/Private/BlueprintCompilationManager.cpp
@@ -2213,6 +2213,7 @@ void FBlueprintCompilationManagerImpl::ReparentHierarchies(const TMap<UClass*, U
 	}
 	BatchOptions.ObjectsThatShouldUseOldStuff = &OldObjects;
 	BatchOptions.InstancesThatShouldUseOldClass = &OldObjects;
+	BatchOptions.bReplaceReferencesToOldCDOs = true;
 
 	FBlueprintCompileReinstancer::BatchReplaceInstancesOfClass(OldClassToNewClassDerivedTypes, BatchOptions );
 }
