ReparentHierarchies provided map to for old-to-new template objects for reinstancing reference replacement
* Matches blueprint compilation path

diff --git a/Engine/Source/Editor/Kismet/Private/BlueprintCompilationManager.cpp b/Engine/Source/Editor/Kismet/Private/BlueprintCompilationManager.cpp
--- a/Engine/Source/Editor/Kismet/Private/BlueprintCompilationManager.cpp
+++ b/Engine/Source/Editor/Kismet/Private/BlueprintCompilationManager.cpp
@@ -2175,7 +2175,8 @@ void FBlueprintCompilationManagerImpl::ReparentHierarchies(const TMap<UClass*, U
 	}
 	
 	// Reparenting done, reinstance the hierarchy and update archetypes:
-	ReinstanceBatch(Reinstancers, OldClassToNewClassIncludingChildren, nullptr);
+	TMap<UClass*, TMap<UObject*, UObject*>> OldToNewTemplates;
+	ReinstanceBatch(Reinstancers, OldClassToNewClassIncludingChildren, nullptr, &OldToNewTemplates);
 
 	// Reinstance (non archetype) instances
 	TMap<UClass*, UClass*> OldClassToNewClassDerivedTypes;
@@ -2214,6 +2215,7 @@ void FBlueprintCompilationManagerImpl::ReparentHierarchies(const TMap<UClass*, U
 	BatchOptions.ObjectsThatShouldUseOldStuff = &OldObjects;
 	BatchOptions.InstancesThatShouldUseOldClass = &OldObjects;
 	BatchOptions.bReplaceReferencesToOldCDOs = true;
+	BatchOptions.OldToNewTemplates = &OldToNewTemplates;
 
 	FBlueprintCompileReinstancer::BatchReplaceInstancesOfClass(OldClassToNewClassDerivedTypes, BatchOptions );
 }
