Kick the async task queue while waiting to yield an in-progress Find-in-Blueprint global search to a GC pass on the main thread in order to avoid a potential deadlock scenario.

diff --git a/Engine/Source/Editor/Kismet/Private/FindInBlueprintManager.cpp b/Engine/Source/Editor/Kismet/Private/FindInBlueprintManager.cpp
--- a/Engine/Source/Editor/Kismet/Private/FindInBlueprintManager.cpp
+++ b/Engine/Source/Editor/Kismet/Private/FindInBlueprintManager.cpp
@@ -6,6 +6,7 @@
 #include "HAL/FileManager.h"
 #include "Misc/Paths.h"
 #include "Async/ParallelFor.h"
+#include "Async/TaskGraphInterfaces.h"
 #include "Misc/ScopeLock.h"
 #include "Serialization/MemoryWriter.h"
 #include "Serialization/MemoryReader.h"
@@ -2972,9 +2973,21 @@ void FFindInBlueprintSearchManager::PauseFindInBlueprintSearch()
 
 	// It is UNSAFE to lock any other critical section here, threads need them to finish a cycle of searching. Next cycle they will pause
 
+	// We don't expect to be pausing a search off the main/game thread (i.e. GC).
+	check(IsInGameThread());
+	FTaskGraphInterface& TaskGraphInterface = FTaskGraphInterface::Get();
+
 	// Wait until all threads have come to a stop, it won't take long
 	while(ActiveSearchCounter.GetValue() > 0)
 	{
+		// Async tasks may have been registered to the game thread, so make sure
+		// we're processing those here (e.g. FSearchableValueInfo::GetDisplayText).
+		if (TaskGraphInterface.IsThreadProcessingTasks(ENamedThreads::GameThread))
+		{
+			TaskGraphInterface.ProcessThreadUntilIdle(ENamedThreads::GameThread);
+		}
+
+		// Yield some time to other threads.
 		FPlatformProcess::Sleep(0.1f);
 	}
 }
