Fix reinst in pseudo async loading setting

diff --git a/Engine/Source/Editor/Kismet/Private/BlueprintCompilationManager.cpp b/Engine/Source/Editor/Kismet/Private/BlueprintCompilationManager.cpp
--- a/Engine/Source/Editor/Kismet/Private/BlueprintCompilationManager.cpp
+++ b/Engine/Source/Editor/Kismet/Private/BlueprintCompilationManager.cpp
@@ -1964,7 +1964,7 @@ void FBlueprintCompilationManagerImpl::FlushReinstancingQueueImpl(bool bFindAndR
 
 		// Special case when we run on ALT, we want to cleanup all classes flagged for reinstanciation right away.
 		const bool bIsInActualAsyncLoadingThread = IsInAsyncLoadingThread() && !IsInGameThread();
-		if (IsAsyncLoading() && IsAsyncLoadingMultithreaded() && !bIsInActualAsyncLoadingThread)
+		if (IsAsyncLoading() && (!IsAsyncLoadingMultithreaded() || !bIsInActualAsyncLoadingThread))
 		{
 			// While async loading we only remove classes that have no instances being
 			// async loaded. Those instances will need to be reinstanced once they finish
