Fixed regression introduced in cl 28083785

diff --git a/Engine/Source/Editor/Kismet/Private/BlueprintCompilationManager.cpp b/Engine/Source/Editor/Kismet/Private/BlueprintCompilationManager.cpp
--- a/Engine/Source/Editor/Kismet/Private/BlueprintCompilationManager.cpp
+++ b/Engine/Source/Editor/Kismet/Private/BlueprintCompilationManager.cpp
@@ -2649,7 +2649,7 @@ void FBlueprintCompilationManagerImpl::ReinstanceBatch(TArray<FReinstancingJob>&
 					TMap<UObject*, UObject*> OldToNewInstanceMap(CreatedInstanceMap);
 					for (const auto& Pair : CreatedInstanceMap)
 					{
-						FBlueprintCompileReinstancer::CopyPropertiesForUnrelatedObjects(OldInstance, *NewInstance, /*bClearExternalReferences*/true, bUseDeltaSerialization, /*bOnlyHandleDirectSubObjects*/true, &OldToNewInstanceMap);
+						FBlueprintCompileReinstancer::CopyPropertiesForUnrelatedObjects(Pair.Key, Pair.Value, /*bClearExternalReferences*/true, bUseDeltaSerialization, /*bOnlyHandleDirectSubObjects*/true, &OldToNewInstanceMap);
 					}
 
 					if (OldToNewTemplates && !OldToNewInstanceMap.IsEmpty())
