[BugFix] added warning log for when a user tries to diff a level blueprint and the old revision doesn't have a level blueprint

diff --git a/Engine/Source/Editor/Kismet/Private/SBlueprintEditorToolbar.cpp b/Engine/Source/Editor/Kismet/Private/SBlueprintEditorToolbar.cpp
--- a/Engine/Source/Editor/Kismet/Private/SBlueprintEditorToolbar.cpp
+++ b/Engine/Source/Editor/Kismet/Private/SBlueprintEditorToolbar.cpp
@@ -592,6 +592,10 @@ static void OnDiffRevisionPicked(FRevisionInfo const& RevisionInfo, TWeakObjectP
 									break;
 								}
 							}
+							if (!PreviousAsset)
+							{
+								UE_LOG(LogSourceControl, Warning, TEXT("Revision %s of %s doesn't have a LevelScriptBlueprint"), *Revision->GetRevision(), *Revision->GetFilename());
+							}
 						}
 						// otherwise its a normal Blueprint
 						else
