Add a cvar Maximum Search Depth to the Blueprint Debugger

diff --git a/Engine/Source/Editor/Kismet/Private/Debugging/SKismetDebugTreeView.cpp b/Engine/Source/Editor/Kismet/Private/Debugging/SKismetDebugTreeView.cpp
--- a/Engine/Source/Editor/Kismet/Private/Debugging/SKismetDebugTreeView.cpp
+++ b/Engine/Source/Editor/Kismet/Private/Debugging/SKismetDebugTreeView.cpp
@@ -108,6 +108,9 @@ static const FText ViewInDebuggerTooltipText = LOCTEXT("ViewInDebugger_Tooltip",
 static constexpr float ThumbnailIconSize = 16.0f;
 static constexpr uint32 ThumbnailIconResolution = 16;
 
+static int DebuggerMaxSearchDepth = 5;
+static FAutoConsoleVariableRef CVarDebuggerMaxDepth(TEXT("bp.DebuggerMaxSearchDepth"), DebuggerMaxSearchDepth, TEXT("The maximum search depth of Blueprint Debugger TreeView widgets"), ECVF_Default);
+
 //////////////////////////////////////////////////////////////////////////
 
 const FName SKismetDebugTreeView::ColumnId_Name("Name");
@@ -666,11 +669,15 @@ protected:
 					continue;
 				}
 
-					// if any children need to expand, so should this
-					if (Child->SearchRecursive(InSearchString, DebugTreeView, Parents, ChildSearchFlags))
-					{
-						bVisible = true;
-						bChildMatch = true;
+					// stop recursing if we reached the max depth
+                    if (Parents.Num() <= DebuggerMaxSearchDepth)
+                    {
+						// if any children need to expand, so should this
+						if (Child->SearchRecursive(InSearchString, DebugTreeView, Parents, ChildSearchFlags))
+						{
+							bVisible = true;
+							bChildMatch = true;
+						}
 					}
 			}
 			else
