Blueprint search syntax: filter with a nested subfilter will now intersect initial results with sub-filter results, fixing && behaving as || when either operand is a nested subfilter.

diff --git a/Engine/Source/Editor/Kismet/Private/FiBSearchInstance.cpp b/Engine/Source/Editor/Kismet/Private/FiBSearchInstance.cpp
--- a/Engine/Source/Editor/Kismet/Private/FiBSearchInstance.cpp
+++ b/Engine/Source/Editor/Kismet/Private/FiBSearchInstance.cpp
@@ -293,17 +293,33 @@ bool FFiBSearchInstance::OnFilterFunction(const FTextFilterString& A, ESearchQue
 		// Proceed to doing a sub-search
 		if (SubSearchInstance->PendingSearchables.Num() > 0)
 		{
+			// Make a copy of results so far. We'll intersect with sub-search results.
+			TSet<FImaginaryFiBData*> InitialResultsCopy;
+			for (const FImaginaryFiBDataWeakPtr& ResultCopy : SubSearchInstance->PendingSearchables)
+			{
+				if (ResultCopy.IsValid())
+				{
+					InitialResultsCopy.Add(ResultCopy.Pin().Get());
+				}
+			}
+
 			bSearchSuccess = SubSearchInstance->DoSearchQuery(A.AsString(), InSearchQueryFilter == ESearchQueryFilter::AllFilter);
 			if (bSearchSuccess)
 			{
-				for (auto& MatchesItem : SubSearchInstance->MatchesSearchQuery)
+				for (const FImaginaryFiBData* MatchesItem : SubSearchInstance->MatchesSearchQuery)
 				{
-					LastFunctionResultMatchesSearchQuery.AddUnique(MatchesItem);
+					if (InitialResultsCopy.Contains(MatchesItem))
+					{
+						LastFunctionResultMatchesSearchQuery.AddUnique(MatchesItem);
+					}
 				}
 
-				for (auto& MatchesItem : SubSearchInstance->MatchingSearchComponents)
+				for (const TPair<const FImaginaryFiBData*, FComponentUniqueDisplay>& MatchesItem : SubSearchInstance->MatchingSearchComponents)
 				{
-					LastFunctionMatchingSearchComponents.AddUnique(MatchesItem.Key, MatchesItem.Value);
+					if (InitialResultsCopy.Contains(MatchesItem.Key))
+					{
+						LastFunctionMatchingSearchComponents.AddUnique(MatchesItem.Key, MatchesItem.Value);
+					}
 				}
 			}
 		}
