Improve unreal insights instrumentation of editor loading and blueprint compilation
Removed redundant events and added asset-specific trace events when AssetLoading is enabled
Add SCOPED_LOADTIMER_ASSET_TEXT macro for outputting an asset name during load and improved comments

diff --git a/Engine/Source/Editor/Kismet/Private/BlueprintCompilationManager.cpp b/Engine/Source/Editor/Kismet/Private/BlueprintCompilationManager.cpp
--- a/Engine/Source/Editor/Kismet/Private/BlueprintCompilationManager.cpp
+++ b/Engine/Source/Editor/Kismet/Private/BlueprintCompilationManager.cpp
@@ -35,6 +35,7 @@
 #include "ProfilingDebugging/ScopedTimers.h"
 #include "Serialization/ArchiveHasReferences.h"
 #include "Serialization/ArchiveReplaceObjectRef.h"
+#include "ProfilingDebugging/LoadTimeTracker.h"
 #include "TickableEditorObject.h"
 #include "UObject/FortniteMainBranchObjectVersion.h"
 #include "UObject/MetaData.h"
@@ -263,11 +264,11 @@ void FBlueprintCompilationManagerImpl::QueueForCompilation(const FBPCompileReque
 
 void FBlueprintCompilationManagerImpl::CompileSynchronouslyImpl(const FBPCompileRequestInternal& Request)
 {
+	TRACE_CPUPROFILER_EVENT_SCOPE(CompileSynchronouslyImpl);
+
 #if WITH_EDITOR
 	FScopeLock ScopeLock(&Lock);
 #endif
-	DECLARE_SCOPE_HIERARCHICAL_COUNTER_FUNC()
-
 	Request.UserData.BPToCompile->bQueuedForCompilation = true;
 
 	const bool bIsRegeneratingOnLoad		= (Request.UserData.CompileOptions & EBlueprintCompileOptions::IsRegeneratingOnLoad				) != EBlueprintCompileOptions::None;
@@ -338,7 +339,7 @@ void FBlueprintCompilationManagerImpl::CompileSynchronouslyImpl(const FBPCompile
 
 	if ( GEditor && !bRegenerateSkeletonOnly)
 	{
-		DECLARE_SCOPE_HIERARCHICAL_COUNTER(BroadcastBlueprintReinstanced)
+		TRACE_CPUPROFILER_EVENT_SCOPE(BroadcastBlueprintReinstanced)
 
 		// Make sure clients know they're being reinstanced as part of blueprint compilation. After this point
 		// compilation is completely done:
@@ -350,15 +351,13 @@ void FBlueprintCompilationManagerImpl::CompileSynchronouslyImpl(const FBPCompile
 
 	if(!bSkipGarbageCollection)
 	{
-		DECLARE_SCOPE_HIERARCHICAL_COUNTER(CollectGarbage)
-
 		TGuardValue<bool> GuardTemplateNameFlag(GIsGCingAfterBlueprintCompile, true);
 		CollectGarbage(GARBAGE_COLLECTION_KEEPFLAGS);
 	}
 
 	if (!bRegenerateSkeletonOnly)
 	{
-		DECLARE_SCOPE_HIERARCHICAL_COUNTER(BroadcastChanged)
+		TRACE_CPUPROFILER_EVENT_SCOPE(BroadcastChanged);
 
 		for(UBlueprint* BP : SkeletonCompiledBlueprints)
 		{
@@ -372,7 +371,7 @@ void FBlueprintCompilationManagerImpl::CompileSynchronouslyImpl(const FBPCompile
 
 	if (!bBatchCompile && !bRegenerateSkeletonOnly)
 	{
-		DECLARE_SCOPE_HIERARCHICAL_COUNTER(BroadCastCompiled)
+		TRACE_CPUPROFILER_EVENT_SCOPE(BroadCastCompiled);
 
 		for(UBlueprint* BP : SkeletonCompiledBlueprints)
 		{
@@ -605,8 +604,7 @@ namespace UE::Kismet::BlueprintCompilationManager::Private
 
 void FBlueprintCompilationManagerImpl::FlushCompilationQueueImpl(bool bSuppressBroadcastCompiled, TArray<UBlueprint*>* BlueprintsCompiled, TArray<UBlueprint*>* BlueprintsCompiledOrSkeletonCompiled, FUObjectSerializeContext* InLoadContext, TMap<UClass*, TMap<UObject*, UObject*>>* OldToNewTemplates /* = nullptr*/)
 {
-	TRACE_CPUPROFILER_EVENT_SCOPE(FBlueprintCompilationManager::FlushCompilationQueue);
-	DECLARE_SCOPE_HIERARCHICAL_COUNTER_FUNC()
+	TRACE_CPUPROFILER_EVENT_SCOPE(FlushCompilationQueueImpl);
 
 #if WITH_EDITOR
 	FScopeLock ScopeLock(&Lock);
@@ -1027,8 +1025,7 @@ void FBlueprintCompilationManagerImpl::FlushCompilationQueueImpl(bool bSuppressB
 					}
 					else
 					{
-						TRACE_CPUPROFILER_EVENT_SCOPE(MoveSkelCDOAside_OLD);
-
+						// Old code path
 						MoveSkelCDOAside(OldSkeletonClass, NewSkeletonToOldSkeleton);
 					}
 				}
@@ -1048,7 +1045,8 @@ void FBlueprintCompilationManagerImpl::FlushCompilationQueueImpl(bool bSuppressB
 		
 				if(CompilerData.ShouldRegenerateSkeleton())
 				{
-					DECLARE_SCOPE_HIERARCHICAL_COUNTER(RecompileSkeleton)
+					TRACE_CPUPROFILER_EVENT_SCOPE(RecompileSkeleton);
+					SCOPED_LOADTIMER_ASSET_TEXT(*BP->GetPathName());
 
 					if(BlueprintsCompiledOrSkeletonCompiled)
 					{
@@ -1089,7 +1087,7 @@ void FBlueprintCompilationManagerImpl::FlushCompilationQueueImpl(bool bSuppressB
 				}
 				else
 				{
-					DECLARE_SCOPE_HIERARCHICAL_COUNTER(RelinkSkeleton)
+					TRACE_CPUPROFILER_EVENT_SCOPE(RelinkSkeleton);
 
 					// Just relink, note that UProperties that reference *other* types may be stale until
 					// we fixup below:
@@ -1103,7 +1101,6 @@ void FBlueprintCompilationManagerImpl::FlushCompilationQueueImpl(bool bSuppressB
 					BP->bHasBeenRegenerated = true;
 					if (BP->GeneratedClass)
 					{
-						DECLARE_SCOPE_HIERARCHICAL_COUNTER(ClearFunctionMapsCaches)
 						BP->GeneratedClass->ClearFunctionMapsCaches();
 					}
 				}
@@ -1113,7 +1110,7 @@ void FBlueprintCompilationManagerImpl::FlushCompilationQueueImpl(bool bSuppressB
 			// that may have been created as part of skeleton generation:
 			for (FCompilerData& CompilerData : CurrentlyCompilingBPs)
 			{
-				DECLARE_SCOPE_HIERARCHICAL_COUNTER(FixUpDelegateParameters)
+				TRACE_CPUPROFILER_EVENT_SCOPE(FixUpDelegateParameters);
 
 				UBlueprint* BP = CompilerData.BP;
 				TArray<FSkeletonFixupData>& ParamsToFix = CompilerData.SkeletonFixupData;
@@ -1206,7 +1203,7 @@ void FBlueprintCompilationManagerImpl::FlushCompilationQueueImpl(bool bSuppressB
 		// Detect any variable-based properties that are not in the old generated class, save them for after reinstancing. This can occur 
 		//    when a new variable is introduced in an ancestor class, and we'll need to use its default as our generated class's initial value.
 		{
-			DECLARE_SCOPE_HIERARCHICAL_COUNTER(DetectNewDefaultVariables)
+			TRACE_CPUPROFILER_EVENT_SCOPE(DetectNewDefaultVariables);
 
 			for (FCompilerData& CompilerData : CurrentlyCompilingBPs)
 			{
@@ -1242,9 +1239,9 @@ void FBlueprintCompilationManagerImpl::FlushCompilationQueueImpl(bool bSuppressB
 				continue;
 			}
 
-			DECLARE_SCOPE_HIERARCHICAL_COUNTER(ReconstructNodes)
-
 			UBlueprint* BP = CompilerData.BP;
+			TRACE_CPUPROFILER_EVENT_SCOPE(ReconstructNodes);
+			SCOPED_LOADTIMER_ASSET_TEXT(*BP->GetPathName());
 			UE_TRACK_REFERENCING_PACKAGE_SCOPED(BP, PackageAccessTrackingOps::NAME_CookerBuildObject);
 
 			ConformToParentAndInterfaces(BP);
@@ -1313,6 +1310,7 @@ void FBlueprintCompilationManagerImpl::FlushCompilationQueueImpl(bool bSuppressB
 		// compiled will be parented to REINST versions of the class, so type checks (IsA, etc) involving those types
 		// will be incoherent!
 		{
+			TRACE_CPUPROFILER_EVENT_SCOPE(ReinstanceQueued);
 			for (FCompilerData& CompilerData : CurrentlyCompilingBPs)
 			{
 				// we including skeleton only compilation jobs for reinstancing because we need UpdateCustomPropertyListForPostConstruction
@@ -1323,8 +1321,6 @@ void FBlueprintCompilationManagerImpl::FlushCompilationQueueImpl(bool bSuppressB
 					continue;
 				}
 
-				DECLARE_SCOPE_HIERARCHICAL_COUNTER(ReinstanceQueuedBlueprint)
-
 				// no need to reinstance skeleton or relink jobs that are not in a hierarchy that has had reinstancing initiated:
 				bool bRequiresReinstance = CompilerData.ShouldInitiateReinstancing();
 				if (!bRequiresReinstance)
@@ -1352,6 +1348,7 @@ void FBlueprintCompilationManagerImpl::FlushCompilationQueueImpl(bool bSuppressB
 				}
 
 				UBlueprint* BP = CompilerData.BP;
+				SCOPED_LOADTIMER_ASSET_TEXT(*BP->GetPathName());
 
 				if(BP->GeneratedClass)
 				{
@@ -1407,7 +1404,8 @@ void FBlueprintCompilationManagerImpl::FlushCompilationQueueImpl(bool bSuppressB
 			UBlueprint* BP = CompilerData.BP;
 			if(CompilerData.ShouldCompileClassLayout())
 			{
-				DECLARE_SCOPE_HIERARCHICAL_COUNTER(CompileClassLayout)
+				TRACE_CPUPROFILER_EVENT_SCOPE(CompileClassLayout);
+				SCOPED_LOADTIMER_ASSET_TEXT(*BP->GetPathName());
 
 				ensure( BP->GeneratedClass == nullptr ||
 						BP->GeneratedClass->ClassDefaultObject == nullptr || 
@@ -1494,7 +1492,8 @@ void FBlueprintCompilationManagerImpl::FlushCompilationQueueImpl(bool bSuppressB
 			}
 			else
 			{
-				DECLARE_SCOPE_HIERARCHICAL_COUNTER(CompileClassFunctions)
+				TRACE_CPUPROFILER_EVENT_SCOPE(CompileClassFunctions);
+				SCOPED_LOADTIMER_ASSET_TEXT(*BP->GetPathName());
 
 				// default value propagation occurs below:
 				if(BPGC)
@@ -1579,7 +1578,7 @@ void FBlueprintCompilationManagerImpl::FlushCompilationQueueImpl(bool bSuppressB
 	// STAGE XIV: Now we can finish the first stage of the reinstancing operation, moving old classes to new classes:
 	{
 		{
-			DECLARE_SCOPE_HIERARCHICAL_COUNTER(MoveOldClassesToNewClasses)
+			TRACE_CPUPROFILER_EVENT_SCOPE(MoveOldClassesToNewClasses);
 
 			TArray<FReinstancingJob> Reinstancers;
 			// Set up reinstancing jobs - we need a reference to the compiler in order to honor 
@@ -1602,7 +1601,7 @@ void FBlueprintCompilationManagerImpl::FlushCompilationQueueImpl(bool bSuppressB
 
 		// Set default values on any newly-introduced variables (from ancestor BPs)
 		{
-			DECLARE_SCOPE_HIERARCHICAL_COUNTER(SetNewDefaultVariables)
+			TRACE_CPUPROFILER_EVENT_SCOPE(SetNewDefaultVariables);
 
 			for (FCompilerData& CompilerData : CurrentlyCompilingBPs)
 			{
@@ -1629,10 +1628,11 @@ void FBlueprintCompilationManagerImpl::FlushCompilationQueueImpl(bool bSuppressB
 		// STAGE XV: POST CDO COMPILED
 		for (FCompilerData& CompilerData : CurrentlyCompilingBPs)
 		{
-			DECLARE_SCOPE_HIERARCHICAL_COUNTER(PostCDOCompiled)
+			TRACE_CPUPROFILER_EVENT_SCOPE(PostCDOCompiled);
 
 			if (CompilerData.Compiler.IsValid())
 			{
+				SCOPED_LOADTIMER_ASSET_TEXT(*CompilerData.BP->GetPathName());
 				UObject::FPostCDOCompiledContext PostCDOCompiledContext;
 				PostCDOCompiledContext.bIsRegeneratingOnLoad = CompilerData.BP->bIsRegeneratingOnLoad;
 				PostCDOCompiledContext.bIsSkeletonOnly = CompilerData.IsSkeletonOnly();
@@ -1644,7 +1644,7 @@ void FBlueprintCompilationManagerImpl::FlushCompilationQueueImpl(bool bSuppressB
 		// STAGE XVI: CLEAR TEMPORARY FLAGS
 		for (FCompilerData& CompilerData : CurrentlyCompilingBPs)
 		{
-			DECLARE_SCOPE_HIERARCHICAL_COUNTER(ClearTemporaryFlags)
+			TRACE_CPUPROFILER_EVENT_SCOPE(ClearTemporaryFlags);
 
 			UBlueprint* BP = CompilerData.BP;
 
@@ -1790,7 +1790,7 @@ void FBlueprintCompilationManagerImpl::FlushCompilationQueueImpl(bool bSuppressB
 
 		if (!bSuppressBroadcastCompiled)
 		{
-			DECLARE_SCOPE_HIERARCHICAL_COUNTER(BroadcastBlueprintCompiled)
+			TRACE_CPUPROFILER_EVENT_SCOPE(BroadcastBlueprintCompiled);
 
 			if(GEditor)
 			{
@@ -1811,7 +1811,7 @@ void FBlueprintCompilationManagerImpl::FlushCompilationQueueImpl(bool bSuppressB
 	}
 
 	{
-		DECLARE_SCOPE_HIERARCHICAL_COUNTER(UEdGraphPin::Purge)
+		TRACE_CPUPROFILER_EVENT_SCOPE(UEdGraphPin::Purge);
 		UEdGraphPin::Purge();
 	}
 
@@ -1825,7 +1825,7 @@ void FBlueprintCompilationManagerImpl::FlushCompilationQueueImpl(bool bSuppressB
 
 void FBlueprintCompilationManagerImpl::FixupDelegateProperties(const TArray<FCompilerData>& InCurrentlyCompilingBPs)
 {
-	DECLARE_SCOPE_HIERARCHICAL_COUNTER_FUNC();
+	TRACE_CPUPROFILER_EVENT_SCOPE(FixupDelegateProperties);
 
 	if (InCurrentlyCompilingBPs.Num() <= 1)
 	{
@@ -1877,7 +1877,7 @@ void FBlueprintCompilationManagerImpl::FixupDelegateProperties(const TArray<FCom
 
 void FBlueprintCompilationManagerImpl::ProcessExtensions(const TArray<FCompilerData>& InCurrentlyCompilingBPs)
 {
-	DECLARE_SCOPE_HIERARCHICAL_COUNTER_FUNC()
+	TRACE_CPUPROFILER_EVENT_SCOPE(ProcessExtensions);
 
 	if(CompilerExtensions.Num() == 0)
 	{
@@ -1930,7 +1930,7 @@ void FBlueprintCompilationManagerImpl::ProcessExtensions(const TArray<FCompilerD
 
 void FBlueprintCompilationManagerImpl::FlushReinstancingQueueImpl(bool bFindAndReplaceCDOReferences, TMap<UClass*, TMap<UObject*, UObject*>>* OldToNewTemplates /* = nullptr*/)
 {
-	DECLARE_SCOPE_HIERARCHICAL_COUNTER_FUNC()
+	TRACE_CPUPROFILER_EVENT_SCOPE(FlushReinstancingQueueImpl);
 
 #if WITH_EDITOR
 	FScopeLock ScopeLock(&Lock);
@@ -2069,6 +2069,7 @@ void FBlueprintCompilationManagerImpl::VerifyNoQueuedRequests(const TArray<FComp
 void FBlueprintCompilationManagerImpl::ReparentHierarchies(const TMap<UClass*, UClass*>& OldToNewClasses, EReparentClassOptions Options)
 {
 	const bool bReplaceReferencesToOldClasses = (Options & EReparentClassOptions::ReplaceReferencesToOldClasses) != EReparentClassOptions::None;
+	TRACE_CPUPROFILER_EVENT_SCOPE(ReparentHierarchies);
 
 	// something has decided to replace instances of a class. We need to update all the children of those types:
 	TArray< UClass* > ClassesOrdered;
@@ -2519,6 +2520,8 @@ void FBlueprintCompilationManagerImpl::ReinstanceBatch(TArray<FReinstancingJob>&
 		UClass* OldClass = ReinstancingJob.OldToNew.Key;
 		if(OldClass)
 		{
+			SCOPED_LOADTIMER_ASSET_TEXT(*WriteToString<256>(TEXT("Reinstancing "), *GetPathNameSafe(ReinstancingJob.OldToNew.Value)));
+
 			UClass* NewClass = ReinstancingJob.OldToNew.Value;
 			if (NewClass && 
 				OldClass->ClassDefaultObject && 
@@ -2628,6 +2631,8 @@ void FBlueprintCompilationManagerImpl::ReinstanceBatch(TArray<FReinstancingJob>&
 		UClass* OldClass = ReinstancingJob.OldToNew.Key;
 		if(OldClass)
 		{
+			SCOPED_LOADTIMER_ASSET_TEXT(*WriteToString<256>(TEXT("FinishReinstancing "), *GetPathNameSafe(ReinstancingJob.OldToNew.Value)));
+
 			TArray<UObject*> OldInstances;
 			GetObjectsOfClass( OldClass, OldInstances, false );
 
diff --git a/Engine/Source/Editor/KismetCompiler/Private/KismetCompiler.cpp b/Engine/Source/Editor/KismetCompiler/Private/KismetCompiler.cpp
--- a/Engine/Source/Editor/KismetCompiler/Private/KismetCompiler.cpp
+++ b/Engine/Source/Editor/KismetCompiler/Private/KismetCompiler.cpp
@@ -576,7 +576,6 @@ FKismetCompilerContext::~FKismetCompilerContext()
 
 UEdGraphSchema_K2* FKismetCompilerContext::CreateSchema()
 {
-	TRACE_CPUPROFILER_EVENT_SCOPE(WILD_FKismetCompilerContext::CreateSchema);
 	return NewObject<UEdGraphSchema_K2>();
 }
 
@@ -796,7 +795,7 @@ void FKismetCompilerContext::SaveSubObjectsFromCleanAndSanitizeClass(FSubobjectC
 
 void FKismetCompilerContext::PostCreateSchema()
 {
-	TRACE_CPUPROFILER_EVENT_SCOPE(WILD_FKismetCompilerContext::PostCreateSchema);
+	TRACE_CPUPROFILER_EVENT_SCOPE(PostCreateSchema);
 	NodeHandlers.Add(UEdGraphNode_Comment::StaticClass(), new FNodeHandlingFunctor(*this));
 
 	TArray<UClass*> ClassesOfUK2Node;
@@ -4590,7 +4589,7 @@ void FKismetCompilerContext::ValidateFunctionGraphNames()
 // Creates a copy of the graph to allow further transformations to occur
 void FKismetCompilerContext::CreateFunctionList()
 {
-	TRACE_CPUPROFILER_EVENT_SCOPE(WILD_FKismetCompilerContext::CreateFunctionList);
+	TRACE_CPUPROFILER_EVENT_SCOPE(CreateFunctionList);
 	{
 		BP_SCOPED_COMPILER_EVENT_STAT(EKismetCompilerStats_GenerateFunctionGraphs);
 
diff --git a/Engine/Source/Editor/UnrealEd/Private/Kismet2/KismetReinstanceUtilities.cpp b/Engine/Source/Editor/UnrealEd/Private/Kismet2/KismetReinstanceUtilities.cpp
--- a/Engine/Source/Editor/UnrealEd/Private/Kismet2/KismetReinstanceUtilities.cpp
+++ b/Engine/Source/Editor/UnrealEd/Private/Kismet2/KismetReinstanceUtilities.cpp
@@ -42,6 +42,7 @@
 #include "InstancedReferenceSubobjectHelper.h"
 #include "UObject/PropertyOptional.h"
 #include "UObject/PropertyBagRepository.h"
+#include "ProfilingDebugging/LoadTimeTracker.h"
 
 DECLARE_CYCLE_STAT(TEXT("Replace Instances"), EKismetReinstancerStats_ReplaceInstancesOfClass, STATGROUP_KismetReinstancer );
 DECLARE_CYCLE_STAT(TEXT("Find Referencers"), EKismetReinstancerStats_FindReferencers, STATGROUP_KismetReinstancer );
@@ -2048,6 +2049,7 @@ bool FBlueprintCompileReinstancer::IsReinstClass(const UClass* Class)
 
 static void ReplaceObjectHelper(UObject*& OldObject, UClass* OldClass, UObject*& NewUObject, UClass* NewClass, TMap<UObject*, UObject*>& OldToNewInstanceMap, const TMap<UClass*, UClass*>& OldToNewClassMap, TMap<UObject*, FName>& OldToNewNameMap, int32 OldObjIndex, TArray<UObject*>& ObjectsToReplace, TArray<UObject*>& PotentialEditorsForRefreshing, TSet<AActor*>& OwnersToRerunConstructionScript, TFunctionRef<TArray<TObjectPtr<USceneComponent>>&(USceneComponent*)> GetAttachChildrenArray, bool bIsComponent, bool bArchetypesAreUpToDate)
 {
+	SCOPED_LOADTIMER_ASSET_TEXT(*WriteToString<256>(TEXT("ReplaceObjectHelper "), *GetPathNameSafe(OldObject)));
 	// If the old object was spawned from an archetype (i.e. not the CDO), we must use the new version of that archetype as the template object when constructing the new instance.
 	UObject* NewArchetype = nullptr;
 	if(bArchetypesAreUpToDate)
@@ -3017,6 +3019,7 @@ void FBlueprintCompileReinstancer::ReparentChild(UClass* ChildClass)
 
 void FBlueprintCompileReinstancer::CopyPropertiesForUnrelatedObjects(UObject* OldObject, UObject* NewObject, bool bClearExternalReferences, bool bForceDeltaSerialization /* = false */, bool bOnlyHandleDirectSubObjects/* = false */, TMap<UObject*, UObject*>* OldToNewInstanceMap /*=nullptr*/)
 {
+	SCOPED_LOADTIMER_ASSET_TEXT(*WriteToString<256>(TEXT("CopyPropertiesForUnrelatedObjects "), *GetPathNameSafe(NewObject)));
 	UEngine::FCopyPropertiesForUnrelatedObjectsParams Params;
 	// During a blueprint reparent, delta serialization must be enabled to correctly copy all properties
 	Params.bDoDelta = bForceDeltaSerialization || !OldObject->HasAnyFlags(RF_ClassDefaultObject);
diff --git a/Engine/Source/Runtime/Core/Public/ProfilingDebugging/LoadTimeTracker.h b/Engine/Source/Runtime/Core/Public/ProfilingDebugging/LoadTimeTracker.h
--- a/Engine/Source/Runtime/Core/Public/ProfilingDebugging/LoadTimeTracker.h
+++ b/Engine/Source/Runtime/Core/Public/ProfilingDebugging/LoadTimeTracker.h
@@ -280,8 +280,10 @@ struct FScopedLoadTimeAccumulatorTimer : public FScopedDurationTimer
 #define SCOPED_ACCUM_LOADTIME(TimerName, InstanceName)
 #endif
 
+// Uses raw timers to store cumulative load times, does not support specific strings
 #if ENABLE_LOADTIME_RAW_TIMINGS
 #define SCOPED_LOADTIMER_TEXT(TimerName)
+#define SCOPED_LOADTIMER_ASSET_TEXT(TimerName)
 #define SCOPED_LOADTIMER(TimerName) FScopedDurationTimer DurationTimer_##TimerName(FLoadTimeTracker::Get().TimerName);
 #define SCOPED_CUSTOM_LOADTIMER(TimerName)
 #define SCOPED_LOADTIMER_CNT(TimerName) FScopedDurationTimer DurationTimer_##TimerName(FLoadTimeTracker::Get().TimerName); FLoadTimeTracker::Get().TimerName##Cnt++;
@@ -290,14 +292,29 @@ struct FScopedLoadTimeAccumulatorTimer : public FScopedDurationTimer
 
 #define CUSTOM_LOADTIMER_LOG Cpu
 
+// Uses trace system that can be read by Insights
 #if LOADTIMEPROFILERTRACE_ENABLED
+
+// Writes any string to the LoadTime channel, normally used for class names 
 #define SCOPED_LOADTIMER_TEXT(TimerName) TRACE_CPUPROFILER_EVENT_SCOPE_TEXT_ON_CHANNEL(TimerName, LoadTimeChannel)
+
+// Writes any string to the AssetLoadTime channel, for full asset paths
+#define SCOPED_LOADTIMER_ASSET_TEXT(TimerName) TRACE_CPUPROFILER_EVENT_SCOPE_TEXT_ON_CHANNEL(TimerName, AssetLoadTimeChannel)
+
+// Writes raw scope name to LoadTime channel
 #define SCOPED_LOADTIMER(TimerName) TRACE_CPUPROFILER_EVENT_SCOPE_ON_CHANNEL(TimerName, LoadTimeChannel)
+
+// Used to create a custom trace event and add metadata
 #define SCOPED_CUSTOM_LOADTIMER(TimerName) UE_TRACE_LOG_SCOPED_T(CUSTOM_LOADTIMER_LOG, TimerName, LoadTimeChannel)
 #define ADD_CUSTOM_LOADTIMER_META(TimerName, Key, Value) << TimerName.Key(Value)
+
+// Increment cumulative event count, disabled in this mode
 #define SCOPED_LOADTIMER_CNT(TimerName)
+
+// All load time tracking is disabled
 #else
 #define SCOPED_LOADTIMER_TEXT(TimerName)
+#define SCOPED_LOADTIMER_ASSET_TEXT(TimerName)
 #define SCOPED_LOADTIMER(TimerName)
 #define SCOPED_CUSTOM_LOADTIMER(TimerName)
 #define ADD_CUSTOM_LOADTIMER_META(TimerName, Key, Value)
diff --git a/Engine/Source/Runtime/CoreUObject/Private/Serialization/AsyncLoading2.cpp b/Engine/Source/Runtime/CoreUObject/Private/Serialization/AsyncLoading2.cpp
--- a/Engine/Source/Runtime/CoreUObject/Private/Serialization/AsyncLoading2.cpp
+++ b/Engine/Source/Runtime/CoreUObject/Private/Serialization/AsyncLoading2.cpp
@@ -5339,6 +5339,7 @@ void FAsyncPackage2::StartLoading(FAsyncLoadingThreadState2& ThreadState, FIoBat
 EEventLoadNodeExecutionResult FAsyncPackage2::ProcessLinkerLoadPackageSummary(FAsyncLoadingThreadState2& ThreadState)
 {
 	TRACE_CPUPROFILER_EVENT_SCOPE(ProcessLinkerLoadPackageSummary);
+	SCOPED_LOADTIMER_ASSET_TEXT(*LinkerLoadState->Linker->GetDebugName());
 
 #if ALT2_ENABLE_NEW_ARCHIVE_FOR_LINKERLOAD
 	FLinkerLoad::ELinkerStatus LinkerResult = FLinkerLoad::LINKER_Failed;
@@ -5349,6 +5350,7 @@ EEventLoadNodeExecutionResult FAsyncPackage2::ProcessLinkerLoadPackageSummary(FA
 		LinkerResult = LinkerLoadState->Linker->ProcessPackageSummary(nullptr);
 	}
 #else
+	TRACE_LOADTIME_PROCESS_SUMMARY_SCOPE(this);
 	FLinkerLoad::ELinkerStatus LinkerResult = LinkerLoadState->Linker->Tick(/* RemainingTimeLimit */ 0.0, /* bUseTimeLimit */ false, /* bUseFullTimeLimit */ false, nullptr);
 #endif
 	check(LinkerResult != FLinkerLoad::LINKER_TimedOut); // TODO: Add support for timeouts here
@@ -5703,6 +5705,7 @@ EEventLoadNodeExecutionResult FAsyncPackage2::ExecutePostLoadLinkerLoadPackageEx
 {
 	if (!bLoadHasFailed)
 	{
+		SCOPED_LOADTIMER(PostLoadObjectsTime);
 		TRACE_LOADTIME_POSTLOAD_SCOPE;
 
 		FAsyncLoadingTickScope2 InAsyncLoadingTick(AsyncLoadingThread);
@@ -5720,8 +5723,6 @@ EEventLoadNodeExecutionResult FAsyncPackage2::ExecutePostLoadLinkerLoadPackageEx
 		// End async loading, simulates EndLoad
 		ON_SCOPE_EXIT { ThreadObjLoaded.Reset(); EndAsyncLoad(); };
 
-		SCOPED_LOADTIMER(PostLoadObjectsTime);
-
 		FUObjectThreadContext& ThreadContext = FUObjectThreadContext::Get();
 		TGuardValue<bool> GuardIsRoutingPostLoad(ThreadContext.IsRoutingPostLoad, true);
 
@@ -5738,6 +5739,9 @@ EEventLoadNodeExecutionResult FAsyncPackage2::ExecutePostLoadLinkerLoadPackageEx
 			{
 				if (Object->HasAnyFlags(RF_NeedPostLoad) && CanPostLoadOnAsyncLoadingThread(Object) && Object->IsReadyForAsyncPostLoad())
 				{
+#if WITH_EDITOR
+					SCOPED_LOADTIMER_ASSET_TEXT(*Object->GetPathName());
+#endif
 					ThreadContext.CurrentlyPostLoadedObjectByALT = Object;
 					Object->ConditionalPostLoad();
 					ThreadContext.CurrentlyPostLoadedObjectByALT = nullptr;
@@ -5764,6 +5768,7 @@ EEventLoadNodeExecutionResult FAsyncPackage2::ExecutePostLoadLinkerLoadPackageEx
 
 EEventLoadNodeExecutionResult FAsyncPackage2::ExecuteDeferredPostLoadLinkerLoadPackageExports(FAsyncLoadingThreadState2& ThreadState)
 {
+	SCOPED_LOADTIMER(PostLoadDeferredObjectsTime);
 	TRACE_LOADTIME_POSTLOAD_SCOPE;
 
 	FAsyncLoadingTickScope2 InAsyncLoadingTick(AsyncLoadingThread);
@@ -5792,7 +5797,9 @@ EEventLoadNodeExecutionResult FAsyncPackage2::ExecuteDeferredPostLoadLinkerLoadP
 				--LinkerLoadState->PostLoadExportIndex;
 				return EEventLoadNodeExecutionResult::Timeout;
 			}
-
+#if WITH_EDITOR
+			SCOPED_LOADTIMER_ASSET_TEXT(*Object->GetPathName());
+#endif
 			Object->ConditionalPostLoad();
 		}
 
@@ -7174,6 +7181,7 @@ EEventLoadNodeExecutionResult FAsyncPackage2::Event_PostLoadExportBundle(FAsyncL
 		Package->BeginAsyncLoad();
 
 		SCOPED_LOADTIMER(PostLoadObjectsTime);
+		TRACE_LOADTIME_POSTLOAD_SCOPE;
 
 		FUObjectThreadContext& ThreadContext = FUObjectThreadContext::Get();
 		TGuardValue<bool> GuardIsRoutingPostLoad(ThreadContext.IsRoutingPostLoad, true);
@@ -7200,7 +7208,6 @@ EEventLoadNodeExecutionResult FAsyncPackage2::Event_PostLoadExportBundle(FAsyncL
 		const FAsyncPackageHeaderData* HeaderData = &Package->HeaderData;
 #endif
 
-		TRACE_LOADTIME_POSTLOAD_SCOPE;
 		while (Package->ExportBundleEntryIndex < HeaderData->ExportBundleEntriesCopyForPostLoad.Num())
 		{
 			const FExportBundleEntry& BundleEntry = HeaderData->ExportBundleEntriesCopyForPostLoad[Package->ExportBundleEntryIndex];
@@ -7231,6 +7238,9 @@ EEventLoadNodeExecutionResult FAsyncPackage2::Event_PostLoadExportBundle(FAsyncL
 					check(Object->IsReadyForAsyncPostLoad());
 					if (!bIsMultithreaded || (bAsyncPostLoadEnabled && CanPostLoadOnAsyncLoadingThread(Object)))
 					{
+#if WITH_EDITOR
+						SCOPED_LOADTIMER_ASSET_TEXT(*Object->GetPathName());
+#endif
 						ThreadContext.CurrentlyPostLoadedObjectByALT = Object;
 						Object->ConditionalPostLoad();
 						ThreadContext.CurrentlyPostLoadedObjectByALT = nullptr;
@@ -7302,6 +7312,9 @@ EEventLoadNodeExecutionResult FAsyncPackage2::Event_DeferredPostLoadExportBundle
 	}
 	else
 	{
+		SCOPED_LOADTIMER(PostLoadDeferredObjectsTime);
+		TRACE_LOADTIME_POSTLOAD_SCOPE;
+
 		FAsyncLoadingTickScope2 InAsyncLoadingTick(Package->AsyncLoadingThread);
 
 #if WITH_EDITOR
@@ -7322,7 +7335,6 @@ EEventLoadNodeExecutionResult FAsyncPackage2::Event_DeferredPostLoadExportBundle
 		const FAsyncPackageHeaderData* HeaderData = &Package->HeaderData;
 #endif
 
-		TRACE_LOADTIME_POSTLOAD_SCOPE;
 		while (Package->ExportBundleEntryIndex < HeaderData->ExportBundleEntriesCopyForPostLoad.Num())
 		{
 			const FExportBundleEntry& BundleEntry = HeaderData->ExportBundleEntriesCopyForPostLoad[Package->ExportBundleEntryIndex];
@@ -7347,6 +7359,9 @@ EEventLoadNodeExecutionResult FAsyncPackage2::Event_DeferredPostLoadExportBundle
 					check(!Object->HasAnyFlags(RF_NeedLoad));
 					if (Object->HasAnyFlags(RF_NeedPostLoad))
 					{
+#if WITH_EDITOR
+						SCOPED_LOADTIMER_ASSET_TEXT(*Object->GetPathName());
+#endif
 						PackageScope.ThreadContext.CurrentlyPostLoadedObjectByALT = Object;
 						{
 							FScopeCycleCounterUObject ConstructorScope(Object, GET_STATID(STAT_FAsyncPackage_PostLoadObjectsGameThread));
