Fix missing copyright
#rnx

diff --git a/Engine/Source/Editor/Kismet/Internal/Reflection/FunctionUtils.h b/Engine/Source/Editor/Kismet/Internal/Reflection/FunctionUtils.h
--- a/Engine/Source/Editor/Kismet/Internal/Reflection/FunctionUtils.h
+++ b/Engine/Source/Editor/Kismet/Internal/Reflection/FunctionUtils.h
@@ -1,3 +1,5 @@
+// Copyright Epic Games, Inc. All Rights Reserved.
+
 #pragma once
 
 #include "FunctionUtilsPrivate.h"
diff --git a/Engine/Source/Editor/Kismet/Internal/Reflection/FunctionUtilsPrivate.h b/Engine/Source/Editor/Kismet/Internal/Reflection/FunctionUtilsPrivate.h
--- a/Engine/Source/Editor/Kismet/Internal/Reflection/FunctionUtilsPrivate.h
+++ b/Engine/Source/Editor/Kismet/Internal/Reflection/FunctionUtilsPrivate.h
@@ -1,3 +1,5 @@
+// Copyright Epic Games, Inc. All Rights Reserved.
+
 #pragma once
 
 #include "UObject/UnrealType.h"
