Reinstanciation fixes to have the new Overridable Serialization work correctly
PS: some new feature are only enabled on object that support overridable serialization but I wil look at eventually to merger the new code path at large.

diff --git a/Engine/Source/Editor/Kismet/Private/BlueprintCompilationManager.cpp b/Engine/Source/Editor/Kismet/Private/BlueprintCompilationManager.cpp
--- a/Engine/Source/Editor/Kismet/Private/BlueprintCompilationManager.cpp
+++ b/Engine/Source/Editor/Kismet/Private/BlueprintCompilationManager.cpp
@@ -46,6 +46,7 @@
 #include "Animation/AnimBlueprint.h"
 #include "Stats/StatsHierarchical.h"
 #include "UObject/PropertyBagRepository.h"
+#include "UObject/OverridableManager.h"
 
 extern UNREALED_API UUnrealEdEngine* GUnrealEd;
 
@@ -558,17 +559,26 @@ void FReinstancingJob::CalculateBPClassDependencies()
 		const UObject* OldCDO = OldToNew.Key->ClassDefaultObject;
 
 		// Gather subobjects on old CDO and remember depends BP classes
-		TArray<UObject*> ContainedOldObjects;
-		GetObjectsWithOuter(OldCDO, ContainedOldObjects);
-		for (const UObject* OldObject : ContainedOldObjects)
+		auto GatherDependentBPClasses = [this](const UObject* CDO, const auto& Recurse) -> void
 		{
-			const UClass* DependentClass = OldObject->GetClass();
-			while (DependentClass && UBlueprint::GetBlueprintFromClass(DependentClass))
+			TArray<UObject*> ContainedOldObjects;
+			GetObjectsWithOuter(CDO, ContainedOldObjects);
+			for (const UObject* OldObject : ContainedOldObjects)
 			{
-				BPClassDependencies.Add(DependentClass);
-				DependentClass = DependentClass->GetSuperClass();
+				const UClass* DependentClass = OldObject->GetClass();
+				while (DependentClass && UBlueprint::GetBlueprintFromClass(DependentClass))
+				{
+					bool bIsAlreadyInSet = false;
+					BPClassDependencies.Add(DependentClass, &bIsAlreadyInSet);
+					if(!bIsAlreadyInSet)
+					{
+						Recurse(DependentClass->ClassDefaultObject, Recurse);
+					}
+					DependentClass = DependentClass->GetSuperClass();
+				}
 			}
-		}
+		};
+		GatherDependentBPClasses(OldCDO, GatherDependentBPClasses);
 	}
 }
 
@@ -2402,11 +2412,17 @@ void FBlueprintCompilationManagerImpl::ReinstanceBatch(TArray<FReinstancingJob>&
 	Reinstancers.Sort(
 		[](const FReinstancingJob& ReinstancingDataA, const FReinstancingJob& ReinstancingDataB)
 		{
-			return !ReinstancingDataA.DependsOn(ReinstancingDataB) && 
-				FBlueprintCompileReinstancer::ReinstancerOrderingFunction(
+			if(ReinstancingDataA.DependsOn(ReinstancingDataB))
+			{
+				return false;
+			}
+			if(ReinstancingDataB.DependsOn(ReinstancingDataA))
+			{
+				return true;
+			}
+			return FBlueprintCompileReinstancer::ReinstancerOrderingFunction(
 					ReinstancingDataA.OldToNew.Value, 
-					ReinstancingDataB.OldToNew.Value
-				);
+					ReinstancingDataB.OldToNew.Value);
 		}
 	);
 
@@ -2415,15 +2431,36 @@ void FBlueprintCompilationManagerImpl::ReinstanceBatch(TArray<FReinstancingJob>&
 	for (const FReinstancingJob& ReinstancingJob : Reinstancers)
 	{
 		UObject* OldCDO = nullptr;
-		if (ReinstancingJob.OldToNew.Key)
+		UClass* OldClass = ReinstancingJob.OldToNew.Key;
+
+		if (OldClass)
 		{
-			OldCDO = ReinstancingJob.OldToNew.Key->ClassDefaultObject;
+			OldCDO = OldClass->ClassDefaultObject;
 			if (OldCDO && ReinstancingJob.Reinstancer.IsValid())
 			{
-				const bool bUseDeltaSerialization = ReinstancingJob.Reinstancer.IsValid() ? ReinstancingJob.Reinstancer->bUseDeltaSerializationToCopyProperties : false;
-				UObject* NewCDO = ReinstancingJob.OldToNew.Value->GetDefaultObject(true);
-
-				UBlueprint* CompiledBlueprint = UBlueprint::GetBlueprintFromClass(ReinstancingJob.OldToNew.Key);
+				// Object using overridable serialization need to use delta serialization for it do work appropriatly
+				// Eventually we should do this to all BP classes CDO no matter what
+				const bool bUsingOverrideSerialization = FOverridableManager::Get().IsEnabled(*OldCDO);
+				const bool bUseDeltaSerialization = bUsingOverrideSerialization ? true : (ReinstancingJob.Reinstancer.IsValid() ? ReinstancingJob.Reinstancer->bUseDeltaSerializationToCopyProperties : false);
+				UClass* NewClass = ReinstancingJob.OldToNew.Value;
+
+				// We do not expect the CDO to be already created at this point. 
+				// Who ever bother to create it before didn't look if the parent was ready or not.
+				if (NewClass->ClassDefaultObject != nullptr && bUsingOverrideSerialization)
+				{
+					// Override serialization does not currently want to use any CDO created before this
+					// point. I do not have a test case justifying this need, but renaming the cdo here 
+					// causes various regressions for existing blueprints, because their compilers have put
+					// needed data on the generated CDO. If we want to use override serialization broadly
+					// we will have to address these shortcomings (likely by ordering compilation itself,
+					// rather than reinstancing, more carefully). Discarding the CDO here is also wasteful
+					// - if we cannot use the objects, why create them at all?
+					NewClass->ClassDefaultObject->Rename(nullptr, GetTransientPackage(), REN_DoNotDirty | REN_DontCreateRedirectors | REN_ForceNoResetLoaders | REN_NonTransactional);
+					NewClass->ClassDefaultObject = nullptr;
+				}
+				UObject* NewCDO = NewClass->GetDefaultObject(true);
+
+				UBlueprint* CompiledBlueprint = UBlueprint::GetBlueprintFromClass(OldClass);
 				if (CompiledBlueprint && CompiledBlueprint->bIsRegeneratingOnLoad)
 				{
 					// This is a catch-all for any deferred dependencies that didn't get resolved during loading/linking (that system is not bulletproof for complex circular dependencies)
@@ -2435,6 +2472,10 @@ void FBlueprintCompilationManagerImpl::ReinstanceBatch(TArray<FReinstancingJob>&
 
 				// We only need to copy properties of the pre-created instances, the rest of the default sub object is done inside the UEditorEngine::CopyPropertiesForUnrelatedObjects
 				TMap<UObject*, UObject*> OldToNewInstanceMap(CreatedInstanceMap);
+				if (TMap<UObject*, UObject*>* OldToNewTemplateMap = OldToNewTemplates ? OldToNewTemplates->Find(OldClass->GetSuperClass()) : nullptr)
+				{
+					OldToNewInstanceMap.Append(*OldToNewTemplateMap);
+				}
 				for(const auto& Pair : CreatedInstanceMap)
 				{
 					FBlueprintCompileReinstancer::CopyPropertiesForUnrelatedObjects(Pair.Key, Pair.Value, /*bClearExternalReferences*/true, bUseDeltaSerialization, /*bOnlyHandleDirectSubObjects*/true, &OldToNewInstanceMap);
@@ -2442,7 +2483,7 @@ void FBlueprintCompilationManagerImpl::ReinstanceBatch(TArray<FReinstancingJob>&
 
 				if (OldToNewTemplates && !OldToNewInstanceMap.IsEmpty())
 				{
-					OldToNewTemplates->FindOrAdd(ReinstancingJob.OldToNew.Key).Append(OldToNewInstanceMap);
+					OldToNewTemplates->FindOrAdd(OldClass).Append(OldToNewInstanceMap);
 				}
 
 				if (ReinstancingJob.Compiler.IsValid())
@@ -2642,6 +2683,10 @@ void FBlueprintCompilationManagerImpl::ReinstanceBatch(TArray<FReinstancingJob>&
 
 					// We only need to copy properties of the pre-created instances, the rest of the default sub object is done inside the UEditorEngine::CopyPropertiesForUnrelatedObjects
 					TMap<UObject*, UObject*> OldToNewInstanceMap(CreatedInstanceMap);
+					if (TMap<UObject*, UObject*>* OldToNewTemplateMap = OldToNewTemplates ? OldToNewTemplates->Find(OldClass->GetSuperClass()) : nullptr)
+					{
+						OldToNewInstanceMap.Append(*OldToNewTemplateMap);
+					}
 					for (const auto& Pair : CreatedInstanceMap)
 					{
 						FBlueprintCompileReinstancer::CopyPropertiesForUnrelatedObjects(Pair.Key, Pair.Value, /*bClearExternalReferences*/true, bUseDeltaSerialization, /*bOnlyHandleDirectSubObjects*/true, &OldToNewInstanceMap);
diff --git a/Engine/Source/Editor/UnrealEd/Private/Kismet2/KismetReinstanceUtilities.cpp b/Engine/Source/Editor/UnrealEd/Private/Kismet2/KismetReinstanceUtilities.cpp
--- a/Engine/Source/Editor/UnrealEd/Private/Kismet2/KismetReinstanceUtilities.cpp
+++ b/Engine/Source/Editor/UnrealEd/Private/Kismet2/KismetReinstanceUtilities.cpp
@@ -40,6 +40,7 @@
 #include "Subsystems/AssetEditorSubsystem.h"
 #include "Engine/ScopedMovementUpdate.h"
 #include "InstancedReferenceSubobjectHelper.h"
+#include "UObject/OverridableManager.h"
 #include "UObject/PropertyOptional.h"
 #include "UObject/PropertyBagRepository.h"
 #include "ProfilingDebugging/LoadTimeTracker.h"
@@ -2171,6 +2172,7 @@ static void ReplaceObjectHelper(UObject*& OldObject, UClass* OldClass, UObject*&
 	Options.bNotifyObjectReplacement = true;
 	Options.bSkipCompilerGeneratedDefaults = true;
 	Options.bOnlyHandleDirectSubObjects = true;
+	Options.bReplaceInternalReferenceUponRead = FOverridableManager::Get().IsEnabled(*OldObject);
 	Options.OptionalReplacementMappings = &OldToNewInstanceMap;
 	// this currently happens because of some misguided logic in UBlueprintGeneratedClass::FindArchetype that
 	// points us to a mismatched archetype, in which case delta serialization becomes unsafe.. without
@@ -2578,6 +2580,11 @@ void FBlueprintCompileReinstancer::ReplaceInstancesOfClass_Inner(const TMap<UCla
 					bFixupSCS = (NewClass->IsChildOf<USceneComponent>() != OldClass->IsChildOf<USceneComponent>());
 				}
 
+				if (TMap<UObject*, UObject*>* ReplaceTemplateMapping = Params.OldToNewTemplates ? Params.OldToNewTemplates->Find(OldClass) : nullptr)
+				{
+					OldToNewInstanceMap.Append(*ReplaceTemplateMapping);
+				}
+
 				const bool bIncludeDerivedClasses = false;
 				ObjectsToReplace.Reset();
 				GetObjectsOfClass(OldClass, ObjectsToReplace, bIncludeDerivedClasses);
@@ -2665,6 +2672,11 @@ void FBlueprintCompileReinstancer::ReplaceInstancesOfClass_Inner(const TMap<UCla
 					}
 				}
 
+				if (TMap<UObject*, UObject*>* ReplaceTemplateMapping = Params.OldToNewTemplates ? Params.OldToNewTemplates->Find(OldClass) : nullptr)
+				{
+					OldToNewInstanceMap.Append(*ReplaceTemplateMapping);
+				}
+
 				// Then fix 'real' (non archetype) instances of the class
 				for (int32 OldObjIndex = 0; OldObjIndex < ObjectsToReplace.Num(); ++OldObjIndex)
 				{
@@ -3054,6 +3066,9 @@ void FBlueprintCompileReinstancer::CopyPropertiesForUnrelatedObjects(UObject* Ol
 	Params.bNotifyObjectReplacement = true;
 	Params.OptionalReplacementMappings = OldToNewInstanceMap;
 	Params.bOnlyHandleDirectSubObjects = bOnlyHandleDirectSubObjects;
+	// Overridable serialization needs this to be able to merge back containers of subobjects.
+	Params.bReplaceInternalReferenceUponRead = FOverridableManager::Get().IsEnabled(*OldObject);
+
 	UEngine::CopyPropertiesForUnrelatedObjects(OldObject, NewObject, Params);
 }
 
diff --git a/Engine/Source/Editor/UnrealEd/Public/Kismet2/KismetReinstanceUtilities.h b/Engine/Source/Editor/UnrealEd/Public/Kismet2/KismetReinstanceUtilities.h
--- a/Engine/Source/Editor/UnrealEd/Public/Kismet2/KismetReinstanceUtilities.h
+++ b/Engine/Source/Editor/UnrealEd/Public/Kismet2/KismetReinstanceUtilities.h
@@ -285,6 +285,8 @@ public:
 	* @param NewObject						The new Object to copy properties to
 	* @param bClearExternalReferences		If true then attempt to replace references to old classes and instances on this object with the corresponding new ones
 	* @param bForceDeltaSerialization		If true the delta serialization will be used when copying
+	* @param bOnlyHandleDirectSubObjects	If true will only copy/handle immediate subobjects
+	* @param OldToNewInstanceMap			If != null it will be used to replace any references found in this object
 	*/
 	static UNREALED_API void CopyPropertiesForUnrelatedObjects(UObject* OldObject, UObject* NewObject, bool bClearExternalReferences, bool bForceDeltaSerialization = false, bool bOnlyHandleDirectSubObjects = false, TMap<UObject*, UObject*>* OldToNewInstanceMap =nullptr);
 
diff --git a/Engine/Source/Runtime/CoreUObject/Private/UObject/CoreNative.cpp b/Engine/Source/Runtime/CoreUObject/Private/UObject/CoreNative.cpp
--- a/Engine/Source/Runtime/CoreUObject/Private/UObject/CoreNative.cpp
+++ b/Engine/Source/Runtime/CoreUObject/Private/UObject/CoreNative.cpp
@@ -9,6 +9,7 @@
 #include "Misc/PackageName.h"
 #include "Misc/RuntimeErrors.h"
 #include "UObject/Stack.h"
+#include "UObject/OverridableManager.h"
 
 void UClassRegisterAllCompiledInClasses();
 bool IsInAsyncLoadingThreadCoreUObjectInternal();
@@ -165,6 +166,12 @@ UObject* FObjectInstancingGraph::GetInstancedSubobject( UObject* SourceSubobject
 
 		if ( bShouldInstance )
 		{
+			// If the CurrentValue is within the SourceRoot, lets use it to instantiate as it must have come from the merge result of the serialization
+			if (FOverridableManager::Get().IsEnabled(*SourceSubobject) && SourceSubobject != CurrentValue && CurrentValue->IsIn(SourceRoot))
+			{
+				SourceSubobject = CurrentValue;
+			}
+
 			// search for the unique subobject instance that corresponds to this subobject template
 			InstancedSubobject = GetDestinationObject(SourceSubobject);
 			if ( InstancedSubobject == nullptr )
diff --git a/Engine/Source/Runtime/Engine/Classes/Engine/Engine.h b/Engine/Source/Runtime/Engine/Classes/Engine/Engine.h
--- a/Engine/Source/Runtime/Engine/Classes/Engine/Engine.h
+++ b/Engine/Source/Runtime/Engine/Classes/Engine/Engine.h
@@ -2955,7 +2955,9 @@ public:
 		bool bSkipCompilerGeneratedDefaults;
 		bool bNotifyObjectReplacement;
 		bool bClearReferences;
+		UE_DEPRECATED(5.4, "This isn't used anymore by the code.")
 		bool bDontClearReferenceIfNewerClassExists;
+		bool bReplaceInternalReferenceUponRead; // While reading back object ptr, immediately replace them if they are in the replacement map.
 
 		// In cases where the SourceObject will no longer be able to look up its correct Archetype, it can be supplied
 		UObject* SourceObjectArchetype;
diff --git a/Engine/Source/Runtime/Engine/Private/UnrealEngine.cpp b/Engine/Source/Runtime/Engine/Private/UnrealEngine.cpp
--- a/Engine/Source/Runtime/Engine/Private/UnrealEngine.cpp
+++ b/Engine/Source/Runtime/Engine/Private/UnrealEngine.cpp
@@ -16900,9 +16900,10 @@ private:
 struct FCPFUOReader : public FObjectReader, public FCPFUOArchive
 {
 public:
-	FCPFUOReader(FCPFUOWriter& DataSrc, UObject* DstObject)
+	FCPFUOReader(FCPFUOWriter& DataSrc, UObject* DstObject, TMap<UObject*, UObject*>* InReferenceReplacementMap = nullptr)
 		: FObjectReader(DataSrc.SavedPropertyData)
 		, FCPFUOArchive(DataSrc)
+		, ReferenceReplacementMap(InReferenceReplacementMap)
 	{
 		ArIgnoreArchetypeRef = true;
 		ArIgnoreClassRef = true;
@@ -16929,6 +16930,86 @@ public:
 	virtual void MarkScriptSerializationStart(const UObject* Object) override { OpenTaggedDataScope(); }
 	virtual void MarkScriptSerializationEnd(const UObject* Object) override   { CloseTaggedDataScope(); }
 	// ~End FArchive Interface
+
+	TMap<UObject*, UObject*>* ReferenceReplacementMap = nullptr;
+
+	FArchive& operator<<(FWeakObjectPtr& Value) override
+	{
+		FObjectReader::operator<<(Value);
+		if (UObject* Object = Value.Get())
+		{
+			if (UObject** Reference = ReferenceReplacementMap ? ReferenceReplacementMap->Find(Object) : nullptr)
+			{
+				Value = *Reference;
+			}
+		}
+		return *this;
+	}
+
+	FArchive& operator<<(FSoftObjectPtr& Value) override
+	{
+		FObjectReader::operator<<(Value);
+		if (UObject* Object = Value.Get())
+		{
+			if (UObject** Reference = ReferenceReplacementMap ? ReferenceReplacementMap->Find(Object) : nullptr)
+			{
+				Value = *Reference;
+			}
+		}
+		return *this;
+	}
+
+	FArchive& operator<<(FSoftObjectPath& Value) override
+	{
+		FObjectReader::operator<<(Value);
+		if (UObject* Object = Value.ResolveObject())
+		{
+			if (UObject** Reference = ReferenceReplacementMap ? ReferenceReplacementMap->Find(Object) : nullptr)
+			{
+				Value = *Reference;
+			}
+		}
+		return *this;
+	}
+
+	FArchive& operator<<(FLazyObjectPtr& Value) override
+	{
+		FObjectReader::operator<<(Value);
+		if (UObject* Object = Value.Get())
+		{
+			if (UObject** Reference = ReferenceReplacementMap ? ReferenceReplacementMap->Find(Object) : nullptr)
+			{
+				Value = *Reference;
+			}
+		}
+		return *this;
+	}
+
+	FArchive& operator<<(FObjectPtr& Value) override
+	{
+		FObjectReader::operator<<(Value);
+		if (UObject* Object = Value.Get())
+		{
+			if (UObject** Reference = ReferenceReplacementMap ? ReferenceReplacementMap->Find(Object) : nullptr)
+			{
+				Value = *Reference;
+			}
+		}
+		return *this;
+	}
+
+	FArchive& operator<<(UObject*& Value) override
+	{
+		FObjectReader::operator<<(Value);
+		if (UObject* Object = Value)
+		{
+			if (UObject** Reference = ReferenceReplacementMap ? ReferenceReplacementMap->Find(Object) : nullptr)
+			{
+				Value = *Reference;
+			}
+		}
+		return *this;
+	}
 };
 
 PRAGMA_DISABLE_DEPRECATION_WARNINGS
@@ -16943,6 +17024,8 @@ UEngine::FCopyPropertiesForUnrelatedObjectsParams::FCopyPropertiesForUnrelatedOb
 	, bSkipCompilerGeneratedDefaults(false)
 	, bNotifyObjectReplacement(false)
 	, bClearReferences(true)
+	, bDontClearReferenceIfNewerClassExists(false)
+	, bReplaceInternalReferenceUponRead(false)
 	, SourceObjectArchetype(nullptr)
 	, OptionalReplacementMappings(nullptr)
 {}
@@ -17103,12 +17186,6 @@ void UEngine::CopyPropertiesForUnrelatedObjects(UObject* OldObject, UObject* New
 	const bool bDoAggressiveReplacement = CVarUseAggressiveReferenceReplacement.GetValueOnAnyThread();
 
 	{
-		// Serialize in the modified properties from the old CDO to the new CDO
-		if (Writer.SavedPropertyData.Num() > 0)
-		{
-			FCPFUOReader Reader(Writer, NewObject);
-		}
-
 		CollectAllSubobjects( NewObject, ComponentsOnNewObject );
 
 		// populate the ReferenceReplacementMap 
@@ -17181,6 +17258,19 @@ void UEngine::CopyPropertiesForUnrelatedObjects(UObject* OldObject, UObject* New
 			}
 		}
 
+		// Serialize in the modified properties from the old CDO to the new CDO
+		if (Writer.SavedPropertyData.Num() > 0)
+		{
+			if (Params.bReplaceInternalReferenceUponRead)
+			{
+				FCPFUOReader Reader(Writer, NewObject, Params.OptionalReplacementMappings ? Params.OptionalReplacementMappings : &ReferenceReplacementMap);
+			}
+			else
+			{
+				FCPFUOReader Reader(Writer, NewObject);
+			}
+		}
+
 		TArray<UObject*> EditInlineSubobjectsOfComponents;
 		for (UObject* NewInstance : ComponentsOnNewObject)
 		{
