Fixed inconsistent tooltip for Blueprint config variables.

GetConfigName doesn't always return the base name of an .ini file, which is what we want. For example, if the class uses DefaultEngine.ini for its config values, GetConfigName will return the base name, "Engine". However, if the class uses DefaultEditor.ini, GetConfigName returns the path to the merged .ini file that's located under the game's "Saved/Config" directory.

Since ClassConfigName corresponds to the base name of a config file, we can just use that for the tooltip.

diff --git a/Engine/Source/Editor/Kismet/Private/BlueprintDetailsCustomization.cpp b/Engine/Source/Editor/Kismet/Private/BlueprintDetailsCustomization.cpp
--- a/Engine/Source/Editor/Kismet/Private/BlueprintDetailsCustomization.cpp
+++ b/Engine/Source/Editor/Kismet/Private/BlueprintDetailsCustomization.cpp
@@ -642,25 +642,24 @@ void FBlueprintVarActionDetails::CustomizeDetails( IDetailLayoutBuilder& DetailL
 		.ToolTip(ExposeToCinematicsTooltip)
 	];
 
-	FText LocalisedTooltip;
+	FText LocalizedTooltip;
 	if (IsConfigCheckBoxEnabled())
 	{
 		// Build the property specific config variable tool tip
 		FFormatNamedArguments ConfigTooltipArgs;
 		if (UClass* OwnerClass = VariableProperty->GetOwnerClass())
 		{
-			OwnerClass = OwnerClass->GetAuthoritativeClass();
-			ConfigTooltipArgs.Add(TEXT("ConfigPath"), FText::FromString(OwnerClass->GetConfigName()));
+			ConfigTooltipArgs.Add(TEXT("ConfigName"), FText::FromName(OwnerClass->ClassConfigName));
 			ConfigTooltipArgs.Add(TEXT("ConfigSection"), FText::FromString(OwnerClass->GetPathName()));
 		}
-		LocalisedTooltip = FText::Format(LOCTEXT("VariableExposeToConfig_Tooltip", "Should this variable read its default value from a config file if it is present?\r\n\r\nThis is used for customising variable default values and behavior between different projects and configurations.\r\n\r\nConfig file [{ConfigPath}]\r\nConfig section [{ConfigSection}]"), ConfigTooltipArgs);
+		LocalizedTooltip = FText::Format(LOCTEXT("VariableExposeToConfig_Tooltip", "Should this variable read its default value from a config file if it is present?\r\n\r\nThis is used for customizing variable default values and behavior between different projects and configurations.\r\n\r\nConfig file [{ConfigName}]\r\nConfig section [{ConfigSection}]"), ConfigTooltipArgs);
 	}
 	else if (IsVariableInBlueprint())
 	{
 		// mimics the error that UHT would throw
-		LocalisedTooltip = LOCTEXT("ObjectVariableConfig_Tooltip", "Not allowed to use 'config' with object variables");
+		LocalizedTooltip = LOCTEXT("ObjectVariableConfig_Tooltip", "Not allowed to use 'config' with object variables");
 	}
-	TSharedPtr<SToolTip> ExposeToConfigTooltip = IDocumentation::Get()->CreateToolTip(LocalisedTooltip, NULL, DocLink, TEXT("ExposeToConfig"));
+	TSharedPtr<SToolTip> ExposeToConfigTooltip = IDocumentation::Get()->CreateToolTip(LocalizedTooltip, NULL, DocLink, TEXT("ExposeToConfig"));
 
 	Category.AddCustomRow( LOCTEXT("VariableExposeToConfig", "Config Variable"), true )
 	.Visibility(TAttribute<EVisibility>(this, &FBlueprintVarActionDetails::ExposeConfigVisibility))
