Allow users to change Class Settings in the Blueprint editor's defaults-only (DOBP) mode.

Change summary:
- Added FBlueprintEditor::bWasOpenedInDefaultsMode.
- Modified FBlueprintEditor::CreateDefaultTabContents() to configure DefaultEditor based on defaults-only mode (to avoid breaking editors that utilize it alongside the Inspector - e.g. AnimBP).
- Modified FBlueprintEditor::RefreshStandAloneDefaultsEditor() to select BP or CDO based on current UI selection.
- Modified FBlueprintEditor::EditGlobalOptions_Clicked() to refresh the DefaultEditor in defaults-only mode.
- Modified FBlueprintEditor::EditClassDefaults_Clicked() to decouple it from single-BP-only edit mode.
- Modified FBlueprintEditor::IsDetailsPanelEditingClassDefaults() to consider defaults-only mode.
- Switched SKismetInspector::ShowPublicViewControl to be a bindable attribute.
- Modified SKismetInspector::MakeContextualEditingWidget() to bind visibility to the Public View.
- Modified FBlueprintDefaultsApplicationMode to show the Global Options toolbar in defaults-only mode.
- Modified FBlueprintGlobalOptionsDetails::CustomizeDetails() to hide import/interface sections from Class Settings in defaults-only mode.

diff --git a/Engine/Source/Editor/Kismet/Private/BlueprintDetailsCustomization.cpp b/Engine/Source/Editor/Kismet/Private/BlueprintDetailsCustomization.cpp
--- a/Engine/Source/Editor/Kismet/Private/BlueprintDetailsCustomization.cpp
+++ b/Engine/Source/Editor/Kismet/Private/BlueprintDetailsCustomization.cpp
@@ -4,6 +4,7 @@
 
 #include "AssetRegistry/AssetData.h"
 #include "BlueprintEditor.h"
+#include "BlueprintEditorModes.h"
 #include "BlueprintEditorModule.h"
 #include "BlueprintEditorSettings.h"
 #include "BlueprintNamespaceRegistry.h"
@@ -6890,44 +6891,50 @@ void FBlueprintGlobalOptionsDetails::CustomizeDetails(IDetailLayoutBuilder& Deta
 			]
 		]
 		.PropertyHandleList({ParentClassProperty});
-		
+
 		const bool bIsInterfaceBP = FBlueprintEditorUtils::IsInterfaceBlueprint(Blueprint);
 		const bool bIsMacroLibrary = Blueprint->BlueprintType == BPTYPE_MacroLibrary;
 		const bool bIsLevelScriptBP = FBlueprintEditorUtils::IsLevelScriptBlueprint(Blueprint);
 		const bool bIsFunctionLibrary = Blueprint->BlueprintType == BPTYPE_FunctionLibrary;
-		const bool bSupportsInterfaces = !bIsInterfaceBP && !bIsMacroLibrary && !bIsFunctionLibrary;
-		const bool bSupportsNamespaces = GetDefault<UBlueprintEditorSettings>()->bEnableNamespaceImportingFeatures;
-
-		if(bSupportsNamespaces)
+		
+		// Interfaces/imports currently rely on the full Blueprint editor context to function properly (e.g. add/remove operations).
+		TSharedPtr<FBlueprintEditor> PinnedBlueprintEditorPtr = BlueprintEditorPtr.Pin();
+		if (PinnedBlueprintEditorPtr.IsValid() && PinnedBlueprintEditorPtr->GetCurrentMode() != FBlueprintEditorApplicationModes::BlueprintDefaultsMode)
 		{
-			// Imported namespace details
-			IDetailCategoryBuilder& ImportsCategory = DetailLayout.EditCategory("Imports", LOCTEXT("BlueprintImportDetailsCategory", "Imports"));
+			const bool bSupportsInterfaces = !bIsInterfaceBP && !bIsMacroLibrary && !bIsFunctionLibrary;
+			const bool bSupportsNamespaces = GetDefault<UBlueprintEditorSettings>()->bEnableNamespaceImportingFeatures;
 
-			TSharedRef<FBlueprintImportsLayout> DefaultImportsLayout = MakeShareable(new FBlueprintImportsLayout(SharedThis(this), /*bShowDefaultImports = */true));
-			ImportsCategory.AddCustomBuilder(DefaultImportsLayout);
+			if (bSupportsNamespaces)
+			{
+				// Imported namespace details
+				IDetailCategoryBuilder& ImportsCategory = DetailLayout.EditCategory("Imports", LOCTEXT("BlueprintImportDetailsCategory", "Imports"));
 
-			TSharedRef<FBlueprintImportsLayout> LocalImportsLayout = MakeShareable(new FBlueprintImportsLayout(SharedThis(this), /*bShowDefaultImports = */false));
-			ImportsCategory.AddCustomBuilder(LocalImportsLayout);
-		}
+				TSharedRef<FBlueprintImportsLayout> DefaultImportsLayout = MakeShareable(new FBlueprintImportsLayout(SharedThis(this), /*bShowDefaultImports = */true));
+				ImportsCategory.AddCustomBuilder(DefaultImportsLayout);
 
-		if (bSupportsInterfaces)
-		{
-			// Interface details customization
-			IDetailCategoryBuilder& InterfacesCategory = DetailLayout.EditCategory("Interfaces", LOCTEXT("BlueprintInterfacesDetailsCategory", "Interfaces"));
+				TSharedRef<FBlueprintImportsLayout> LocalImportsLayout = MakeShareable(new FBlueprintImportsLayout(SharedThis(this), /*bShowDefaultImports = */false));
+				ImportsCategory.AddCustomBuilder(LocalImportsLayout);
+			}
 
-			// ImplementedInterfaces is a hidden property so we have to add it to the property map manually to use it
-			const TSharedPtr<IPropertyHandle> InterfacesProperty = DetailLayout.AddObjectPropertyData({const_cast<UBlueprint*>(Blueprint)}, TEXT("ImplementedInterfaces"));
-			
-			TSharedRef<FBlueprintInterfaceLayout> InheritedInterfacesLayout = MakeShareable(new FBlueprintInterfaceLayout(
-				SharedThis(this)
+			if (bSupportsInterfaces)
+			{
+				// Interface details customization
+				IDetailCategoryBuilder& InterfacesCategory = DetailLayout.EditCategory("Interfaces", LOCTEXT("BlueprintInterfacesDetailsCategory", "Interfaces"));
+
+				// ImplementedInterfaces is a hidden property so we have to add it to the property map manually to use it
+				const TSharedPtr<IPropertyHandle> InterfacesProperty = DetailLayout.AddObjectPropertyData({ const_cast<UBlueprint*>(Blueprint) }, TEXT("ImplementedInterfaces"));
+
+				TSharedRef<FBlueprintInterfaceLayout> InheritedInterfacesLayout = MakeShareable(new FBlueprintInterfaceLayout(
+					SharedThis(this)
 				));
-			InterfacesCategory.AddCustomBuilder(InheritedInterfacesLayout);
+				InterfacesCategory.AddCustomBuilder(InheritedInterfacesLayout);
 
-			TSharedRef<FBlueprintInterfaceLayout> LocalInterfacesLayout = MakeShareable(new FBlueprintInterfaceLayout(
-				SharedThis(this),
-				InterfacesProperty.ToSharedRef()
+				TSharedRef<FBlueprintInterfaceLayout> LocalInterfacesLayout = MakeShareable(new FBlueprintInterfaceLayout(
+					SharedThis(this),
+					InterfacesProperty.ToSharedRef()
 				));
-			InterfacesCategory.AddCustomBuilder(LocalInterfacesLayout);
+				InterfacesCategory.AddCustomBuilder(LocalInterfacesLayout);
+			}
 		}
 
 		// Hide the bDeprecate, we override the functionality.
diff --git a/Engine/Source/Editor/Kismet/Private/BlueprintEditor.cpp b/Engine/Source/Editor/Kismet/Private/BlueprintEditor.cpp
--- a/Engine/Source/Editor/Kismet/Private/BlueprintEditor.cpp
+++ b/Engine/Source/Editor/Kismet/Private/BlueprintEditor.cpp
@@ -2048,6 +2048,8 @@ void FBlueprintEditor::CommonInitialization(const TArray<UBlueprint*>& InitBluep
 		InitBlueprint->OnSetObjectBeingDebugged().AddSP(this, &FBlueprintEditor::HandleSetObjectBeingDebugged);
 	}
 
+	bWasOpenedInDefaultsMode = bShouldOpenInDefaultsMode;
+
 	CreateDefaultTabContents(InitBlueprints);
 
 	FCoreUObjectDelegates::OnPreObjectPropertyChanged.AddSP(this, &FBlueprintEditor::OnPreObjectPropertyChanged);
@@ -2962,15 +2964,15 @@ void FBlueprintEditor::CreateDefaultTabContents(const TArray<UBlueprint*>& InBlu
 
 	if ( InBlueprints.Num() > 0 )
 	{
-		const bool bShowPublicView = true;
-		const bool bHideNameArea = false;
+		// Don't show the object name in defaults mode.
+		const bool bHideNameArea = bWasOpenedInDefaultsMode;
 
 		this->DefaultEditor = 
 			SNew(SKismetInspector)
 			. Kismet2(SharedThis(this))
 			. ViewIdentifier(FName("BlueprintDefaults"))
 			. IsEnabled(!bIsInterface)
-			. ShowPublicViewControl(bShowPublicView)
+			. ShowPublicViewControl(this, &FBlueprintEditor::ShouldShowPublicViewControl)
 			. ShowTitleArea(false)
 			. HideNameArea(bHideNameArea)
 			. OnFinishedChangingProperties( FOnFinishedChangingProperties::FDelegate::CreateSP( this, &FBlueprintEditor::OnFinishedChangingProperties ) );
@@ -3555,20 +3557,32 @@ bool FBlueprintEditor::IsDetailsPanelEditingGlobalOptions() const
 
 void FBlueprintEditor::EditGlobalOptions_Clicked()
 {
-	UBlueprint* Blueprint = GetBlueprintObj();
-	if ( Blueprint != nullptr )
-	{
-		SetUISelectionState(FBlueprintEditor::SelectionState_ClassSettings);
+	SetUISelectionState(FBlueprintEditor::SelectionState_ClassSettings);
 
-		// Show details for the Blueprint instance we're editing
-		Inspector->ShowDetailsForSingleObject(Blueprint);
+	if (bWasOpenedInDefaultsMode)
+	{
+		RefreshStandAloneDefaultsEditor();
+	}
+	else
+	{
+		UBlueprint* Blueprint = GetBlueprintObj();
+		if (Blueprint != nullptr)
+		{
+			// Show details for the Blueprint instance we're editing
+			Inspector->ShowDetailsForSingleObject(Blueprint);
 
-		TryInvokingDetailsTab();
+			TryInvokingDetailsTab();
+		}
 	}
 }
 
 bool FBlueprintEditor::IsDetailsPanelEditingClassDefaults() const
 {
+	if (bWasOpenedInDefaultsMode)
+	{
+		return !IsDetailsPanelEditingGlobalOptions();
+	}
+
 	UBlueprint* Blueprint = GetBlueprintObj();
 	if ( Blueprint != nullptr )
 	{
@@ -3584,11 +3598,7 @@ bool FBlueprintEditor::IsDetailsPanelEditingClassDefaults() const
 
 void FBlueprintEditor::EditClassDefaults_Clicked()
 {
-	if ( IsEditingSingleBlueprint() )
-	{
-		UBlueprint* Blueprint = GetBlueprintObj();
-		StartEditingDefaults( true, true );
-	}
+	StartEditingDefaults(true, true);
 }
 
 // Zooming to fit the entire graph
@@ -9248,10 +9258,16 @@ void FBlueprintEditor::RefreshStandAloneDefaultsEditor()
 	TArray<UObject*> DefaultObjects;
 	for ( int32 i = 0; i < GetEditingObjects().Num(); ++i )
 	{
-		UBlueprintCore* Blueprint = Cast<UBlueprintCore>(GetEditingObjects()[i]);
-		if ( Blueprint && Blueprint->GeneratedClass )
+		if (UBlueprint* Blueprint = Cast<UBlueprint>(GetEditingObjects()[i]))
 		{
-			DefaultObjects.Add(Blueprint->GeneratedClass->GetDefaultObject());
+			if (CurrentUISelection == FBlueprintEditor::SelectionState_ClassSettings)
+			{
+				DefaultObjects.Add(Blueprint);
+			}
+			else if (Blueprint->GeneratedClass)
+			{
+				DefaultObjects.Add(Blueprint->GeneratedClass->GetDefaultObject());
+			}
 		}
 	}
 
@@ -9595,6 +9611,12 @@ void FBlueprintEditor::NotifyPostChange(const FPropertyChangedEvent& PropertyCha
 	}
 }
 
+bool FBlueprintEditor::ShouldShowPublicViewControl() const
+{
+	// In defaults-only mode, hide the "Public View" checkbox when Class Settings is selected into the Details view.
+	return !bWasOpenedInDefaultsMode || CurrentUISelection != FBlueprintEditor::SelectionState_ClassSettings;
+}
+
 void FBlueprintEditor::OnFinishedChangingProperties(const FPropertyChangedEvent& PropertyChangedEvent)
 {
 	FName PropertyName = PropertyChangedEvent.GetPropertyName();
diff --git a/Engine/Source/Editor/Kismet/Private/BlueprintEditorModes.cpp b/Engine/Source/Editor/Kismet/Private/BlueprintEditorModes.cpp
--- a/Engine/Source/Editor/Kismet/Private/BlueprintEditorModes.cpp
+++ b/Engine/Source/Editor/Kismet/Private/BlueprintEditorModes.cpp
@@ -181,6 +181,7 @@ FBlueprintDefaultsApplicationMode::FBlueprintDefaultsApplicationMode(TSharedPtr<
 	if (UToolMenu* Toolbar = InBlueprintEditor->RegisterModeToolbarIfUnregistered(GetModeName()))
 	{
 		InBlueprintEditor->GetToolbarBuilder()->AddCompileToolbar(Toolbar);
+		InBlueprintEditor->GetToolbarBuilder()->AddBlueprintGlobalOptionsToolbar(Toolbar);
 	}
 }
 
diff --git a/Engine/Source/Editor/Kismet/Private/SKismetInspector.cpp b/Engine/Source/Editor/Kismet/Private/SKismetInspector.cpp
--- a/Engine/Source/Editor/Kismet/Private/SKismetInspector.cpp
+++ b/Engine/Source/Editor/Kismet/Private/SKismetInspector.cpp
@@ -235,22 +235,23 @@ TSharedRef<SWidget> SKismetInspector::MakeContextualEditingWidget(struct FKismet
 			]
 		];
 
-		if (bShowPublicView)
-		{
-			ContextualEditingWidget->AddSlot()
-			.AutoHeight()
-			.VAlign( VAlign_Top )
+		ContextualEditingWidget->AddSlot()
+		.AutoHeight()
+		.VAlign(VAlign_Top)
+		[
+			SNew(SCheckBox)
+			.ToolTipText(LOCTEXT("TogglePublicView", "Toggle Public View"))
+			.IsChecked(this, &SKismetInspector::GetPublicViewCheckboxState)
+			.OnCheckStateChanged(this, &SKismetInspector::SetPublicViewCheckboxState)
 			[
-				SNew(SCheckBox)
-				.ToolTipText(LOCTEXT("TogglePublicView", "Toggle Public View"))
-				.IsChecked(this, &SKismetInspector::GetPublicViewCheckboxState)
-				.OnCheckStateChanged( this, &SKismetInspector::SetPublicViewCheckboxState)
-				[
-					SNew(STextBlock)
-					.Text(LOCTEXT("PublicViewCheckboxLabel", "Public View"))
-				]
-			];
-		}
+				SNew(STextBlock)
+				.Text(LOCTEXT("PublicViewCheckboxLabel", "Public View"))
+			]
+			.Visibility_Lambda([bShowPublicView = this->bShowPublicView]()
+			{
+				return bShowPublicView.Get() ? EVisibility::Visible : EVisibility::Hidden;
+			})
+		];
 	}
 
 	return ContextualEditingWidget;
diff --git a/Engine/Source/Editor/Kismet/Public/BlueprintEditor.h b/Engine/Source/Editor/Kismet/Public/BlueprintEditor.h
--- a/Engine/Source/Editor/Kismet/Public/BlueprintEditor.h
+++ b/Engine/Source/Editor/Kismet/Public/BlueprintEditor.h
@@ -1217,6 +1217,9 @@ protected:
 	virtual void NotifyPostChange( const FPropertyChangedEvent& PropertyChangedEvent, FProperty* PropertyThatChanged) override;
 	//~ End FNotifyHook Interface
 
+	/** Callback to determine visibility of the public view checkbox in the Defaults editor */
+	bool ShouldShowPublicViewControl() const;
+
 	/** Callback when properties have finished being handled */
 	virtual void OnFinishedChangingProperties(const FPropertyChangedEvent& PropertyChangedEvent);
 
@@ -1506,6 +1509,9 @@ protected:
 	/** If a regular node (not a comment node) has been selected */
 	bool bSelectRegularNode;
 
+	/** True if the editor was opened in defaults mode */
+	bool bWasOpenedInDefaultsMode;
+
 	/** Focus nodes which are related to the selected nodes */
 	void ResetAllNodesUnrelatedStates();
 	void CollectExecUpstreamNodes(UEdGraphNode* CurrentNode, TArray<UEdGraphNode*>& CollectedNodes);
diff --git a/Engine/Source/Editor/Kismet/Public/SKismetInspector.h b/Engine/Source/Editor/Kismet/Public/SKismetInspector.h
--- a/Engine/Source/Editor/Kismet/Public/SKismetInspector.h
+++ b/Engine/Source/Editor/Kismet/Public/SKismetInspector.h
@@ -54,7 +54,7 @@ public:
 
 		SLATE_ARGUMENT(TWeakPtr<FBlueprintEditor>, Kismet2)
 		SLATE_ARGUMENT(TWeakPtr<SMyBlueprint>, MyBlueprintWidget)
-		SLATE_ARGUMENT( bool, ShowPublicViewControl )
+		SLATE_ATTRIBUTE( bool, ShowPublicViewControl )
 		SLATE_ARGUMENT( bool, HideNameArea )
 		SLATE_ARGUMENT( FIsPropertyEditingEnabled, IsPropertyEditingEnabledDelegate )
 		SLATE_ARGUMENT( FOnFinishedChangingProperties::FDelegate, OnFinishedChangingProperties )
@@ -160,7 +160,7 @@ protected:
 	TSharedPtr<SBorder> ContextualEditingBorderWidget;
 
 	/** If true show the public view control */
-	bool bShowPublicView;
+	TAttribute<bool> bShowPublicView;
 
 	/** If true show the kismet inspector title widget */
 	bool bShowTitleArea;
