Added binding to create a new Custom Event Node. No binding by default. Similar to pressing C to create a new Comment Node.

diff --git a/Engine/Source/Editor/BlueprintGraph/Classes/EdGraphSchema_K2_Actions.h b/Engine/Source/Editor/BlueprintGraph/Classes/EdGraphSchema_K2_Actions.h
--- a/Engine/Source/Editor/BlueprintGraph/Classes/EdGraphSchema_K2_Actions.h
+++ b/Engine/Source/Editor/BlueprintGraph/Classes/EdGraphSchema_K2_Actions.h
@@ -317,7 +317,7 @@ struct FEdGraphSchemaAction_K2AddCustomEvent : public FEdGraphSchemaAction_K2New
 	{}
 
 	// FEdGraphSchemaAction interface
-	virtual UEdGraphNode* PerformAction(class UEdGraph* ParentGraph, UEdGraphPin* FromPin, const FVector2D Location, bool bSelectNewNode = true) override;
+	BLUEPRINTGRAPH_API virtual UEdGraphNode* PerformAction(class UEdGraph* ParentGraph, UEdGraphPin* FromPin, const FVector2D Location, bool bSelectNewNode = true) override;
 	// End of FEdGraphSchemaAction interface
 };
 
diff --git a/Engine/Source/Editor/GraphEditor/Private/GraphEditorActions.cpp b/Engine/Source/Editor/GraphEditor/Private/GraphEditorActions.cpp
--- a/Engine/Source/Editor/GraphEditor/Private/GraphEditorActions.cpp
+++ b/Engine/Source/Editor/GraphEditor/Private/GraphEditorActions.cpp
@@ -95,6 +95,7 @@ void FGraphEditorCommandsImpl::RegisterCommands()
 	UI_COMMAND( SelectAllOutputNodes, "Select All Output Nodes", "Adds all output Nodes linked to this Pin to selection", EUserInterfaceActionType::Button, FInputChord() )
 
 	UI_COMMAND( CreateComment, "Create Comment", "Create a comment box", EUserInterfaceActionType::Button, FInputChord(EKeys::C))
+	UI_COMMAND( CreateCustomEvent, "Create Custom Event", "Create a new custom event node.", EUserInterfaceActionType::Button, FInputChord())
 
 	UI_COMMAND( ZoomIn, "Zoom In", "Zoom in on the graph editor", EUserInterfaceActionType::Button, FInputChord(EKeys::Add))
 	UI_COMMAND( ZoomOut, "Zoom Out", "Zoom out from the graph editor", EUserInterfaceActionType::Button, FInputChord(EKeys::Subtract))
diff --git a/Engine/Source/Editor/GraphEditor/Public/GraphEditorActions.h b/Engine/Source/Editor/GraphEditor/Public/GraphEditorActions.h
--- a/Engine/Source/Editor/GraphEditor/Public/GraphEditorActions.h
+++ b/Engine/Source/Editor/GraphEditor/Public/GraphEditorActions.h
@@ -124,6 +124,9 @@ public:
 
 	//create a comment node
 	TSharedPtr< FUICommandInfo > CreateComment;
+	
+	// Create a custom event node
+	TSharedPtr< FUICommandInfo > CreateCustomEvent;
 
 	// Zoom in and out on the graph editor
 	TSharedPtr< FUICommandInfo > ZoomIn;
diff --git a/Engine/Source/Editor/Kismet/Private/BlueprintEditor.cpp b/Engine/Source/Editor/Kismet/Private/BlueprintEditor.cpp
--- a/Engine/Source/Editor/Kismet/Private/BlueprintEditor.cpp
+++ b/Engine/Source/Editor/Kismet/Private/BlueprintEditor.cpp
@@ -1636,6 +1636,10 @@ TSharedRef<SGraphEditor> FBlueprintEditor::CreateGraphEditorWidget(TSharedRef<FT
 			GraphEditorCommands->MapAction( FGraphEditorCommands::Get().CreateComment,
 				FExecuteAction::CreateSP( this, &FBlueprintEditor::OnCreateComment )
 				);
+				
+			GraphEditorCommands->MapAction( FGraphEditorCommands::Get().CreateCustomEvent,
+				FExecuteAction::CreateSP( this, &FBlueprintEditor::OnCreateCustomEvent )
+				);
 
 			GraphEditorCommands->MapAction(FGraphEditorCommands::Get().ShowAllPins,
 				FExecuteAction::CreateSP(this, &FBlueprintEditor::SetPinVisibility, SGraphEditor::Pin_Show)
@@ -10157,6 +10161,27 @@ void FBlueprintEditor::OnCreateComment()
 	}
 }
 
+void FBlueprintEditor::OnCreateCustomEvent()
+{
+	TSharedPtr<SGraphEditor> GraphEditor = FocusedGraphEdPtr.Pin();
+	if (GraphEditor.IsValid())
+	{
+		if (UEdGraph* Graph = GraphEditor->GetCurrentGraph())
+		{
+			if (const UEdGraphSchema* Schema = Graph->GetSchema())
+			{
+				if (Schema->IsA(UEdGraphSchema_K2::StaticClass()))
+				{
+					// BlueprintEventNodeSpawner seems better but we'll use FEdGraphSchemaAction_K2AddCustomEvent
+					FEdGraphSchemaAction_K2AddCustomEvent EventAction;
+					EventAction.NodeTemplate = NewObject<UK2Node_CustomEvent>();
+					EventAction.PerformAction(Graph, nullptr, GraphEditor->GetPasteLocation());
+				}
+			}
+		}
+	}
+}
+
 void FBlueprintEditor::SetPinVisibility(SGraphEditor::EPinVisibility Visibility)
 {
 	PinVisibility = Visibility;
diff --git a/Engine/Source/Editor/Kismet/Public/BlueprintEditor.h b/Engine/Source/Editor/Kismet/Public/BlueprintEditor.h
--- a/Engine/Source/Editor/Kismet/Public/BlueprintEditor.h
+++ b/Engine/Source/Editor/Kismet/Public/BlueprintEditor.h
@@ -1280,6 +1280,9 @@ protected:
 	
 	/* Create comment node on graph */
 	virtual void OnCreateComment();
+	
+	/** Create custom event node on graph */
+	virtual void OnCreateCustomEvent();
 
 	// Create new graph editor widget for the supplied document container
 	virtual TSharedRef<SGraphEditor> CreateGraphEditorWidget(TSharedRef<class FTabInfo> InTabInfo, class UEdGraph* InGraph);
