Improves "Find References" in blueprints: now supports function search by [class, function] call-sites and definitions rather than a generic search for function name. Replaces the "Find References" context menu action with a sub-menu for variables and functions in BlueprintEditor, WidgetBlueprintEditor, SubobjectEditor and SubobjectEditor. That context sub-menu has 'By Name' and 'By [REDACTED_PERSON] search options and local+global versions of both.

Made changes to Find-in-Blueprints metadata that is generated per blueprint asset to be able to do "specific function of a specific class" type queries; some search types were unsupported with previous metadata. Incremented the EFiBVersion so that the Find-in-Blueprints search window will ask to re-index all blueprints and resave. Added an opt-out editor setting "Allow Index All Blueprints" (default: true) to disable the button, which can be decided per project. Added an action in the Find-in-Blueprints modal when outdated metadata is detected to export the list of affected assets.

diff --git a/Engine/Source/Editor/BlueprintGraph/Classes/K2Node_CallFunction.h b/Engine/Source/Editor/BlueprintGraph/Classes/K2Node_CallFunction.h
--- a/Engine/Source/Editor/BlueprintGraph/Classes/K2Node_CallFunction.h
+++ b/Engine/Source/Editor/BlueprintGraph/Classes/K2Node_CallFunction.h
@@ -97,6 +97,7 @@ public:
 	virtual void GetPinHoverText(const UEdGraphPin& Pin, FString& HoverTextOut) const override;
 	virtual void AllocateDefaultPins() override;
 	virtual FLinearColor GetNodeTitleColor() const override;
+	virtual FString GetFindReferenceSearchString_Impl(EGetFindReferenceSearchStringFlags InFlags) const override;
 	virtual FText GetTooltipText() const override;
 	virtual FText GetNodeTitle(ENodeTitleType::Type TitleType) const override;
 	virtual FString GetDescriptiveCompiledName() const override;
@@ -109,6 +110,7 @@ public:
 	virtual bool CanPasteHere(const UEdGraph* TargetGraph) const override;
 	virtual void PinDefaultValueChanged(UEdGraphPin* Pin) override;
 	virtual void AddSearchMetaDataInfo(TArray<struct FSearchTagDataPair>& OutTaggedMetaData) const override;
+	virtual void AddPinSearchMetaDataInfo(const UEdGraphPin* Pin, TArray<struct FSearchTagDataPair>& OutTaggedMetaData) const override;
 	virtual TSharedPtr<SWidget> CreateNodeImage() const override;
 	virtual UObject* GetJumpTargetForDoubleClick() const override;
 	virtual bool CanJumpToDefinition() const override;
diff --git a/Engine/Source/Editor/BlueprintGraph/Classes/K2Node_CallParentFunction.h b/Engine/Source/Editor/BlueprintGraph/Classes/K2Node_CallParentFunction.h
--- a/Engine/Source/Editor/BlueprintGraph/Classes/K2Node_CallParentFunction.h
+++ b/Engine/Source/Editor/BlueprintGraph/Classes/K2Node_CallParentFunction.h
@@ -26,6 +26,7 @@ class UK2Node_CallParentFunction : public UK2Node_CallFunction
 	virtual void AllocateDefaultPins() override;
 	virtual FLinearColor GetNodeTitleColor() const override;
 	virtual void PostPlacedNewNode() override;
+	virtual void AddSearchMetaDataInfo(TArray<struct FSearchTagDataPair>& OutTaggedMetaData) const override;
 	//~ End EdGraphNode Interface
 
 	virtual void SetFromFunction(const UFunction* Function) override;
diff --git a/Engine/Source/Editor/BlueprintGraph/Classes/K2Node_Event.h b/Engine/Source/Editor/BlueprintGraph/Classes/K2Node_Event.h
--- a/Engine/Source/Editor/BlueprintGraph/Classes/K2Node_Event.h
+++ b/Engine/Source/Editor/BlueprintGraph/Classes/K2Node_Event.h
@@ -87,9 +87,10 @@ class UK2Node_Event : public UK2Node_EditablePinBase, public IK2Node_EventNodeIn
 	BLUEPRINTGRAPH_API virtual FEdGraphNodeDeprecationResponse GetDeprecationResponse(EEdGraphNodeDeprecationType DeprecationType) const override;
 	BLUEPRINTGRAPH_API virtual UObject* GetJumpTargetForDoubleClick() const override;
 	BLUEPRINTGRAPH_API virtual FSlateIcon GetIconAndTint(FLinearColor& OutColor) const override;
-	BLUEPRINTGRAPH_API virtual FString GetFindReferenceSearchString() const override;
+	BLUEPRINTGRAPH_API virtual FString GetFindReferenceSearchString_Impl(EGetFindReferenceSearchStringFlags InFlags) const override;
 	BLUEPRINTGRAPH_API virtual void FindDiffs(UEdGraphNode* OtherNode, struct FDiffResults& Results) override;
 	BLUEPRINTGRAPH_API virtual bool HasExternalDependencies(TArray<class UStruct*>* OptionalOutput) const override;
+	BLUEPRINTGRAPH_API virtual void AddSearchMetaDataInfo(TArray<struct FSearchTagDataPair>& OutTaggedMetaData) const override;
 	//~ End UEdGraphNode Interface
 
 	//~ Begin UK2Node Interface
diff --git a/Engine/Source/Editor/BlueprintGraph/Classes/K2Node_FunctionEntry.h b/Engine/Source/Editor/BlueprintGraph/Classes/K2Node_FunctionEntry.h
--- a/Engine/Source/Editor/BlueprintGraph/Classes/K2Node_FunctionEntry.h
+++ b/Engine/Source/Editor/BlueprintGraph/Classes/K2Node_FunctionEntry.h
@@ -76,6 +76,7 @@ class UK2Node_FunctionEntry : public UK2Node_FunctionTerminator
 	virtual void FindDiffs(UEdGraphNode* OtherNode, struct FDiffResults& Results) override;
 	virtual bool IsCompatibleWithGraph(const UEdGraph* InGraph) const override;
 	virtual void PostPasteNode() override;
+	virtual void AddSearchMetaDataInfo(TArray<struct FSearchTagDataPair>& OutTaggedMetaData) const override;
 	//~ End UEdGraphNode Interface
 
 	//~ Begin UK2Node Interface
diff --git a/Engine/Source/Editor/BlueprintGraph/Classes/K2Node_FunctionTerminator.h b/Engine/Source/Editor/BlueprintGraph/Classes/K2Node_FunctionTerminator.h
--- a/Engine/Source/Editor/BlueprintGraph/Classes/K2Node_FunctionTerminator.h
+++ b/Engine/Source/Editor/BlueprintGraph/Classes/K2Node_FunctionTerminator.h
@@ -37,6 +37,7 @@ class UK2Node_FunctionTerminator : public UK2Node_EditablePinBase
 
 	//~ Begin UEdGraphNode Interface
 	virtual FLinearColor GetNodeTitleColor() const override;
+	virtual FString GetFindReferenceSearchString_Impl(EGetFindReferenceSearchStringFlags InFlags) const override;
 	virtual FName CreateUniquePinName(FName SourcePinName) const override;
 	virtual void ValidateNodeDuringCompilation(class FCompilerResultsLog& MessageLog) const override;
 	virtual bool HasExternalDependencies(TArray<class UStruct*>* OptionalOutput) const override;
diff --git a/Engine/Source/Editor/BlueprintGraph/Classes/K2Node_StructOperation.h b/Engine/Source/Editor/BlueprintGraph/Classes/K2Node_StructOperation.h
--- a/Engine/Source/Editor/BlueprintGraph/Classes/K2Node_StructOperation.h
+++ b/Engine/Source/Editor/BlueprintGraph/Classes/K2Node_StructOperation.h
@@ -43,7 +43,7 @@ class UK2Node_StructOperation : public UK2Node_Variable
 	//virtual bool DrawNodeAsVariable() const override { return true; }
 	virtual bool ShouldShowNodeProperties() const override { return true; }
 	virtual void ValidateNodeDuringCompilation(class FCompilerResultsLog& MessageLog) const override;
-	virtual FString GetFindReferenceSearchString() const override;
+	virtual FString GetFindReferenceSearchString_Impl(EGetFindReferenceSearchStringFlags InFlags) const override;
 	virtual bool IsActionFilteredOut(class FBlueprintActionFilter const& Filter) override;
 	//~ End UK2Node Interface
 
diff --git a/Engine/Source/Editor/BlueprintGraph/Classes/K2Node_Variable.h b/Engine/Source/Editor/BlueprintGraph/Classes/K2Node_Variable.h
--- a/Engine/Source/Editor/BlueprintGraph/Classes/K2Node_Variable.h
+++ b/Engine/Source/Editor/BlueprintGraph/Classes/K2Node_Variable.h
@@ -84,7 +84,7 @@ public:
 
 	//~ Begin UEdGraphNode Interface
 	virtual FLinearColor GetNodeTitleColor() const override;
-	virtual FString GetFindReferenceSearchString() const override;
+	virtual FString GetFindReferenceSearchString_Impl(EGetFindReferenceSearchStringFlags InFlags) const override;
 	virtual void ReconstructNode() override;
 	virtual FString GetDocumentationLink() const override;
 	virtual FString GetDocumentationExcerptName() const override;
diff --git a/Engine/Source/Editor/BlueprintGraph/Private/BlueprintEditorSettings.cpp b/Engine/Source/Editor/BlueprintGraph/Private/BlueprintEditorSettings.cpp
--- a/Engine/Source/Editor/BlueprintGraph/Private/BlueprintEditorSettings.cpp
+++ b/Engine/Source/Editor/BlueprintGraph/Private/BlueprintEditorSettings.cpp
@@ -68,6 +68,8 @@ UBlueprintEditorSettings::UBlueprintEditorSettings(const FObjectInitializer& Obj
 	, bShowActionMenuItemSignatures(false)
 	// Perf Settings
 	, NodeTemplateCacheCapMB(20.f)
+	// Find-in-Blueprints Settings
+	, bAllowIndexAllBlueprints(true)
 	// No category
 	, bShowInheritedVariables(false)
 	, bAlwaysShowInterfacesInOverrides(true)
diff --git a/Engine/Source/Editor/BlueprintGraph/Private/EdGraphSchema_K2.cpp b/Engine/Source/Editor/BlueprintGraph/Private/EdGraphSchema_K2.cpp
--- a/Engine/Source/Editor/BlueprintGraph/Private/EdGraphSchema_K2.cpp
+++ b/Engine/Source/Editor/BlueprintGraph/Private/EdGraphSchema_K2.cpp
@@ -1750,8 +1750,23 @@ void UEdGraphSchema_K2::GetContextMenuActions(UToolMenu* Menu, UGraphNodeContext
 					Section.AddMenuEntry(FGraphEditorCommands::Get().GoToDefinition);
 				}
 
-				// show search for references for everyone
-				Section.AddMenuEntry(FGraphEditorCommands::Get().FindReferences);
+				// Show search for references for everyone. Depending on context, it's an action or a submenu.
+				const bool bIsFuncOrVarNode = InGraphNode->IsA<UK2Node_CallFunction>() || InGraphNode->IsA<UK2Node_Event>() || InGraphNode->IsA<UK2Node_FunctionTerminator>() || InGraphNode->IsA<UK2Node_Variable>(); 
+				const bool bExpandFindReferences = bIsFuncOrVarNode;
+				if (!bExpandFindReferences)
+				{
+					Section.AddMenuEntry(FGraphEditorCommands::Get().FindReferences);
+				}
+				else
+				{
+					// Expandable menu: insert sub-menu here
+					Section.AddSubMenu(
+						FName("FindReferenceSubMenu"),
+						LOCTEXT("FindReferences_Label", "Find References"),
+						LOCTEXT("FindReferences_Tooltip", "Options for finding references to class members"),
+						FNewToolMenuChoice(FNewMenuDelegate::CreateStatic(&FGraphEditorCommands::BuildFindReferencesMenu))
+					);
+				}
 
 				if (!bIsDebugging)
 				{
diff --git a/Engine/Source/Editor/BlueprintGraph/Private/K2Node_CallFunction.cpp b/Engine/Source/Editor/BlueprintGraph/Private/K2Node_CallFunction.cpp
--- a/Engine/Source/Editor/BlueprintGraph/Private/K2Node_CallFunction.cpp
+++ b/Engine/Source/Editor/BlueprintGraph/Private/K2Node_CallFunction.cpp
@@ -15,6 +15,7 @@
 #include "GraphEditorSettings.h"
 #include "EdGraph/EdGraph.h"
 #include "EdGraphSchema_K2.h"
+#include "FindInBlueprints.h"
 #include "K2Node_Event.h"
 #include "K2Node_AssignmentStatement.h"
 #include "K2Node_CallArrayFunction.h"
@@ -1589,6 +1590,24 @@ FLinearColor UK2Node_CallFunction::GetNodeTitleColor() const
 	return GetPalletteIconColor(GetTargetFunction());
 }
 
+FString UK2Node_CallFunction::GetFindReferenceSearchString_Impl(EGetFindReferenceSearchStringFlags InFlags) const
+{
+	if (EnumHasAnyFlags(InFlags, EGetFindReferenceSearchStringFlags::UseSearchSyntax))
+	{
+		if (const UFunction* Function = GetTargetFunction())
+		{
+			FString SearchTerm;
+			if (FindInBlueprintsHelpers::ConstructSearchTermFromFunction(Function, SearchTerm))
+			{
+				return SearchTerm;
+			}
+		}
+	}
+
+	// Fallback behavior
+	return Super::GetFindReferenceSearchString_Impl(InFlags);
+}
+
 FText UK2Node_CallFunction::GetTooltipText() const
 {
 	FText Tooltip;
@@ -3510,6 +3529,24 @@ void UK2Node_CallFunction::AddSearchMetaDataInfo(TArray<struct FSearchTagDataPai
 	}
 }
 
+void UK2Node_CallFunction::AddPinSearchMetaDataInfo(const UEdGraphPin* Pin, TArray<FSearchTagDataPair>& OutTaggedMetaData) const
+{
+	Super::AddPinSearchMetaDataInfo(Pin, OutTaggedMetaData);
+
+	// Blueprint graphs that call a function declared in the same blueprint don't store a target type, but rather PinSubCategory == Self.
+	// When this is the case, we will still explicitly index the ObjectClass for the target pin, so that it can be treated the same as 
+	// any other call function nodes.
+	if (Pin->PinName == UEdGraphSchema_K2::PN_Self && Pin->PinType.PinSubCategory == UEdGraphSchema_K2::PSC_Self && !Pin->PinType.PinSubCategoryObject.IsValid())
+	{
+		// Get the parent or interface class that originally defined this function
+		if (const UClass* FuncOriginClass = FindInBlueprintsHelpers::GetFunctionOriginClass(GetTargetFunction()))
+		{
+			const FString FuncOriginClassName = FuncOriginClass->GetPathName();
+			OutTaggedMetaData.Add(FSearchTagDataPair(FFindInBlueprintSearchTags::FiB_ObjectClass, FText::FromString(FuncOriginClassName)));
+		}
+	}
+}
+
 TSharedPtr<SWidget> UK2Node_CallFunction::CreateNodeImage() const
 {
 	// For set, map and array functions we have a cool icon. This helps users quickly
diff --git a/Engine/Source/Editor/BlueprintGraph/Private/K2Node_CallParentFunction.cpp b/Engine/Source/Editor/BlueprintGraph/Private/K2Node_CallParentFunction.cpp
--- a/Engine/Source/Editor/BlueprintGraph/Private/K2Node_CallParentFunction.cpp
+++ b/Engine/Source/Editor/BlueprintGraph/Private/K2Node_CallParentFunction.cpp
@@ -8,6 +8,8 @@
 #include "Engine/Blueprint.h"
 #include "Engine/MemberReference.h"
 #include "GraphEditorSettings.h"
+#include "FindInBlueprintManager.h"
+#include "FindInBlueprints.h"
 #include "HAL/PlatformMath.h"
 #include "Internationalization/Internationalization.h"
 #include "K2Node.h"
@@ -94,4 +96,18 @@ void UK2Node_CallParentFunction::PostPlacedNewNode()
 	UK2Node::PostPlacedNewNode();
 }
 
+void UK2Node_CallParentFunction::AddSearchMetaDataInfo(TArray<struct FSearchTagDataPair>& OutTaggedMetaData) const
+{
+	Super::AddSearchMetaDataInfo(OutTaggedMetaData);
+
+	// CallParentFunction nodes will always target the direct parent class, unlike CallFunction nodes.
+	// However, we still want the function to be searchable using the ancestor class or interface that declares it.
+	// Therefore, index a FuncOriginClass metadata value and don't rely on the target pin's type.
+	if (const UClass* FuncOriginClass = FindInBlueprintsHelpers::GetFunctionOriginClass(GetTargetFunction()))
+	{
+		const FString FuncOriginClassName = FuncOriginClass->GetPathName();
+		OutTaggedMetaData.Add(FSearchTagDataPair(FFindInBlueprintSearchTags::FiB_FuncOriginClass, FText::FromString(FuncOriginClassName)));
+	}
+}
+
 #undef LOCTEXT_NAMESPACE
diff --git a/Engine/Source/Editor/BlueprintGraph/Private/K2Node_CustomEvent.cpp b/Engine/Source/Editor/BlueprintGraph/Private/K2Node_CustomEvent.cpp
--- a/Engine/Source/Editor/BlueprintGraph/Private/K2Node_CustomEvent.cpp
+++ b/Engine/Source/Editor/BlueprintGraph/Private/K2Node_CustomEvent.cpp
@@ -15,6 +15,7 @@
 #include "Engine/BlueprintGeneratedClass.h"
 #include "Engine/MemberReference.h"
 #include "FindInBlueprintManager.h"
+#include "FindInBlueprints.h"
 #include "HAL/PlatformCrt.h"
 #include "Internationalization/Internationalization.h"
 #include "K2Node.h"
@@ -601,16 +602,28 @@ void UK2Node_CustomEvent::AddSearchMetaDataInfo(TArray<struct FSearchTagDataPair
 {
 	Super::AddSearchMetaDataInfo(OutTaggedMetaData);
 
+	bool bNeedsNameUpdate = true;
+	bool bNeedsNativeNameUpdate = true;
 	for (FSearchTagDataPair& SearchData : OutTaggedMetaData)
 	{
 		// Should always be the first item, but there is no guarantee
-		if (SearchData.Key.CompareTo(FFindInBlueprintSearchTags::FiB_Name) == 0)
+		if (bNeedsNameUpdate && SearchData.Key.CompareTo(FFindInBlueprintSearchTags::FiB_Name) == 0)
 		{
 			SearchData.Value = FText::FromString(FName::NameToDisplayString(CustomFunctionName.ToString(), false));
+			bNeedsNameUpdate = false;
+		}
+		else if (bNeedsNativeNameUpdate && SearchData.Key.CompareTo(FFindInBlueprintSearchTags::FiB_NativeName) == 0)
+		{
+			SearchData.Value = FText::FromName(CustomFunctionName);
+			bNeedsNativeNameUpdate = false;
+		}
+
+		// If no more keys need updating, break
+		if (!bNeedsNameUpdate && !bNeedsNativeNameUpdate)
+		{
 			break;
 		}
 	}
-	OutTaggedMetaData.Add(FSearchTagDataPair(FFindInBlueprintSearchTags::FiB_NativeName, FText::FromName(CustomFunctionName)));
 }
 
 FText UK2Node_CustomEvent::GetKeywords() const
diff --git a/Engine/Source/Editor/BlueprintGraph/Private/K2Node_Event.cpp b/Engine/Source/Editor/BlueprintGraph/Private/K2Node_Event.cpp
--- a/Engine/Source/Editor/BlueprintGraph/Private/K2Node_Event.cpp
+++ b/Engine/Source/Editor/BlueprintGraph/Private/K2Node_Event.cpp
@@ -13,6 +13,7 @@
 #include "Engine/Blueprint.h"
 #include "EngineLogs.h"
 #include "EventEntryHandler.h"
+#include "FindInBlueprints.h"
 #include "GameFramework/Actor.h"
 #include "GraphEditorSettings.h"
 #include "HAL/PlatformCrt.h"
@@ -898,24 +899,27 @@ FSlateIcon UK2Node_Event::GetIconAndTint(FLinearColor& OutColor) const
 	return Icon;
 }
 
-FString UK2Node_Event::GetFindReferenceSearchString() const
+FString UK2Node_Event::GetFindReferenceSearchString_Impl(EGetFindReferenceSearchStringFlags InFlags) const
 {
-	// Behavior modeled after UK2Node_Event::GetNodeTitle but without "Event" prepended
-	if (bOverrideFunction || (CustomFunctionName == NAME_None))
+	// Resolve the function
+	if (const UFunction* Function = FFunctionFromNodeHelper::FunctionFromNode(this))
 	{
-		FString FunctionName = EventReference.GetMemberName().ToString(); // If we fail to find the function, still want to search for its expected name.
-
-		if (const UFunction* Function = EventReference.ResolveMember<UFunction>(GetBlueprintClassFromNode()))
+		// Attempt to construct an advanced search syntax query from the function
+		FString SearchTerm;
+		if (EnumHasAnyFlags(InFlags, EGetFindReferenceSearchStringFlags::UseSearchSyntax) && FindInBlueprintsHelpers::ConstructSearchTermFromFunction(Function, SearchTerm))
 		{
-			FunctionName = UEdGraphSchema_K2::GetFriendlySignatureName(Function).ToString();
+			return SearchTerm;
+		}
+		else
+		{
+			// Fallback behavior: function was found but failed to construct a search term from it
+			// Just search for the function's friendly name
+			return UEdGraphSchema_K2::GetFriendlySignatureName(Function).ToString();
 		}
-
-		return FunctionName;
-	}
-	else
-	{
-		return CustomFunctionName.ToString();
 	}
+	
+	// If we fail to find the function, still want to search for its expected name.
+	return EventReference.GetMemberName().ToString();
 }
 
 void UK2Node_Event::FindDiffs(UEdGraphNode* OtherNode, struct FDiffResults& Results)
@@ -961,6 +965,25 @@ bool UK2Node_Event::HasExternalDependencies(TArray<class UStruct*>* OptionalOutp
 	return bSuperResult || bResult;
 }
 
+void UK2Node_Event::AddSearchMetaDataInfo(TArray<FSearchTagDataPair>& OutTaggedMetaData) const
+{
+	Super::AddSearchMetaDataInfo(OutTaggedMetaData);
+
+	if (const UFunction* Function = FFunctionFromNodeHelper::FunctionFromNode(this))
+	{
+		// Index the native name of the function, this will be used in search queries rather than node title
+		const FString FunctionNativeName = Function->GetName();
+		OutTaggedMetaData.Add(FSearchTagDataPair(FFindInBlueprintSearchTags::FiB_NativeName, FText::FromString(FunctionNativeName)));
+
+		// Index the (ancestor) class or interface from which the function originates, can be self
+		if (const UClass* FuncOriginClass = FindInBlueprintsHelpers::GetFunctionOriginClass(Function))
+		{
+			const FString FuncOriginClassName = FuncOriginClass->GetPathName();
+			OutTaggedMetaData.Add(FSearchTagDataPair(FFindInBlueprintSearchTags::FiB_FuncOriginClass, FText::FromString(FuncOriginClassName)));
+		}
+	}
+}
+
 TSharedPtr<FEdGraphSchemaAction> UK2Node_Event::GetEventNodeAction(const FText& ActionCategory)
 {
 	TSharedPtr<FEdGraphSchemaAction_K2Event> EventNodeAction = MakeShareable(new FEdGraphSchemaAction_K2Event(ActionCategory, GetNodeTitle(ENodeTitleType::EditableTitle), GetTooltipText(), 0));
diff --git a/Engine/Source/Editor/BlueprintGraph/Private/K2Node_FunctionEntry.cpp b/Engine/Source/Editor/BlueprintGraph/Private/K2Node_FunctionEntry.cpp
--- a/Engine/Source/Editor/BlueprintGraph/Private/K2Node_FunctionEntry.cpp
+++ b/Engine/Source/Editor/BlueprintGraph/Private/K2Node_FunctionEntry.cpp
@@ -19,6 +19,8 @@
 #include "Engine/Blueprint.h"
 #include "Engine/MemberReference.h"
 #include "EngineLogs.h"
+#include "FindInBlueprintManager.h"
+#include "FindInBlueprints.h"
 #include "HAL/PlatformCrt.h"
 #include "Internationalization/Internationalization.h"
 #include "K2Node_CallFunction.h"
@@ -891,6 +893,25 @@ void UK2Node_FunctionEntry::PostPasteNode()
 	ReconstructNode();
 }
 
+void UK2Node_FunctionEntry::AddSearchMetaDataInfo(TArray<FSearchTagDataPair>& OutTaggedMetaData) const
+{
+	Super::AddSearchMetaDataInfo(OutTaggedMetaData);
+
+	if (const UFunction* Function = FFunctionFromNodeHelper::FunctionFromNode(this))
+	{
+		// Index the native name of the function, this will be used in search queries rather than node title
+		const FString FunctionNativeName = Function->GetName();
+		OutTaggedMetaData.Add(FSearchTagDataPair(FFindInBlueprintSearchTags::FiB_NativeName, FText::FromString(FunctionNativeName)));
+
+		// Index the (ancestor) class or interface from which the function originates, can be self
+		if (const UClass* FuncOriginClass = FindInBlueprintsHelpers::GetFunctionOriginClass(Function))
+		{
+			const FString FuncOriginClassName = FuncOriginClass->GetPathName();
+			OutTaggedMetaData.Add(FSearchTagDataPair(FFindInBlueprintSearchTags::FiB_FuncOriginClass, FText::FromString(FuncOriginClassName)));
+		}
+	}
+}
+
 int32 UK2Node_FunctionEntry::GetFunctionFlags() const
 {
 	int32 ReturnFlags = 0;
diff --git a/Engine/Source/Editor/BlueprintGraph/Private/K2Node_FunctionTerminator.cpp b/Engine/Source/Editor/BlueprintGraph/Private/K2Node_FunctionTerminator.cpp
--- a/Engine/Source/Editor/BlueprintGraph/Private/K2Node_FunctionTerminator.cpp
+++ b/Engine/Source/Editor/BlueprintGraph/Private/K2Node_FunctionTerminator.cpp
@@ -9,6 +9,7 @@
 #include "EdGraph/EdGraphPin.h"
 #include "EdGraphSchema_K2.h"
 #include "Engine/Blueprint.h"
+#include "FindInBlueprints.h"
 #include "GraphEditorSettings.h"
 #include "HAL/Platform.h"
 #include "HAL/PlatformCrt.h"
@@ -53,6 +54,32 @@ FLinearColor UK2Node_FunctionTerminator::GetNodeTitleColor() const
 	return GetDefault<UGraphEditorSettings>()->FunctionTerminatorNodeTitleColor;
 }
 
+FString UK2Node_FunctionTerminator::GetFindReferenceSearchString_Impl(EGetFindReferenceSearchStringFlags InFlags) const
+{
+	if (EnumHasAnyFlags(InFlags, EGetFindReferenceSearchStringFlags::UseSearchSyntax))
+	{
+		// Resolve the function
+		if (const UFunction* Function = FFunctionFromNodeHelper::FunctionFromNode(this))
+		{
+			// Attempt to construct an advanced search syntax query from the function
+			FString SearchTerm;
+			if (FindInBlueprintsHelpers::ConstructSearchTermFromFunction(Function, SearchTerm))
+			{
+				return SearchTerm;
+			}
+			else
+			{
+				// Fallback behavior: function was found but failed to construct a search term from it
+				// Just search for the function's friendly name
+				return UEdGraphSchema_K2::GetFriendlySignatureName(Function).ToString();
+			}
+		}
+	}
+
+	// Fallback behavior: function was not resolved
+	return Super::GetFindReferenceSearchString_Impl(InFlags);
+}
+
 FName UK2Node_FunctionTerminator::CreateUniquePinName(FName InSourcePinName) const
 {
 	const UFunction* FoundFunction = FFunctionFromNodeHelper::FunctionFromNode(this);
diff --git a/Engine/Source/Editor/BlueprintGraph/Private/K2Node_StructOperation.cpp b/Engine/Source/Editor/BlueprintGraph/Private/K2Node_StructOperation.cpp
--- a/Engine/Source/Editor/BlueprintGraph/Private/K2Node_StructOperation.cpp
+++ b/Engine/Source/Editor/BlueprintGraph/Private/K2Node_StructOperation.cpp
@@ -159,9 +159,9 @@ FString UK2Node_StructOperation::GetPinMetaData(FName InPinName, FName InKey)
 	return Super::GetPinMetaData(InPinName, InKey);
 }
 
-FString UK2Node_StructOperation::GetFindReferenceSearchString() const
+FString UK2Node_StructOperation::GetFindReferenceSearchString_Impl(EGetFindReferenceSearchStringFlags InFlags) const
 {
-	return UEdGraphNode::GetFindReferenceSearchString();
+	return UEdGraphNode::GetFindReferenceSearchString_Impl(InFlags);
 }
 
 bool UK2Node_StructOperation::IsActionFilteredOut(const FBlueprintActionFilter& Filter)
diff --git a/Engine/Source/Editor/BlueprintGraph/Private/K2Node_Variable.cpp b/Engine/Source/Editor/BlueprintGraph/Private/K2Node_Variable.cpp
--- a/Engine/Source/Editor/BlueprintGraph/Private/K2Node_Variable.cpp
+++ b/Engine/Source/Editor/BlueprintGraph/Private/K2Node_Variable.cpp
@@ -263,22 +263,24 @@ FLinearColor UK2Node_Variable::GetNodeTitleColor() const
 	return FLinearColor::White;
 }
 
-FString UK2Node_Variable::GetFindReferenceSearchString() const
+FString UK2Node_Variable::GetFindReferenceSearchString_Impl(EGetFindReferenceSearchStringFlags InFlags) const
 {
-	FString ResultSearchString;
-	if (VariableReference.IsLocalScope())
+	if (EnumHasAnyFlags(InFlags, EGetFindReferenceSearchStringFlags::UseSearchSyntax))
 	{
-		ResultSearchString = VariableReference.GetReferenceSearchString(nullptr);
-	}
-	else
-	{
-		FProperty* VariableProperty = VariableReference.ResolveMember<FProperty>(GetBlueprintClassFromNode());
-		if (VariableProperty)
+		if (VariableReference.IsLocalScope())
+		{
+			// Generate local variable search query
+			return VariableReference.GetReferenceSearchString(nullptr);
+		}
+		else if (FProperty* VariableProperty = VariableReference.ResolveMember<FProperty>(GetBlueprintClassFromNode()))
 		{
-			ResultSearchString = VariableReference.GetReferenceSearchString(VariableProperty->GetOwnerClass());
+			// Generate member variable search query
+			return VariableReference.GetReferenceSearchString(VariableProperty->GetOwnerClass());
 		}
 	}
-	return ResultSearchString;
+
+	// Simple query: just search for variable name
+	return VariableReference.GetMemberName().ToString();
 }
 
 UK2Node::ERedirectType UK2Node_Variable::DoPinsMatchForReconstruction(const UEdGraphPin* NewPin, int32 NewPinIndex, const UEdGraphPin* OldPin, int32 OldPinIndex) const 
diff --git a/Engine/Source/Editor/BlueprintGraph/Public/BlueprintEditorSettings.h b/Engine/Source/Editor/BlueprintGraph/Public/BlueprintEditorSettings.h
--- a/Engine/Source/Editor/BlueprintGraph/Public/BlueprintEditorSettings.h
+++ b/Engine/Source/Editor/BlueprintGraph/Public/BlueprintEditorSettings.h
@@ -262,6 +262,12 @@ public:
 	UPROPERTY(EditAnywhere, AdvancedDisplay, config, Category = Performance, DisplayName = "Node-Template Cache Cap (MB)", meta = (ClampMin = "0", UIMin = "0"))
 	float NodeTemplateCacheCapMB;
 
+	// Find-in-Blueprint Settings
+public:
+	/** Whether to enable the "Index All" button in the Find-in-Blueprints search window when blueprint assets with an out-of-date index (search metadata) are found */
+	UPROPERTY(EditAnywhere, config, Category = FindInBlueprints)
+	bool bAllowIndexAllBlueprints;
+
 public:
 	/** If set we'll show the inherited variables in the My Blueprint view. */
 	UPROPERTY(config)
diff --git a/Engine/Source/Editor/GraphEditor/Private/GraphEditorActions.cpp b/Engine/Source/Editor/GraphEditor/Private/GraphEditorActions.cpp
--- a/Engine/Source/Editor/GraphEditor/Private/GraphEditorActions.cpp
+++ b/Engine/Source/Editor/GraphEditor/Private/GraphEditorActions.cpp
@@ -77,7 +77,11 @@ void FGraphEditorCommandsImpl::RegisterCommands()
 
 	UI_COMMAND( SelectReferenceInLevel, "Find Actor in Level", "Select the actor referenced by this node in the level", EUserInterfaceActionType::Button, FInputChord() )
 	UI_COMMAND( AssignReferencedActor, "Assign Selected Actor", "Assign the selected actor to be this node's referenced object", EUserInterfaceActionType::Button, FInputChord() )
-	UI_COMMAND( FindReferences, "Find References", "Find references of this item", EUserInterfaceActionType::Button, FInputChord(EKeys::F, EModifierKey::Shift | EModifierKey::Alt) )
+	UI_COMMAND( FindReferences, "Find References", "Find references of this item", EUserInterfaceActionType::Button, FInputChord() )
+	UI_COMMAND( FindReferencesByNameLocal, "By Name", "Find references of this item by name", EUserInterfaceActionType::Button, FInputChord() )
+	UI_COMMAND( FindReferencesByNameGlobal, "By Name (All)", "Find references of this item by name in all blueprints", EUserInterfaceActionType::Button, FInputChord() )
+	UI_COMMAND( FindReferencesByClassMemberLocal, "By Class Member", "Find references of this item by class member", EUserInterfaceActionType::Button, FInputChord() )
+	UI_COMMAND( FindReferencesByClassMemberGlobal, "By Class Member (All)", "Find references of this item by class member in all blueprints", EUserInterfaceActionType::Button, FInputChord(EKeys::F, EModifierKey::Shift | EModifierKey::Alt) )
 	UI_COMMAND( FindAndReplaceReferences, "Replace References", "Brings up a window to help find and replace all instances of this item", EUserInterfaceActionType::Button, FInputChord() )
 	
 	UI_COMMAND( GoToDefinition, "Goto Definition", "Jumps to the defintion of the selected node if available, e.g., C++ code for a native function or the graph for a Blueprint function.", EUserInterfaceActionType::Button, FInputChord(EKeys::G, EModifierKey::Alt) )
@@ -169,6 +173,18 @@ const FGraphEditorCommandsImpl& FGraphEditorCommands::Get()
 	return FGraphEditorCommandsImpl::Get();
 }
 
+void FGraphEditorCommands::BuildFindReferencesMenu(FMenuBuilder& MenuBuilder)
+{
+	MenuBuilder.BeginSection("FindReferences", LOCTEXT("FindReferences", "Find References"));
+	{
+		MenuBuilder.AddMenuEntry(FGraphEditorCommands::Get().FindReferencesByNameLocal);
+		MenuBuilder.AddMenuEntry(FGraphEditorCommands::Get().FindReferencesByNameGlobal);
+		MenuBuilder.AddMenuEntry(FGraphEditorCommands::Get().FindReferencesByClassMemberLocal);
+		MenuBuilder.AddMenuEntry(FGraphEditorCommands::Get().FindReferencesByClassMemberGlobal);
+	}
+	MenuBuilder.EndSection();
+}
+
 void FGraphEditorCommands::Unregister()
 {
 	return FGraphEditorCommandsImpl::Unregister();
diff --git a/Engine/Source/Editor/GraphEditor/Public/GraphEditorActions.h b/Engine/Source/Editor/GraphEditor/Public/GraphEditorActions.h
--- a/Engine/Source/Editor/GraphEditor/Public/GraphEditorActions.h
+++ b/Engine/Source/Editor/GraphEditor/Public/GraphEditorActions.h
@@ -5,6 +5,7 @@
 #include "Containers/Array.h"
 #include "CoreMinimal.h"
 #include "Framework/Commands/Commands.h"
+#include "Framework/MultiBox/MultiBoxBuilder.h"
 #include "HAL/Platform.h"
 #include "Internationalization/Internationalization.h"
 #include "Styling/AppStyle.h"
@@ -104,6 +105,13 @@ public:
 	// Find references
 	TSharedPtr< FUICommandInfo > FindReferences;
 	GRAPHEDITOR_API TSharedPtr< FUICommandInfo > GetFindReferences() const { return FindReferences; }
+
+	// Find references options that appear by context like for functions and variables
+	TSharedPtr< FUICommandInfo > FindReferencesByNameLocal;
+	TSharedPtr< FUICommandInfo > FindReferencesByNameGlobal;
+	TSharedPtr< FUICommandInfo > FindReferencesByClassMemberLocal;
+	TSharedPtr< FUICommandInfo > FindReferencesByClassMemberGlobal;
+	
 	TSharedPtr< FUICommandInfo > FindAndReplaceReferences;
 
 	// Jumps to the definition of the selected node (or otherwise focuses something interesting about that node, e.g., the inner graph for a collapsed graph)
@@ -155,6 +163,9 @@ public:
 	static void Register();
 
 	static const FGraphEditorCommandsImpl& Get();
+	
+	/** Build "Find References" submenu when a context allows for it */
+	static void BuildFindReferencesMenu(FMenuBuilder& MenuBuilder);
 
 	static void Unregister();
 };
diff --git a/Engine/Source/Editor/Kismet/Private/BlueprintEditor.cpp b/Engine/Source/Editor/Kismet/Private/BlueprintEditor.cpp
--- a/Engine/Source/Editor/Kismet/Private/BlueprintEditor.cpp
+++ b/Engine/Source/Editor/Kismet/Private/BlueprintEditor.cpp
@@ -1654,7 +1654,27 @@ TSharedRef<SGraphEditor> FBlueprintEditor::CreateGraphEditorWidget(TSharedRef<FT
 				);
 
 			GraphEditorCommands->MapAction( FGraphEditorCommands::Get().FindReferences,
-				FExecuteAction::CreateSP( this, &FBlueprintEditor::OnFindReferences ),
+				FExecuteAction::CreateSP( this, &FBlueprintEditor::OnFindReferences, /*bSearchAllBlueprints=*/false, EGetFindReferenceSearchStringFlags::None),
+				FCanExecuteAction::CreateSP( this, &FBlueprintEditor::CanFindReferences )
+				);
+			
+			GraphEditorCommands->MapAction( FGraphEditorCommands::Get().FindReferencesByNameLocal,
+				FExecuteAction::CreateSP( this, &FBlueprintEditor::OnFindReferences, /*bSearchAllBlueprints=*/false, EGetFindReferenceSearchStringFlags::None ),
+				FCanExecuteAction::CreateSP( this, &FBlueprintEditor::CanFindReferences )
+				);
+			
+			GraphEditorCommands->MapAction( FGraphEditorCommands::Get().FindReferencesByNameGlobal,
+				FExecuteAction::CreateSP( this, &FBlueprintEditor::OnFindReferences, /*bSearchAllBlueprints=*/true, EGetFindReferenceSearchStringFlags::None ),
+				FCanExecuteAction::CreateSP( this, &FBlueprintEditor::CanFindReferences )
+				);
+			
+			GraphEditorCommands->MapAction( FGraphEditorCommands::Get().FindReferencesByClassMemberLocal,
+				FExecuteAction::CreateSP( this, &FBlueprintEditor::OnFindReferences, /*bSearchAllBlueprints=*/false, EGetFindReferenceSearchStringFlags::UseSearchSyntax ),
+				FCanExecuteAction::CreateSP( this, &FBlueprintEditor::CanFindReferences )
+				);
+			
+			GraphEditorCommands->MapAction( FGraphEditorCommands::Get().FindReferencesByClassMemberGlobal,
+				FExecuteAction::CreateSP( this, &FBlueprintEditor::OnFindReferences, /*bSearchAllBlueprints=*/true, EGetFindReferenceSearchStringFlags::UseSearchSyntax ),
 				FCanExecuteAction::CreateSP( this, &FBlueprintEditor::CanFindReferences )
 				);
 
@@ -10188,34 +10208,28 @@ void FBlueprintEditor::SetPinVisibility(SGraphEditor::EPinVisibility Visibility)
 	OnSetPinVisibility.Broadcast(PinVisibility);
 }
 
-void FBlueprintEditor::OnFindReferences()
+void FBlueprintEditor::OnFindReferences(bool bSearchAllBlueprints, const EGetFindReferenceSearchStringFlags Flags)
 {
 	TSharedPtr<SGraphEditor> GraphEditor = FocusedGraphEdPtr.Pin();
 	if (GraphEditor.IsValid())
 	{
-		FString SearchTerm;
-
 		const FGraphPanelSelectionSet SelectedNodes = GraphEditor->GetSelectedNodes();
 		for (FGraphPanelSelectionSet::TConstIterator NodeIt(SelectedNodes); NodeIt; ++NodeIt)
 		{
-			UEdGraphNode* SelectedNode = Cast<UEdGraphNode>(*NodeIt);
-			if (SelectedNode != nullptr)
+			if (UEdGraphNode* SelectedNode = Cast<UEdGraphNode>(*NodeIt))
 			{
-				if (UK2Node_CustomEvent* CustomEvent = Cast<UK2Node_CustomEvent>(SelectedNode))
-				{
-					SearchTerm =  CustomEvent->CustomFunctionName.ToString();
-					if (!SearchTerm.IsEmpty())
-					{
-						SummonSearchUI(true, FString::Printf(TEXT("\"%s\""), *SearchTerm));
-					}
-				}
-				else
+				FString SearchTerm = SelectedNode->GetFindReferenceSearchString(Flags);
+				if (!SearchTerm.IsEmpty())
 				{
-					SearchTerm = SelectedNode->GetFindReferenceSearchString();
-					if (!SearchTerm.IsEmpty())
+					// If not using search syntax, surround the search term with quotes
+					if (!EnumHasAnyFlags(Flags, EGetFindReferenceSearchStringFlags::UseSearchSyntax))
 					{
-						SummonSearchUI(true, SearchTerm);
+						SearchTerm = FString::Printf(TEXT("\"%s\""), *SearchTerm);
 					}
+
+					// Start the search
+					const bool bSetFindWithinBlueprint = !bSearchAllBlueprints;
+					SummonSearchUI(bSetFindWithinBlueprint, SearchTerm);
 				}
 			}
 		}
diff --git a/Engine/Source/Editor/Kismet/Private/FindInBlueprintManager.cpp b/Engine/Source/Editor/Kismet/Private/FindInBlueprintManager.cpp
--- a/Engine/Source/Editor/Kismet/Private/FindInBlueprintManager.cpp
+++ b/Engine/Source/Editor/Kismet/Private/FindInBlueprintManager.cpp
@@ -89,6 +89,7 @@ const FText FFindInBlueprintSearchTags::FiB_Comment = LOCTEXT("Comment", "Commen
 const FText FFindInBlueprintSearchTags::FiB_Path = LOCTEXT("Path", "Path");
 const FText FFindInBlueprintSearchTags::FiB_ParentClass = LOCTEXT("ParentClass", "ParentClass");
 const FText FFindInBlueprintSearchTags::FiB_Interfaces = LOCTEXT("Interfaces", "Interfaces");
+const FText FFindInBlueprintSearchTags::FiB_FuncOriginClass = LOCTEXT("FuncOriginClass", "FuncOriginClass");
 
 const FText FFindInBlueprintSearchTags::FiB_Pins = LOCTEXT("Pins", "Pins");
 const FText FFindInBlueprintSearchTags::FiB_PinCategory = LOCTEXT("PinCategory", "PinCategory");
@@ -1089,8 +1090,10 @@ namespace BlueprintSearchMetaDataHelpers
 						InWriter->WriteArrayStart(FFindInBlueprintSearchTags::FiB_Pins);
 						for (const UEdGraphPin* Pin : Node->Pins)
 						{
-							// Hidden pins are not searchable
-							if (Pin->bHidden == false)
+							// Hidden pins are not searchable, except for 'self' pins which represent the target type.
+							// Indexing self pins, even if they are hidden, allows searching for BP function library
+							// function calls and nodes that automatically target self like parent function calls.
+							if (Pin->bHidden == false || Pin->PinName == UEdGraphSchema_K2::PN_Self)
 							{
 								InWriter->WriteObjectStart();
 								{
@@ -3306,25 +3309,9 @@ void FFindInBlueprintSearchManager::CacheAllAssets(TWeakPtr< SFindInBlueprints >
 			FText DialogTitle = LOCTEXT("ConfirmIndexAll_Title", "Indexing All");
 			FFormatNamedArguments Args;
 			Args.Add(TEXT("PackageCount"), UnindexedAssets.Num() + BlueprintsToUpdate.Num());
-
-			FText DialogDisplayText;
-
-			if (UnindexedAssets.Num() && BlueprintsToUpdate.Num())
-			{
-				Args.Add(TEXT("PackageCount"), UnindexedAssets.Num() + BlueprintsToUpdate.Num());
-				Args.Add(TEXT("UnindexedCount"), UnindexedAssets.Num());
-				Args.Add(TEXT("OutOfDateCount"), BlueprintsToUpdate.Num());
-				DialogDisplayText = FText::Format(LOCTEXT("CacheAllConfirmationMessage_UnindexedAndOutOfDate", "This process can take a long time and the editor may become unresponsive; there are {PackageCount} ({UnindexedCount} Unindexed/{OutOfDateCount} Out-of-Date) Blueprints to load.\n\nWould you like to checkout, load, and save all Blueprints to make this indexing permanent? Otherwise, all Blueprints will still be loaded but you will be required to re-index the next time you start the editor!"), Args);
-			}
-			else if (UnindexedAssets.Num() && BlueprintsToUpdate.Num() == 0)
-			{
-				DialogDisplayText = FText::Format(LOCTEXT("CacheAllConfirmationMessage_UnindexedOnly", "This process can take a long time and the editor may become unresponsive; there are {PackageCount} unindexed Blueprints to load.\n\nWould you like to checkout, load, and save all Blueprints to make this indexing permanent? Otherwise, all Blueprints will still be loaded but you will be required to re-index the next time you start the editor!"), Args);
-			}
-			else if (UnindexedAssets.Num() == 0 && BlueprintsToUpdate.Num())
-			{
-				DialogDisplayText = FText::Format(LOCTEXT("CacheAllConfirmationMessage_OutOfDateOnly", "This process can take a long time and the editor may become unresponsive; there are {PackageCount} out-of-date Blueprints to load.\n\nWould you like to checkout, load, and save all Blueprints to make this indexing permanent? Otherwise, all Blueprints will still be loaded but you will be required to re-index the next time you start the editor!"), Args);
-			}
-
+			Args.Add(TEXT("UnindexedCount"), UnindexedAssets.Num());
+			Args.Add(TEXT("OutOfDateCount"), BlueprintsToUpdate.Num());
+			const FText DialogDisplayText = FText::Format(LOCTEXT("CacheAllConfirmationMessage_UnindexedAndOutOfDate", "About to load {PackageCount} Blueprints ({UnindexedCount} unindexed/{OutOfDateCount} out-of-date). The editor may become unresponsive while these assets are loaded for indexing. Save your work before initiating this: broken assets and memory usage can affect editor stability. \n\nLoaded assets must be resaved to make this indexing permanent, otherwise their updated searchability is for this editor session only. Select 'Yes' to checkout, load and resave all Blueprints with an outdated index. Select 'No' to load them only."), Args);
 			const EAppReturnType::Type ReturnValue = FMessageDialog::Open(EAppMsgType::YesNoCancel, DialogDisplayText, DialogTitle);
 
 			// If Yes is chosen, checkout and save all Blueprints, if No is chosen, only load all Blueprints
@@ -3362,6 +3349,60 @@ void FFindInBlueprintSearchManager::CacheAllAssets(TWeakPtr< SFindInBlueprints >
 	}
 }
 
+void FFindInBlueprintSearchManager::ExportOutdatedAssetList()
+{
+	// Construct path for output text file
+	const FString FileLocation = FPaths::ConvertRelativePathToFull(FPaths::ProjectSavedDir());
+	const FString FullPath = FString::Printf(TEXT("%s/FindInBlueprints_OutdatedAssetList.txt"), *FileLocation);
+
+	if (FArchive* Ar = IFileManager::Get().CreateFileWriter(*FullPath))
+	{
+		// Write out all asset paths of unindexed blueprints
+		const FString UnindexedHeader = TEXT("Unindexed assets:\n");
+		Ar->Serialize(TCHAR_TO_ANSI(*UnindexedHeader), UnindexedHeader.Len());
+
+		for (const FSoftObjectPath& SoftObjPath : UnindexedAssets)
+		{
+			const FString UnindexedEntry = FString::Printf(TEXT("%s\n"), *SoftObjPath.ToString());
+			Ar->Serialize(TCHAR_TO_ANSI(*UnindexedEntry), UnindexedEntry.Len());
+		}
+
+		// Write out all asset paths of blueprints with out-of-data metadata
+		const FString OutOfDateHeader = TEXT("\nOut-of-date assets:\n");
+		Ar->Serialize(TCHAR_TO_ANSI(*OutOfDateHeader), OutOfDateHeader.Len());
+		for (FSearchData SearchData : SearchArray)
+		{
+			if ((SearchData.Value.Len() != 0 || SearchData.ImaginaryBlueprint.IsValid()) && SearchData.VersionInfo.FiBDataVersion < EFiBVersion::FIB_VER_LATEST)
+			{
+				const FString OutdatedEntry = FString::Printf(TEXT("%s\n"), *SearchData.AssetPath.ToString());
+				Ar->Serialize(TCHAR_TO_ANSI(*OutdatedEntry), OutdatedEntry.Len());
+			}
+		}
+
+		Ar->Close();
+		delete Ar;
+
+		// Log success message
+		FFormatNamedArguments Args;
+		Args.Add(TEXT("OutputPath"), FText::FromString(FullPath));
+		const FText ExportConfirmationText = FText::Format(LOCTEXT("ExportListConfirmationMessage", "Saved list of blueprints with out-of-date metadata to {OutputPath}"), Args);
+
+		UE_LOG(LogFindInBlueprint, Log, TEXT("%s"), *ExportConfirmationText.ToString());
+		FMessageDialog::Open(EAppMsgType::Ok, ExportConfirmationText);
+	}
+	else
+	{
+		// Log failure message
+		FFormatNamedArguments Args;
+		Args.Add(TEXT("OutputPath"), FText::FromString(FullPath));
+		const FText ExportFailureText = LOCTEXT("ExportListFailureMessage", "Failed to write to {OutputPath}");
+
+		UE_LOG(LogFindInBlueprint, Log, TEXT("%s"), *ExportFailureText.ToString());
+		FMessageDialog::Open(EAppMsgType::Ok, ExportFailureText);
+	}
+
+}
+
 void FFindInBlueprintSearchManager::CancelCacheAll(SFindInBlueprints* InFindInBlueprintWidget)
 {
 	if(IsCacheInProgress() && ((SourceCachingWidget.IsValid() && SourceCachingWidget.Pin().Get() == InFindInBlueprintWidget) || !SourceCachingWidget.IsValid()))
diff --git a/Engine/Source/Editor/Kismet/Private/FindInBlueprints.cpp b/Engine/Source/Editor/Kismet/Private/FindInBlueprints.cpp
--- a/Engine/Source/Editor/Kismet/Private/FindInBlueprints.cpp
+++ b/Engine/Source/Editor/Kismet/Private/FindInBlueprints.cpp
@@ -3,6 +3,7 @@
 
 #include "BlueprintEditor.h"
 #include "BlueprintEditorModule.h"
+#include "BlueprintEditorSettings.h"
 #include "BlueprintEditorTabs.h"
 #include "CoreGlobals.h"
 #include "Dom/JsonValue.h"
@@ -96,6 +97,66 @@ struct FSlateBrush;
 
 #define LOCTEXT_NAMESPACE "FindInBlueprints"
 
+KISMET_API UClass* FindInBlueprintsHelpers::GetFunctionOriginClass(const UFunction* Function)
+{
+	// Abort if invalid param
+	if (!Function)
+	{
+		return nullptr;
+	}
+
+	// Get outermost super function
+	while (const UFunction* SuperFunction = Function->GetSuperFunction())
+	{
+		Function = SuperFunction;
+	}
+
+	// Get that function's class
+	UClass* OwnerClass = Function->GetOwnerClass() && Function->GetOwnerClass()->GetAuthoritativeClass() ? Function->GetOwnerClass()->GetAuthoritativeClass() : Function->GetOwnerClass();
+
+	// Consider case where a blueprint implements an interface function
+	if (const UBlueprintGeneratedClass* BPGC = Cast<UBlueprintGeneratedClass>(OwnerClass))
+	{
+		const FName FunctionName = Function->GetFName();
+		for (const FImplementedInterface& Interface : BPGC->Interfaces)
+		{
+			if (!Interface.Class)
+			{
+				continue;
+			}
+
+			if (UFunction* InterfaceFunction = Interface.Class->FindFunctionByName(FunctionName))
+			{
+				if (InterfaceFunction->IsSignatureCompatibleWith(Function))
+				{
+					OwnerClass = Interface.Class;
+					break;
+				}
+			}
+		}
+	}
+
+	return OwnerClass;
+}
+
+bool FindInBlueprintsHelpers::ConstructSearchTermFromFunction(const UFunction* Function, FString& SearchTerm)
+{
+	if (!Function)
+	{
+		return false;
+	}
+
+	const UClass* FuncOriginClass = GetFunctionOriginClass(Function);
+	if (!FuncOriginClass)
+	{
+		return false;
+	}
+
+	const FString FunctionNativeName = Function->GetName();
+	const FString TargetTypeName = FuncOriginClass->GetPathName();
+	SearchTerm = FString::Printf(TEXT("Nodes(\"Native Name\"=+\"%s\" && (Pins(Name=Target && ObjectClass=+\"%s\") || FuncOriginClass=+\"%s\"))"), *FunctionNativeName, *TargetTypeName, *TargetTypeName);
+	return true;
+}
 
 FText FindInBlueprintsHelpers::AsFText(TSharedPtr< FJsonValue > InJsonValue, const TMap<int32, FText>& InLookupTable)
 {
@@ -907,14 +968,31 @@ void SFindInBlueprints::ConditionallyAddCacheBar()
 							[
 								SNew(SButton)
 								.Text(LOCTEXT("IndexAllBlueprints", "Index All"))
+								.IsEnabled( this, &SFindInBlueprints::CanCacheAllUnindexedBlueprints )
 								.OnClicked( this, &SFindInBlueprints::OnCacheAllUnindexedBlueprints )
 								.ToolTip(IDocumentation::Get()->CreateToolTip(
-									LOCTEXT("IndexAlLBlueprints_Tooltip", "Loads all non-indexed Blueprints and saves them with their search data. This can be a very slow process and the editor may become unresponsive."),
+									LOCTEXT("IndexAlLBlueprints_Tooltip", "Loads all Blueprints with an out-of-date index (search metadata) and resaves them with an up-to-date index. This can be a very slow process and the editor may become unresponsive. This action can be disabled via Blueprint Editor settings."),
 									NULL,
 									TEXT("Shared/Editors/BlueprintEditor"),
 									TEXT("FindInBlueprint_IndexAll"))
 								)
 							]
+
+							+ SHorizontalBox::Slot()
+							.AutoWidth()
+							.Padding(HPadding, 0.f, 0.f, 0.f)
+							.VAlign(VAlign_Center)
+							[
+								SNew(SButton)
+									.Text(LOCTEXT("IndexExportList", "Export Asset List"))
+									.OnClicked(this, &SFindInBlueprints::OnExportUnindexedAssetList)
+									.ToolTip(IDocumentation::Get()->CreateToolTip(
+										LOCTEXT("IndexExportList_Tooltip", "Exports a list of all Blueprints that have an out-of-date index (search metadata)."),
+										NULL,
+										TEXT("Shared/Editors/BlueprintEditor"),
+										TEXT("FindInBlueprint_IndexExportList_Tooltip"))
+									)
+							]
 						
 						]
 					]
@@ -1393,6 +1471,13 @@ FReply SFindInBlueprints::OnCacheAllUnindexedBlueprints()
 	return OnCacheAllBlueprints(CachingOptions);
 }
 
+FReply SFindInBlueprints::OnExportUnindexedAssetList()
+{
+	FFindInBlueprintSearchManager& FindInBlueprintManager = FFindInBlueprintSearchManager::Get();
+	FindInBlueprintManager.ExportOutdatedAssetList();
+	return FReply::Handled();
+}
+
 FReply SFindInBlueprints::OnCacheAllBlueprints(const FFindInBlueprintCachingOptions& InOptions)
 {
 	if(!FFindInBlueprintSearchManager::Get().IsCacheInProgress())
@@ -1553,12 +1638,16 @@ FText SFindInBlueprints::GetCacheBarStatusText() const
 	}
 	else
 	{
+		const FText IndexAllWarningText = LOCTEXT("IndexAllWarning", "Press \"Index All\" to load these assets right now. The editor may become unresponsive while these assets are loaded for indexing. Save your work before initiating this: broken assets and memory usage can affect editor stability. Alternatively, export the asset list to inspect which assets do not have optimal searchability.");
+		const FText IndexAllDisabledText = LOCTEXT("IndexAllDisabled", "Your Blueprint Editor settings disallow loading all these assets from this window. Export the asset list to inspect which assets do not have optimal searchability.");
+
 		const int32 UnindexedCount = FindInBlueprintManager.GetNumberUnindexedAssets();
 		Args.Add(TEXT("UnindexedCount"), UnindexedCount);
 		Args.Add(TEXT("OutOfDateCount"), OutOfDateWithLastSearchBPCount);
 		Args.Add(TEXT("Count"), UnindexedCount + OutOfDateWithLastSearchBPCount);
-
-		ReturnDisplayText = FText::Format(LOCTEXT("UncachedAssets", "Search incomplete. {Count} ({UnindexedCount} non-indexed/{OutOfDateCount} out-of-date) Blueprints need to be loaded and indexed!  Editor may become unresponsive while these assets are loaded for indexing. This may take some time!"), Args);
+		Args.Add(TEXT("Instruction"), CanCacheAllUnindexedBlueprints() ? IndexAllWarningText : IndexAllDisabledText);
+		
+		ReturnDisplayText = FText::Format(LOCTEXT("UncachedAssets", "Search incomplete: {Count} blueprints don't have an up-to-date index ({UnindexedCount} unindexed/{OutOfDateCount} out-of-date). These assets are searchable but some results may be missing. Load and resave these assets to improve their searchability. \n\n{Instruction}"), Args);
 
 		const int32 FailedToCacheCount = FindInBlueprintManager.GetFailedToCacheCount();
 		if (FailedToCacheCount > 0)
@@ -1583,6 +1672,11 @@ FText SFindInBlueprints::GetCacheBarCurrentAssetName() const
 	return FText::FromString(LastCachedAssetPath.ToString());
 }
 
+bool SFindInBlueprints::CanCacheAllUnindexedBlueprints() const
+{
+	return GetDefault<UBlueprintEditorSettings>()->bAllowIndexAllBlueprints;
+}
+
 void SFindInBlueprints::OnCacheStarted(EFiBCacheOpType InOpType, EFiBCacheOpFlags InOpFlags)
 {
 	const bool bShowProgress = EnumHasAnyFlags(InOpFlags, EFiBCacheOpFlags::ShowProgress);
diff --git a/Engine/Source/Editor/Kismet/Private/SMyBlueprint.cpp b/Engine/Source/Editor/Kismet/Private/SMyBlueprint.cpp
--- a/Engine/Source/Editor/Kismet/Private/SMyBlueprint.cpp
+++ b/Engine/Source/Editor/Kismet/Private/SMyBlueprint.cpp
@@ -355,9 +355,34 @@ void SMyBlueprint::Construct(const FArguments& InArgs, TWeakPtr<FBlueprintEditor
 			FExecuteAction::CreateSP(this, &SMyBlueprint::OnImplementFunction),
 			FCanExecuteAction(), FIsActionChecked(),
 			FIsActionButtonVisible::CreateSP(this, &SMyBlueprint::CanImplementFunction) );
-	
+		
 		CommandList->MapAction( FGraphEditorCommands::Get().FindReferences,
-			FExecuteAction::CreateSP(this, &SMyBlueprint::OnFindReference),
+			FExecuteAction::CreateSP(this, &SMyBlueprint::OnFindReference, /*bSearchAllBlueprints=*/false, EGetFindReferenceSearchStringFlags::None),
+			FCanExecuteAction(),
+			FIsActionChecked(),
+			FIsActionButtonVisible::CreateSP(this, &SMyBlueprint::CanFindReference) );
+		
+		CommandList->MapAction( FGraphEditorCommands::Get().FindReferencesByNameLocal,
+			FExecuteAction::CreateSP(this, &SMyBlueprint::OnFindReference, /*bSearchAllBlueprints=*/false, EGetFindReferenceSearchStringFlags::None),
+			FCanExecuteAction(),
+			FIsActionChecked(),
+			FIsActionButtonVisible::CreateSP(this, &SMyBlueprint::CanFindReference) );
+
+		
+		CommandList->MapAction( FGraphEditorCommands::Get().FindReferencesByNameGlobal,
+			FExecuteAction::CreateSP(this, &SMyBlueprint::OnFindReference, /*bSearchAllBlueprints=*/true, EGetFindReferenceSearchStringFlags::None),
+			FCanExecuteAction(),
+			FIsActionChecked(),
+			FIsActionButtonVisible::CreateSP(this, &SMyBlueprint::CanFindReference) );
+		
+		CommandList->MapAction( FGraphEditorCommands::Get().FindReferencesByClassMemberLocal,
+			FExecuteAction::CreateSP(this, &SMyBlueprint::OnFindReference, /*bSearchAllBlueprints=*/false, EGetFindReferenceSearchStringFlags::UseSearchSyntax),
+			FCanExecuteAction(),
+			FIsActionChecked(),
+			FIsActionButtonVisible::CreateSP(this, &SMyBlueprint::CanFindReference) );
+		
+		CommandList->MapAction( FGraphEditorCommands::Get().FindReferencesByClassMemberGlobal,
+			FExecuteAction::CreateSP(this, &SMyBlueprint::OnFindReference, /*bSearchAllBlueprints=*/true, EGetFindReferenceSearchStringFlags::UseSearchSyntax),
 			FCanExecuteAction(),
 			FIsActionChecked(),
 			FIsActionButtonVisible::CreateSP(this, &SMyBlueprint::CanFindReference) );
@@ -2437,6 +2462,11 @@ TSharedPtr<SWidget> SMyBlueprint::OnContextMenuOpening()
 	// Check if the selected action is valid for a context menu
 	if (SelectionHasContextMenu())
 	{
+		FEdGraphSchemaAction_K2Var* Var = SelectionAsVar();
+		FEdGraphSchemaAction_K2Graph* Graph = SelectionAsGraph();
+		FEdGraphSchemaAction_K2Event* Event = SelectionAsEvent();
+		const bool bExpandFindReferences = Graph || Event || Var;
+		
 		MenuBuilder.BeginSection("BasicOperations");
 		{
 			MenuBuilder.AddMenuEntry(FMyBlueprintCommands::Get().OpenGraph);
@@ -2446,7 +2476,26 @@ TSharedPtr<SWidget> SMyBlueprint::OnContextMenuOpening()
 			MenuBuilder.AddMenuEntry(FMyBlueprintCommands::Get().FocusNodeInNewTab);
 			MenuBuilder.AddMenuEntry(FGenericCommands::Get().Rename, NAME_None, LOCTEXT("Rename", "Rename"), LOCTEXT("Rename_Tooltip", "Renames this function or variable from blueprint.") );
 			MenuBuilder.AddMenuEntry(FMyBlueprintCommands::Get().ImplementFunction);
-			MenuBuilder.AddMenuEntry(FGraphEditorCommands::Get().FindReferences);
+
+			// Depending on context, FindReferences can be a button or an expandable menu. For example, the context menu
+			// for functions now lets you choose whether to do search by-name (fast) or by-function (smart).
+			if (!bExpandFindReferences)
+			{
+				// No expandable menu: display the simple 'Find References' action
+				MenuBuilder.AddMenuEntry(FGraphEditorCommands::Get().FindReferences);
+			}
+			else
+			{
+				// Insert "Find References" sub-menu here
+				MenuBuilder.AddSubMenu(
+					LOCTEXT("FindReferences_Label", "Find References"),
+					LOCTEXT("FindReferences_Tooltip", "Options for finding references to class members"),
+					FNewMenuDelegate::CreateStatic(&FGraphEditorCommands::BuildFindReferencesMenu),
+					false,
+					FSlateIcon()
+				);
+			}
+			
 			MenuBuilder.AddMenuEntry(FGraphEditorCommands::Get().FindAndReplaceReferences);
 			MenuBuilder.AddMenuEntry(FMyBlueprintCommands::Get().GotoNativeVarDefinition);
 			MenuBuilder.AddMenuEntry(FGenericCommands::Get().Cut);
@@ -2458,10 +2507,6 @@ TSharedPtr<SWidget> SMyBlueprint::OnContextMenuOpening()
 		}
 		MenuBuilder.EndSection();
 
-		FEdGraphSchemaAction_K2Var* Var = SelectionAsVar();
-		FEdGraphSchemaAction_K2Graph* Graph = SelectionAsGraph();
-		FEdGraphSchemaAction_K2Event* Event = SelectionAsEvent();
-
 		if ( Var && BlueprintEditorPtr.IsValid() && FBlueprintEditorUtils::DoesSupportEventGraphs(GetBlueprintObj()) )
 		{
 			FObjectProperty* ComponentProperty = CastField<FObjectProperty>(Var->GetProperty());
@@ -2848,7 +2893,7 @@ bool SMyBlueprint::IsImplementationDesiredAsFunction(const UFunction* OverrideFu
 	return false;
 }
 
-void SMyBlueprint::OnFindReference()
+void SMyBlueprint::OnFindReference(bool bSearchAllBlueprints, const EGetFindReferenceSearchStringFlags Flags)
 {
 	FString SearchTerm;
 	for(const UEdGraph* UberGraph : GetBlueprintObj()->UbergraphPages)
@@ -2876,21 +2921,48 @@ void SMyBlueprint::OnFindReference()
 
 	if (FEdGraphSchemaAction_K2Graph* GraphAction = SelectionAsGraph())
 	{
-		SearchTerm = GraphAction->FuncName.ToString();
+		bool bSearchTermGenerated = false;
+		if (EnumHasAnyFlags(Flags, EGetFindReferenceSearchStringFlags::UseSearchSyntax))
+		{
+			// Attempt to resolve function, then attempt to generate search term from it
+			if (const UFunction* Function = GraphAction->GetFunction())
+			{
+				if (FindInBlueprintsHelpers::ConstructSearchTermFromFunction(Function, SearchTerm))
+				{
+					bSearchTermGenerated = true;
+					bUseQuotes = false;	
+				}
+			}
+		}
+
+		// If a search term was not generated yet, just search by name. Either the graph doesn't represent a
+		// function or we failed to resolve a dependency for generating a query.
+		if (!bSearchTermGenerated)
+		{
+			SearchTerm = GraphAction->FuncName.ToString();
+		}
 	}
 	else if (FEdGraphSchemaAction_K2Var* VarAction = SelectionAsVar())
 	{
-		if(FProperty* Property = VarAction->GetProperty())
+		bool bSearchTermGenerated = false;
+		if (EnumHasAnyFlags(Flags, EGetFindReferenceSearchStringFlags::UseSearchSyntax))
 		{
-			FMemberReference MemberReference;
-			MemberReference.SetFromField<FProperty>(Property, true, Property->GetOwnerClass());
-			SearchTerm = MemberReference.GetReferenceSearchString(Property->GetOwnerClass());
+			// Attempt to resolve property, then generate search term from it
+			if (const FProperty* Property = VarAction->GetProperty())
+			{
+				FMemberReference MemberReference;
+				MemberReference.SetFromField<FProperty>(Property, true, Property->GetOwnerClass());
+				SearchTerm = MemberReference.GetReferenceSearchString(Property->GetOwnerClass());
+				bSearchTermGenerated = true;
+				bUseQuotes = false;
+			}
 		}
-		else
+
+		// If a search term was not generated yet, just search by variable name
+		if (!bSearchTermGenerated)
 		{
 			SearchTerm = VarAction->GetVariableName().ToString();
 		}
-		bUseQuotes = false;
 	}
 	else if (FEdGraphSchemaAction_K2LocalVar* LocalVarAction = SelectionAsLocalVar())
 	{
@@ -2911,7 +2983,21 @@ void SMyBlueprint::OnFindReference()
 	}
 	else if (FEdGraphSchemaAction_K2Event* EventAction = SelectionAsEvent())
 	{
-		SearchTerm = EventAction->NodeTemplate->GetFindReferenceSearchString();
+		bool bSearchTermGenerated = false;
+		if (EnumHasAnyFlags(Flags, EGetFindReferenceSearchStringFlags::UseSearchSyntax))
+		{
+			if (const UK2Node_Event* EventNode = Cast<UK2Node_Event>(EventAction->NodeTemplate))
+			{
+				SearchTerm = EventNode->GetFindReferenceSearchString(Flags);
+				bSearchTermGenerated = true;
+				bUseQuotes = false;
+			}
+		}
+
+		if (!bSearchTermGenerated)
+		{
+			SearchTerm = EventAction->NodeTemplate->GetFindReferenceSearchString(EGetFindReferenceSearchStringFlags::None);
+		}
 	}
 	else if (FEdGraphSchemaAction_K2InputAction* InputAction = SelectionAsInputAction())
 	{
@@ -2920,13 +3006,15 @@ void SMyBlueprint::OnFindReference()
 			InputAction->GetMenuDescription().ToString();
 	}
 
-	if(!SearchTerm.IsEmpty())
+	if (!SearchTerm.IsEmpty())
 	{
 		if (bUseQuotes)
 		{
 			SearchTerm = FString::Printf(TEXT("\"%s\""), *SearchTerm);
 		}
-		BlueprintEditorPtr.Pin()->SummonSearchUI(true, SearchTerm);
+		
+		const bool bSetFindWithinBlueprint = !bSearchAllBlueprints;
+		BlueprintEditorPtr.Pin()->SummonSearchUI(bSetFindWithinBlueprint, SearchTerm);
 	}
 }
 
diff --git a/Engine/Source/Editor/Kismet/Private/SSCSEditor.cpp b/Engine/Source/Editor/Kismet/Private/SSCSEditor.cpp
--- a/Engine/Source/Editor/Kismet/Private/SSCSEditor.cpp
+++ b/Engine/Source/Editor/Kismet/Private/SSCSEditor.cpp
@@ -3942,7 +3942,23 @@ void SSCSEditor::Construct( const FArguments& InArgs )
 		);
 
 	CommandList->MapAction( FGraphEditorCommands::Get().FindReferences,
-		FUIAction( FExecuteAction::CreateSP( this, &SSCSEditor::OnFindReferences ) )
+		FUIAction( FExecuteAction::CreateSP( this, &SSCSEditor::OnFindReferences, false, EGetFindReferenceSearchStringFlags::UseSearchSyntax ) )
+	);
+
+	CommandList->MapAction( FGraphEditorCommands::Get().FindReferencesByNameLocal,
+		FUIAction( FExecuteAction::CreateSP( this, &SSCSEditor::OnFindReferences, false, EGetFindReferenceSearchStringFlags::None ) )
+	);
+	
+	CommandList->MapAction( FGraphEditorCommands::Get().FindReferencesByNameGlobal,
+		FUIAction( FExecuteAction::CreateSP( this, &SSCSEditor::OnFindReferences, true, EGetFindReferenceSearchStringFlags::None ) )
+	);
+	
+	CommandList->MapAction( FGraphEditorCommands::Get().FindReferencesByClassMemberLocal,
+		FUIAction( FExecuteAction::CreateSP( this, &SSCSEditor::OnFindReferences, false, EGetFindReferenceSearchStringFlags::UseSearchSyntax ) )
+	);
+
+	CommandList->MapAction( FGraphEditorCommands::Get().FindReferencesByClassMemberGlobal,
+		FUIAction( FExecuteAction::CreateSP( this, &SSCSEditor::OnFindReferences, true, EGetFindReferenceSearchStringFlags::UseSearchSyntax ) )
 	);
 
 	FSlateBrush const* MobilityHeaderBrush = FAppStyle::GetBrush(TEXT("ClassIcon.ComponentMobilityHeaderIcon"));
@@ -4359,7 +4375,13 @@ void SSCSEditor::PopulateContextMenu(UToolMenu* Menu)
 						FToolMenuSection& BlueprintSCSSection = Menu->AddSection("BlueprintSCS");
 						if (SelectedItems.Num() == 1)
 						{
-							BlueprintSCSSection.AddMenuEntry(FGraphEditorCommands::Get().FindReferences);
+							// Expandable menu: insert sub-menu here
+                            BlueprintSCSSection.AddSubMenu(
+                            	FName("FindReferenceSubMenu"),
+                            	LOCTEXT("FindReferences_Label", "Find References"),
+                            	LOCTEXT("FindReferences_Tooltip", "Options for finding references to class members"),
+                            	FNewToolMenuChoice(FNewMenuDelegate::CreateStatic(&FGraphEditorCommands::BuildFindReferencesMenu))
+                            );
 						}
 
 						// Create an "Add Event" option in the context menu only if we can edit
@@ -4578,7 +4600,7 @@ void SSCSEditor::ViewEvent(UBlueprint* Blueprint, const FName EventName, const F
 	}
 }
 
-void SSCSEditor::OnFindReferences()
+void SSCSEditor::OnFindReferences(bool bSearchAllBlueprints, const EGetFindReferenceSearchStringFlags Flags)
 {
 	TArray<FSCSEditorTreeNodePtrType> SelectedNodes = SCSTreeWidget->GetSelectedItems();
 	if (SelectedNodes.Num() == 1)
@@ -4590,9 +4612,10 @@ void SSCSEditor::OnFindReferences()
 
 			FMemberReference MemberReference;
 			MemberReference.SetSelfMember(*VariableName);
-			const FString SearchTerm = MemberReference.GetReferenceSearchString(GetBlueprint()->SkeletonGeneratedClass);
+			const FString SearchTerm = EnumHasAnyFlags(Flags, EGetFindReferenceSearchStringFlags::UseSearchSyntax) ? MemberReference.GetReferenceSearchString(GetBlueprint()->SkeletonGeneratedClass) : FString::Printf(TEXT("\"%s\""), *VariableName);;
 
 			TSharedRef<IBlueprintEditor> BlueprintEditor = StaticCastSharedRef<IBlueprintEditor>(FoundAssetEditor.ToSharedRef());
+			const bool bSetFindWithinBlueprint = !bSearchAllBlueprints;
 			BlueprintEditor->SummonSearchUI(true, SearchTerm);
 		}
 	}
diff --git a/Engine/Source/Editor/Kismet/Public/BlueprintEditor.h b/Engine/Source/Editor/Kismet/Public/BlueprintEditor.h
--- a/Engine/Source/Editor/Kismet/Public/BlueprintEditor.h
+++ b/Engine/Source/Editor/Kismet/Public/BlueprintEditor.h
@@ -1305,7 +1305,7 @@ private:
 	void NavigateTab(FDocumentTracker::EOpenDocumentCause InCause);
 
 	/** Find all references of the selected node. */
-	void OnFindReferences();
+	void OnFindReferences(bool bSearchAllBlueprints, const EGetFindReferenceSearchStringFlags Flags);
 
 	/** Checks if we can currently find all references of the node selection. */
 	bool CanFindReferences();
diff --git a/Engine/Source/Editor/Kismet/Public/FindInBlueprintManager.h b/Engine/Source/Editor/Kismet/Public/FindInBlueprintManager.h
--- a/Engine/Source/Editor/Kismet/Public/FindInBlueprintManager.h
+++ b/Engine/Source/Editor/Kismet/Public/FindInBlueprintManager.h
@@ -115,6 +115,8 @@ struct KISMET_API FFindInBlueprintSearchTags
 	static const FText FiB_ParentClass;
 	/** Interfaces tag */
 	static const FText FiB_Interfaces;
+	/** Class that originally defined the function tag */
+	static const FText FiB_FuncOriginClass;
 
 	/** Pin type tags */
 
@@ -152,6 +154,7 @@ enum EFiBVersion : int
 	FIB_VER_BASE = 0, // All Blueprints prior to versioning will automatically be assumed to be at 0 if they have FiB data collected
 	FIB_VER_VARIABLE_REFERENCE, // Variable references (FMemberReference) is collected in FiB
 	FIB_VER_INTERFACE_GRAPHS, // Implemented Interface Graphs is collected in FiB
+	FIB_VER_FUNC_CALL_SITES, // Hidden target pins and function origin class are collected in FiB for improved function call site searchability
 
 	// -----<new versions can be added before this line>-------------------------------------------------
 	FIB_VER_PLUS_ONE,
@@ -614,7 +617,12 @@ public:
 	 * @param InCachingOptions				Options to configure the caching task
 	 */
 	void CacheAllAssets(TWeakPtr< class SFindInBlueprints > InSourceWidget, const FFindInBlueprintCachingOptions& InCachingOptions);
-	
+
+	/**
+	 * Exports a list of all unindexed assets to Saved/FindInBlueprints_OutdatedAssetList.txt
+	 */
+	void ExportOutdatedAssetList();
+
 	/**
 	 * Starts the actual caching process
 	 *
diff --git a/Engine/Source/Editor/Kismet/Public/FindInBlueprints.h b/Engine/Source/Editor/Kismet/Public/FindInBlueprints.h
--- a/Engine/Source/Editor/Kismet/Public/FindInBlueprints.h
+++ b/Engine/Source/Editor/Kismet/Public/FindInBlueprints.h
@@ -75,6 +75,12 @@ namespace FindInBlueprintsHelpers
 		}
 	};
 
+	/** Utility function to find the ancestor class or interface from which a function is inherited. */
+	KISMET_API UClass* GetFunctionOriginClass(const UFunction* Function);
+
+	/** Constructs a search term for a function using Find-in-Blueprints search syntax */
+	KISMET_API bool ConstructSearchTermFromFunction(const UFunction* Function, FString& SearchTerm);
+
 	static uint32 GetTypeHash(const FindInBlueprintsHelpers::FSimpleFTextKeyStorage& InObject)
 	{
 		return GetTypeHash(InObject.Text.BuildSourceString());
@@ -364,9 +370,15 @@ private:
 	/** Callback to return the current asset name during a cache operation */
 	FText GetCacheBarCurrentAssetName() const;
 
+	/** Whether user is allowed to initiate loading and indexing all blueprints with out-of-date metadata */
+	bool CanCacheAllUnindexedBlueprints() const;
+
 	/** Callback to cache all unindexed Blueprints */
 	FReply OnCacheAllUnindexedBlueprints();
 
+	/** Callback to export a list of all blueprints that need reindexing */
+	FReply OnExportUnindexedAssetList();
+
 	/** Callback to cache all Blueprints according to the given options */
 	FReply OnCacheAllBlueprints(const FFindInBlueprintCachingOptions& InOptions);
 
diff --git a/Engine/Source/Editor/Kismet/Public/SMyBlueprint.h b/Engine/Source/Editor/Kismet/Public/SMyBlueprint.h
--- a/Engine/Source/Editor/Kismet/Public/SMyBlueprint.h
+++ b/Engine/Source/Editor/Kismet/Public/SMyBlueprint.h
@@ -267,7 +267,7 @@ private:
 	void ImplementFunction(TSharedPtr<FEdGraphSchemaAction_K2Graph> GraphAction);
 	void ImplementFunction(FEdGraphSchemaAction_K2Graph* GraphAction);
 	bool CanImplementFunction() const;
-	void OnFindReference();
+	void OnFindReference(bool bSearchAllBlueprints, const EGetFindReferenceSearchStringFlags Flags);
 	bool CanFindReference() const;
 	void OnFindAndReplaceReference();
 	bool CanFindAndReplaceReference() const;
diff --git a/Engine/Source/Editor/Kismet/Public/SSCSEditor.h b/Engine/Source/Editor/Kismet/Public/SSCSEditor.h
--- a/Engine/Source/Editor/Kismet/Public/SSCSEditor.h
+++ b/Engine/Source/Editor/Kismet/Public/SSCSEditor.h
@@ -1151,7 +1151,7 @@ public:
 	bool CanDeleteNodes() const;
 
 	/** Callbacks to find references of the selected component */
-	void OnFindReferences();
+	void OnFindReferences(bool bSearchAllBlueprints, const EGetFindReferenceSearchStringFlags Flags);
 
 	/** Removes an existing component node from the tree */
 	void RemoveComponentNode(FSCSEditorTreeNodePtrType InNodePtr);
diff --git a/Engine/Source/Editor/SubobjectEditor/Private/SSubobjectBlueprintEditor.cpp b/Engine/Source/Editor/SubobjectEditor/Private/SSubobjectBlueprintEditor.cpp
--- a/Engine/Source/Editor/SubobjectEditor/Private/SSubobjectBlueprintEditor.cpp
+++ b/Engine/Source/Editor/SubobjectEditor/Private/SSubobjectBlueprintEditor.cpp
@@ -536,7 +536,12 @@ void SSubobjectBlueprintEditor::PopulateContextMenuImpl(UToolMenu* InMenu, TArra
 	FToolMenuSection& BlueprintSCSSection = InMenu->AddSection("BlueprintSCS");
 	if (InSelectedItems.Num() == 1)
 	{
-		BlueprintSCSSection.AddMenuEntry(FGraphEditorCommands::Get().FindReferences);
+		BlueprintSCSSection.AddSubMenu(
+			FName("FindReferenceSubMenu"),
+			LOCTEXT("FindReferences_Label", "Find References"),
+			LOCTEXT("FindReferences_Tooltip", "Options for finding references to class members"),
+			FNewToolMenuChoice(FNewMenuDelegate::CreateStatic(&FGraphEditorCommands::BuildFindReferencesMenu))
+		);
 	}
 
 	// Create an "Add Event" option in the context menu only if we can edit
diff --git a/Engine/Source/Editor/SubobjectEditor/Private/SSubobjectEditor.cpp b/Engine/Source/Editor/SubobjectEditor/Private/SSubobjectEditor.cpp
--- a/Engine/Source/Editor/SubobjectEditor/Private/SSubobjectEditor.cpp
+++ b/Engine/Source/Editor/SubobjectEditor/Private/SSubobjectEditor.cpp
@@ -2148,7 +2148,23 @@ void SSubobjectEditor::CreateCommandList()
 	);
 
 	CommandList->MapAction(FGraphEditorCommands::Get().GetFindReferences(),
-		FUIAction(FExecuteAction::CreateSP(this, &SSubobjectEditor::OnFindReferences))
+		FUIAction(FExecuteAction::CreateSP(this, &SSubobjectEditor::OnFindReferences, false, EGetFindReferenceSearchStringFlags::UseSearchSyntax ) )
+	);
+	
+	CommandList->MapAction( FGraphEditorCommands::Get().FindReferencesByNameLocal,
+		FUIAction( FExecuteAction::CreateSP( this, &SSubobjectEditor::OnFindReferences, false, EGetFindReferenceSearchStringFlags::None) )
+	);
+	
+	CommandList->MapAction( FGraphEditorCommands::Get().FindReferencesByNameGlobal,
+		FUIAction( FExecuteAction::CreateSP( this, &SSubobjectEditor::OnFindReferences, true, EGetFindReferenceSearchStringFlags::None) )
+	);
+	
+	CommandList->MapAction( FGraphEditorCommands::Get().FindReferencesByClassMemberLocal,
+		FUIAction( FExecuteAction::CreateSP( this, &SSubobjectEditor::OnFindReferences, false, EGetFindReferenceSearchStringFlags::UseSearchSyntax) )
+	);
+
+	CommandList->MapAction( FGraphEditorCommands::Get().FindReferencesByClassMemberGlobal,
+		FUIAction( FExecuteAction::CreateSP( this, &SSubobjectEditor::OnFindReferences, true, EGetFindReferenceSearchStringFlags::UseSearchSyntax) )
 	);
 }
 
@@ -2310,7 +2326,7 @@ void SSubobjectEditor::HandleItemDoubleClicked(FSubobjectEditorTreeNodePtrType I
 	OnItemDoubleClicked.ExecuteIfBound(InItem);
 }
 
-void SSubobjectEditor::OnFindReferences()
+void SSubobjectEditor::OnFindReferences(bool bSearchAllBlueprints, const EGetFindReferenceSearchStringFlags Flags)
 {
 	TArray<FSubobjectEditorTreeNodePtrType> SelectedNodes = TreeWidget->GetSelectedItems();
 	if (SelectedNodes.Num() == 1)
@@ -2323,11 +2339,13 @@ void SSubobjectEditor::OnFindReferences()
 
 			FMemberReference MemberReference;
 			MemberReference.SetSelfMember(*VariableName);
-			const FString SearchTerm = MemberReference.GetReferenceSearchString(GetBlueprint()->SkeletonGeneratedClass);
+			const FString SearchTerm = EnumHasAnyFlags(Flags, EGetFindReferenceSearchStringFlags::UseSearchSyntax) ? MemberReference.GetReferenceSearchString(GetBlueprint()->SkeletonGeneratedClass) : FString::Printf(TEXT("\"%s\""), *VariableName);
 
 			TSharedRef<IBlueprintEditor> BlueprintEditor = StaticCastSharedRef<IBlueprintEditor>(
 				FoundAssetEditor.ToSharedRef());
-			BlueprintEditor->SummonSearchUI(true, SearchTerm);
+			
+			const bool bSetFindWithinBlueprint = !bSearchAllBlueprints;
+			BlueprintEditor->SummonSearchUI(bSetFindWithinBlueprint, SearchTerm);
 		}
 	}
 }
diff --git a/Engine/Source/Editor/SubobjectEditor/Public/SSubobjectEditor.h b/Engine/Source/Editor/SubobjectEditor/Public/SSubobjectEditor.h
--- a/Engine/Source/Editor/SubobjectEditor/Public/SSubobjectEditor.h
+++ b/Engine/Source/Editor/SubobjectEditor/Public/SSubobjectEditor.h
@@ -446,8 +446,8 @@ protected:
 
 	/** Callback when a component item is double clicked. */
 	void HandleItemDoubleClicked(FSubobjectEditorTreeNodePtrType InItem);
-	
-	void OnFindReferences();
+
+	void OnFindReferences(bool bSearchAllBlueprints, const EGetFindReferenceSearchStringFlags Flags);
 
 	/** @return The visibility of the components filter box */
 	EVisibility GetComponentsFilterBoxVisibility() const;
diff --git a/Engine/Source/Editor/UMGEditor/Private/WidgetBlueprintEditor.cpp b/Engine/Source/Editor/UMGEditor/Private/WidgetBlueprintEditor.cpp
--- a/Engine/Source/Editor/UMGEditor/Private/WidgetBlueprintEditor.cpp
+++ b/Engine/Source/Editor/UMGEditor/Private/WidgetBlueprintEditor.cpp
@@ -205,7 +205,23 @@ void FWidgetBlueprintEditor::InitWidgetBlueprintEditor(const EToolkitMode::Type
 		);
 
 	DesignerCommandList->MapAction(FGraphEditorCommands::Get().FindReferences,
-		FExecuteAction::CreateSP(this, &FWidgetBlueprintEditor::OnFindWidgetReferences),
+		FExecuteAction::CreateSP(this, &FWidgetBlueprintEditor::OnFindWidgetReferences, false, EGetFindReferenceSearchStringFlags::UseSearchSyntax),
+		FCanExecuteAction::CreateSP(this, &FWidgetBlueprintEditor::CanFindWidgetReferences));
+	
+	DesignerCommandList->MapAction(FGraphEditorCommands::Get().FindReferencesByNameLocal,
+		FExecuteAction::CreateSP(this, &FWidgetBlueprintEditor::OnFindWidgetReferences, false, EGetFindReferenceSearchStringFlags::None),
+		FCanExecuteAction::CreateSP(this, &FWidgetBlueprintEditor::CanFindWidgetReferences));
+
+	DesignerCommandList->MapAction(FGraphEditorCommands::Get().FindReferencesByNameGlobal,
+		FExecuteAction::CreateSP(this, &FWidgetBlueprintEditor::OnFindWidgetReferences, true, EGetFindReferenceSearchStringFlags::None),
+		FCanExecuteAction::CreateSP(this, &FWidgetBlueprintEditor::CanFindWidgetReferences));
+	
+	DesignerCommandList->MapAction(FGraphEditorCommands::Get().FindReferencesByClassMemberLocal,
+		FExecuteAction::CreateSP(this, &FWidgetBlueprintEditor::OnFindWidgetReferences, false, EGetFindReferenceSearchStringFlags::UseSearchSyntax),
+		FCanExecuteAction::CreateSP(this, &FWidgetBlueprintEditor::CanFindWidgetReferences));
+	
+	DesignerCommandList->MapAction(FGraphEditorCommands::Get().FindReferencesByClassMemberGlobal,
+		FExecuteAction::CreateSP(this, &FWidgetBlueprintEditor::OnFindWidgetReferences, true, EGetFindReferenceSearchStringFlags::UseSearchSyntax),
 		FCanExecuteAction::CreateSP(this, &FWidgetBlueprintEditor::CanFindWidgetReferences));
 
 	TSharedPtr<class IToolkitHost> PinnedToolkitHost = ToolkitHost.Pin();
@@ -901,17 +917,19 @@ void FWidgetBlueprintEditor::DuplicateSelectedWidgets()
 	SelectWidgets(DuplicatedWidgetRefs, false);
 }
 
-void FWidgetBlueprintEditor::OnFindWidgetReferences()
+void FWidgetBlueprintEditor::OnFindWidgetReferences(bool bSearchAllBlueprints, const EGetFindReferenceSearchStringFlags Flags)
 {
 	FWidgetReference WidgetReference = *GetSelectedWidgets().CreateConstIterator();
 	const FString VariableName = WidgetReference.GetTemplate()->GetName();
 
 	FMemberReference MemberReference;
 	MemberReference.SetSelfMember(*VariableName);
-	const FString SearchTerm = MemberReference.GetReferenceSearchString(GetBlueprintObj()->SkeletonGeneratedClass);
+	const FString SearchTerm = EnumHasAnyFlags(Flags, EGetFindReferenceSearchStringFlags::UseSearchSyntax) ? MemberReference.GetReferenceSearchString(GetBlueprintObj()->SkeletonGeneratedClass) : FString::Printf(TEXT("\"%s\""), *VariableName);
 
 	SetCurrentMode(FWidgetBlueprintApplicationModes::GraphMode);
-	SummonSearchUI(true, SearchTerm);
+	
+	const bool bSetFindWithinBlueprint = !bSearchAllBlueprints;
+	SummonSearchUI(bSetFindWithinBlueprint, SearchTerm);
 }
 
 bool FWidgetBlueprintEditor::CanFindWidgetReferences() const
diff --git a/Engine/Source/Editor/UMGEditor/Private/WidgetBlueprintEditorUtils.cpp b/Engine/Source/Editor/UMGEditor/Private/WidgetBlueprintEditorUtils.cpp
--- a/Engine/Source/Editor/UMGEditor/Private/WidgetBlueprintEditorUtils.cpp
+++ b/Engine/Source/Editor/UMGEditor/Private/WidgetBlueprintEditorUtils.cpp
@@ -450,7 +450,15 @@ void FWidgetBlueprintEditorUtils::CreateWidgetContextMenu(FMenuBuilder& MenuBuil
 			MenuBuilder.AddMenuEntry(FGenericCommands::Get().Paste);
 			MenuBuilder.AddMenuEntry(FGenericCommands::Get().Duplicate);
 			MenuBuilder.AddMenuEntry(FGenericCommands::Get().Delete);
-			MenuBuilder.AddMenuEntry(FGraphEditorCommands::Get().FindReferences);
+			
+			// Insert "Find References" sub-menu here
+			MenuBuilder.AddSubMenu(
+				LOCTEXT("FindReferences_Label", "Find References"),
+				LOCTEXT("FindReferences_Tooltip", "Options for finding references to class members"),
+				FNewMenuDelegate::CreateStatic(&FGraphEditorCommands::BuildFindReferencesMenu),
+				false,
+				FSlateIcon()
+			);
 		}
 		MenuBuilder.PopCommandList();
 
diff --git a/Engine/Source/Editor/UMGEditor/Public/WidgetBlueprintEditor.h b/Engine/Source/Editor/UMGEditor/Public/WidgetBlueprintEditor.h
--- a/Engine/Source/Editor/UMGEditor/Public/WidgetBlueprintEditor.h
+++ b/Engine/Source/Editor/UMGEditor/Public/WidgetBlueprintEditor.h
@@ -343,7 +343,7 @@ private:
 	bool CanDuplicateSelectedWidgets();
 	void DuplicateSelectedWidgets();
 
-	void OnFindWidgetReferences();
+	void OnFindWidgetReferences(bool bSearchAllBlueprints, const EGetFindReferenceSearchStringFlags Flags);
 	bool CanFindWidgetReferences() const;
 
 	/** Is creating a native base class for the current widget blueprint allowed */
diff --git a/Engine/Source/Editor/UnrealEd/Private/EditorActor.cpp b/Engine/Source/Editor/UnrealEd/Private/EditorActor.cpp
--- a/Engine/Source/Editor/UnrealEd/Private/EditorActor.cpp
+++ b/Engine/Source/Editor/UnrealEd/Private/EditorActor.cpp
@@ -1058,7 +1058,7 @@ bool UUnrealEdEngine::DeleteActors(const TArray<AActor*>& InActorsToDelete, UWor
 
 				for (UK2Node* Node : ReferencedToActorsFromLevelScriptArray)
 				{
-					LevelScriptReferenceString += Node->GetFindReferenceSearchString();
+					LevelScriptReferenceString += Node->GetFindReferenceSearchString(EGetFindReferenceSearchStringFlags::None);
 
 					if (bReferencedByLevelScript && bReferencedByActor)
 					{
diff --git a/Engine/Source/Editor/UnrealEd/Private/Kismet2/BlueprintEditorUtils.cpp b/Engine/Source/Editor/UnrealEd/Private/Kismet2/BlueprintEditorUtils.cpp
--- a/Engine/Source/Editor/UnrealEd/Private/Kismet2/BlueprintEditorUtils.cpp
+++ b/Engine/Source/Editor/UnrealEd/Private/Kismet2/BlueprintEditorUtils.cpp
@@ -5160,9 +5160,10 @@ void FBlueprintEditorUtils::ChangeMemberVariableType(UBlueprint* Blueprint, cons
 
 							if (FirstVariableNode)
 							{
-								const bool bSetFindWithinBlueprint = false;
-								const bool bSelectFirstResult = false;
-								BlueprintEditor->SummonSearchUI(bSetFindWithinBlueprint, FirstVariableNode->GetFindReferenceSearchString(), bSelectFirstResult);
+								constexpr bool bSetFindWithinBlueprint = false;
+								constexpr bool bSelectFirstResult = false;
+								constexpr EGetFindReferenceSearchStringFlags Flags = EGetFindReferenceSearchStringFlags::UseSearchSyntax;
+								BlueprintEditor->SummonSearchUI(bSetFindWithinBlueprint, FirstVariableNode->GetFindReferenceSearchString(Flags), bSelectFirstResult);
 							}
 						}
 					}
@@ -5673,9 +5674,10 @@ void FBlueprintEditorUtils::ChangeLocalVariableType(UBlueprint* InBlueprint, con
 
 					if (FirstVariableNode)
 					{
-						const bool bSetFindWithinBlueprint = true;
-						const bool bSelectFirstResult = false;
-						BlueprintEditor->SummonSearchUI(bSetFindWithinBlueprint, VariableNodes[0]->GetFindReferenceSearchString(), bSelectFirstResult);
+						constexpr bool bSetFindWithinBlueprint = true;
+						constexpr bool bSelectFirstResult = false;
+						constexpr EGetFindReferenceSearchStringFlags Flags = EGetFindReferenceSearchStringFlags::UseSearchSyntax;
+						BlueprintEditor->SummonSearchUI(bSetFindWithinBlueprint, VariableNodes[0]->GetFindReferenceSearchString(Flags), bSelectFirstResult);
 					}
 				}
 			}
diff --git a/Engine/Source/Runtime/Engine/Classes/EdGraph/EdGraphNode.h b/Engine/Source/Runtime/Engine/Classes/EdGraph/EdGraphNode.h
--- a/Engine/Source/Runtime/Engine/Classes/EdGraph/EdGraphNode.h
+++ b/Engine/Source/Runtime/Engine/Classes/EdGraph/EdGraphNode.h
@@ -269,6 +269,15 @@ struct FEdGraphNodeDeprecationResponse
 	FText MessageText;
 };
 
+/** Options when getting a Find References search string from a node. */
+ENGINE_API enum class EGetFindReferenceSearchStringFlags : int8
+{
+	None = 0,
+	UseSearchSyntax = 1 << 0,
+};
+
+ENUM_CLASS_FLAGS(EGetFindReferenceSearchStringFlags);
+
 UCLASS(MinimalAPI)
 class UEdGraphNode : public UObject
 {
@@ -722,8 +731,24 @@ public:
 	ENGINE_API virtual FText GetNodeTitle(ENodeTitleType::Type TitleType) const;
 
 	/** Gets the search string to find references to this node */
+	UE_DEPRECATED(5.3, "GetFindReferenceSearchString() signature changed. Call the version with flags instead and override the _Impl version.")
 	ENGINE_API virtual FString GetFindReferenceSearchString() const;
 
+	/**
+	 * Gets the search string to find references to this node
+	 * @param InFlags Options for the requested search string, such as whether to use Find-in-Blueprints class member search syntax
+	 */
+	ENGINE_API FString GetFindReferenceSearchString(EGetFindReferenceSearchStringFlags InFlags) const;
+
+protected:
+	virtual FString GetFindReferenceSearchString_Impl(EGetFindReferenceSearchStringFlags InFlags) const
+	{
+		PRAGMA_DISABLE_DEPRECATION_WARNINGS
+		return GetFindReferenceSearchString();
+		PRAGMA_ENABLE_DEPRECATION_WARNINGS
+	}
+
+public:
 	/** 
 	 * Gets the draw color of a node's title bar
 	 */
diff --git a/Engine/Source/Runtime/Engine/Private/EdGraph/EdGraphNode.cpp b/Engine/Source/Runtime/Engine/Private/EdGraph/EdGraphNode.cpp
--- a/Engine/Source/Runtime/Engine/Private/EdGraph/EdGraphNode.cpp
+++ b/Engine/Source/Runtime/Engine/Private/EdGraph/EdGraphNode.cpp
@@ -774,6 +774,11 @@ FString UEdGraphNode::GetFindReferenceSearchString() const
 	return GetNodeTitle(ENodeTitleType::ListView).ToString();
 }
 
+FString UEdGraphNode::GetFindReferenceSearchString(EGetFindReferenceSearchStringFlags InFlags) const
+{
+	return GetFindReferenceSearchString_Impl(InFlags);
+}
+
 UObject* UEdGraphNode::GetJumpTargetForDoubleClick() const
 {
 	return nullptr;
