PR #11165: FieldNotify: Fix checkbox state for field notify properties from other blueprints and native classes

diff --git a/Engine/Source/Editor/Kismet/Private/BlueprintDetailsCustomization.cpp b/Engine/Source/Editor/Kismet/Private/BlueprintDetailsCustomization.cpp
--- a/Engine/Source/Editor/Kismet/Private/BlueprintDetailsCustomization.cpp
+++ b/Engine/Source/Editor/Kismet/Private/BlueprintDetailsCustomization.cpp
@@ -547,7 +547,7 @@ void FBlueprintVarActionDetails::CustomizeDetails( IDetailLayoutBuilder& DetailL
 					SNew(SCheckBox)
 					.IsChecked(this, &FBlueprintVarActionDetails::OnFieldNotifyCheckboxState)
 					.OnCheckStateChanged(this, &FBlueprintVarActionDetails::OnFieldNotifyChanged)
-					.IsEnabled(GetPropertyOwnerBlueprint() && IsABlueprintVariable(VariableProperty) && IsAUserVariable(VariableProperty))
+					.IsEnabled(IsVariableInBlueprint() && GetPropertyOwnerBlueprint() && IsABlueprintVariable(VariableProperty) && IsAUserVariable(VariableProperty))
 					.ToolTip(FieldNotificationTooltip)
 				]
 				+ SHorizontalBox::Slot()
@@ -2007,22 +2007,33 @@ void FBlueprintVarActionDetails::OnVariableUnitsChanged(TSharedPtr<FString> Unit
 
 ECheckBoxState FBlueprintVarActionDetails::OnFieldNotifyCheckboxState() const
 {
-	UBlueprint* const BlueprintObj = GetBlueprintObj();
+	UBlueprint* const BlueprintObj = GetPropertyOwnerBlueprint();
 	const FName VarName = CachedVariableName;
 
-	if (BlueprintObj && !VarName.IsNone())
+	if (!VarName.IsNone())
 	{
-		const int32 VarIndex = FBlueprintEditorUtils::FindNewVariableIndex(BlueprintObj, VarName);
-		if (VarIndex != INDEX_NONE)
+		if (BlueprintObj)
+		{
+			const int32 VarIndex = FBlueprintEditorUtils::FindNewVariableIndex(BlueprintObj, VarName);
+			if (VarIndex != INDEX_NONE)
+			{
+				return BlueprintObj->NewVariables[VarIndex].HasMetaData(FBlueprintMetadata::MD_FieldNotify) ? ECheckBoxState::Checked : ECheckBoxState::Unchecked;
+			}
+		}
+		
+		const UClass* VarSourceClass = BlueprintObj ? BlueprintObj->GeneratedClass : nullptr;
+		if (VarSourceClass == nullptr && CachedVariableProperty.IsValid())
 		{
-			return BlueprintObj->NewVariables[VarIndex].HasMetaData(FBlueprintMetadata::MD_FieldNotify) ? ECheckBoxState::Checked : ECheckBoxState::Unchecked;
+			VarSourceClass = CachedVariableProperty->GetOwner<UClass>();
 		}
-		else if (BlueprintObj->GeneratedClass && BlueprintObj->GeneratedClass->ImplementsInterface(UNotifyFieldValueChanged::StaticClass()) && BlueprintObj->GeneratedClass->GetDefaultObject())
+
+		if (VarSourceClass && VarSourceClass->ImplementsInterface(UNotifyFieldValueChanged::StaticClass()) && VarSourceClass->GetDefaultObject())
 		{
-			TScriptInterface<INotifyFieldValueChanged> DefaultObject = BlueprintObj->GeneratedClass->GetDefaultObject();
-			return DefaultObject->GetFieldNotificationDescriptor().GetField(BlueprintObj->GeneratedClass, VarName).IsValid() ? ECheckBoxState::Checked : ECheckBoxState::Unchecked;
+			TScriptInterface<INotifyFieldValueChanged> DefaultObject = VarSourceClass->GetDefaultObject();
+			return DefaultObject->GetFieldNotificationDescriptor().GetField(VarSourceClass, VarName).IsValid() ? ECheckBoxState::Checked : ECheckBoxState::Unchecked;
 		}
 	}
+	
 	return ECheckBoxState::Unchecked;
 }
 
