Disable searching external objects in the blueprint debugger by default
- in this context, external means objects that don't contain the object at the root of the tree view in their outer chain, the root of the tree view is the value of the pin being inspected, not the object being debugged
- external objects can still be searched using bp.DebuggerEnableExternalSearch
- using the above along with bp.DebuggerMaxSearchDepth (which is now a very large number by default) should allow users to tailor the perf tradeoffs of the debugger to their needs

diff --git a/Engine/Source/Editor/Kismet/Private/Debugging/SKismetDebugTreeView.cpp b/Engine/Source/Editor/Kismet/Private/Debugging/SKismetDebugTreeView.cpp
--- a/Engine/Source/Editor/Kismet/Private/Debugging/SKismetDebugTreeView.cpp
+++ b/Engine/Source/Editor/Kismet/Private/Debugging/SKismetDebugTreeView.cpp
@@ -108,8 +108,11 @@ static const FText ViewInDebuggerTooltipText = LOCTEXT("ViewInDebugger_Tooltip",
 static constexpr float ThumbnailIconSize = 16.0f;
 static constexpr uint32 ThumbnailIconResolution = 16;
 
-static int DebuggerMaxSearchDepth = 5;
-static FAutoConsoleVariableRef CVarDebuggerMaxDepth(TEXT("bp.DebuggerMaxSearchDepth"), DebuggerMaxSearchDepth, TEXT("The maximum search depth of Blueprint Debugger TreeView widgets"), ECVF_Default);
+static int DebuggerMaxSearchDepth = 50;
+static FAutoConsoleVariableRef CVarDebuggerMaxDepth(TEXT("bp.DebuggerMaxSearchDepth"), DebuggerMaxSearchDepth, TEXT("The maximum search depth of Blueprint Debugger TreeView widgets (set to <= 0 for infinite depth)"), ECVF_Default);
+
+static bool bDebuggerEnableExternalSearch = false;
+static FAutoConsoleVariableRef CVarDebuggerEnableExternalSearch(TEXT("bp.DebuggerEnableExternalSearch"), bDebuggerEnableExternalSearch, TEXT("Allows the Blueprint Debugger TreeView widget to search external objects"), ECVF_Default);
 
 //////////////////////////////////////////////////////////////////////////
 
@@ -628,6 +631,26 @@ protected:
 
 	virtual void GatherChildren(TArray<FDebugTreeItemPtr>& OutChildren, const FString& InSearchString, bool bRespectSearch) {}
 
+	/** returns the object related to this line item if there is one, used to avoid searching external objects */
+	virtual const UObject* GetRelatedObject() const { return nullptr; };
+
+	/** returns whether a child is  */
+	virtual bool IsExternalTo(const FLineItemWithChildren* Parent) const
+	{
+		if (const UObject* ChildObj = GetRelatedObject())
+		{
+			if (const UObject* ParentObj = Parent->GetRelatedObject())
+			{
+				if (!ChildObj->IsIn(ParentObj))
+				{
+					return true;
+				}
+			}
+		}
+
+		return false;
+	}
+	
 	/**
 	* returns whether this node should be visible according to the users
 	* search query
@@ -669,16 +692,45 @@ protected:
 					continue;
 				}
 
-					// stop recursing if we reached the max depth
-                    if (Parents.Num() <= DebuggerMaxSearchDepth)
-                    {
+				// stop recursing if we reached the max depth
+				if (DebuggerMaxSearchDepth > 0 && Parents.Num() > DebuggerMaxSearchDepth)
+				{
+					continue;
+				}
+
+				if (!bDebuggerEnableExternalSearch)
+				{
+					FLineItemWithChildren* ObjectParent = nullptr;
+					for (FLineItemWithChildren* Parent : Parents)
+					{
+						if (Parent->GetRelatedObject())
+						{
+							ObjectParent = Parent;
+							break;
+						}
+					}
+					// check if we need to stop searching due to external objects
+					if (ObjectParent && Child->IsExternalTo(ObjectParent))
+					{
+						// update this child, but don't recurse
+						Child->UpdateSearch(InSearchString, ChildSearchFlags);
+
 						// if any children need to expand, so should this
-						if (Child->SearchRecursive(InSearchString, DebugTreeView, Parents, ChildSearchFlags))
+						if (ChildRef->IsVisible())
 						{
 							bVisible = true;
 							bChildMatch = true;
 						}
+						continue;
 					}
+				}
+
+				// if any children need to expand, so should this
+				if (Child->SearchRecursive(InSearchString, DebugTreeView, Parents, ChildSearchFlags))
+				{
+					bVisible = true;
+					bChildMatch = true;
+				}
 			}
 			else
 			{
@@ -1112,6 +1164,11 @@ protected:
 		return {};
 	}
 
+	const UObject* GetRelatedObject() const override
+	{
+		return Data->Object.Get();
+	}
+
 	bool CanOpenAsset() const
 	{
 		if (FSlateApplication::Get().InKismetDebuggingMode())
@@ -1433,6 +1490,15 @@ protected:
 		}
 	}
 	
+	virtual const UObject* GetRelatedObject() const override
+	{
+		if (CachedPropertyInfo.Get())
+		{
+			return CachedPropertyInfo->Object.Get();
+		}
+		return nullptr;
+	}
+
 	virtual FText GetDisplayName() const override;
 	const FSlateBrush* GetPinIcon() const;
 	const FSlateBrush* GetSecondaryPinIcon() const;
@@ -2475,7 +2541,7 @@ public:
 			];
 	}
 
-	virtual UObject* GetParentObject() override
+	virtual UObject* GetParentObject() const override
 	{
 		return ObjectRef.Get();
 	}
diff --git a/Engine/Source/Editor/Kismet/Public/Debugging/SKismetDebugTreeView.h b/Engine/Source/Editor/Kismet/Public/Debugging/SKismetDebugTreeView.h
--- a/Engine/Source/Editor/Kismet/Public/Debugging/SKismetDebugTreeView.h
+++ b/Engine/Source/Editor/Kismet/Public/Debugging/SKismetDebugTreeView.h
@@ -100,7 +100,7 @@ public:
 	virtual bool CanHaveChildren() { return false; }
 
 	// @return The object that will act as a parent to more items in the tree, or nullptr if this is a leaf node
-	virtual UObject* GetParentObject() { return nullptr; }
+	virtual UObject* GetParentObject() const { return nullptr; }
 
 	virtual EDebugLineType GetType() const
 	{
