UE-203565 - OnBrowseSocket BP Menu failing to open when the SubobjectEditor list is docked to the side panel or closed.

When a window is docked in the Blueprint editor, it gets added to the STabSidebar, which is created by the toolkit host. If we use the current SSubobjectEditor as the widget path, we won't actually find a window, because it now lives under the STabSidebar. As a result, the drop down menu for mesh sockets would fail to open.

The easiest fix is to simply use the toolkit host's parent widget, which finds the correct widget path in the docked and undocked cases.

diff --git a/Engine/Source/Editor/Kismet/Private/BlueprintDetailsCustomization.cpp b/Engine/Source/Editor/Kismet/Private/BlueprintDetailsCustomization.cpp
--- a/Engine/Source/Editor/Kismet/Private/BlueprintDetailsCustomization.cpp
+++ b/Engine/Source/Editor/Kismet/Private/BlueprintDetailsCustomization.cpp
@@ -7529,7 +7529,7 @@ void FBlueprintComponentDetails::OnBrowseSocket()
 				{
 					// Pop up a combo box to pick socket from mesh
 					FSlateApplication::Get().PushMenu(
-						Editor.ToSharedRef(),
+						BlueprintEditorPtr.Pin()->GetToolkitHost()->GetParentWidget(),
 						FWidgetPath(),
 						SNew(SSocketChooserPopup)
 						.SceneComponent( ParentSceneComponent )
