Added Archetype reinstantiation sorting and compute the right template object for creation

diff --git a/Engine/Source/Editor/Kismet/Private/BlueprintCompilationManager.cpp b/Engine/Source/Editor/Kismet/Private/BlueprintCompilationManager.cpp
--- a/Engine/Source/Editor/Kismet/Private/BlueprintCompilationManager.cpp
+++ b/Engine/Source/Editor/Kismet/Private/BlueprintCompilationManager.cpp
@@ -514,6 +514,31 @@ struct FReinstancingJob
 
 	// always set:
 	TPair<UClass*, UClass*> OldToNew;
+
+	struct FArchetypeInfo
+	{
+		FArchetypeInfo(UObject* Archetype)
+			: Archetype(Archetype)
+			, ArchetypeTemplate(Archetype ? Archetype->GetArchetype() : nullptr)
+		{}
+
+		bool operator==(const FArchetypeInfo& Other) const
+		{
+			return Archetype == Other.Archetype;
+		}
+
+		friend int32 GetTypeHash(const FArchetypeInfo& Info)
+		{
+			return GetTypeHash(Info.Archetype);
+		}
+
+		UObject* Archetype = nullptr;
+		UObject* ArchetypeTemplate;
+	};
+
+
+	// Old archetype to re-instantiate and its old associated template
+	TArray<FArchetypeInfo> OldArchetypeObjects;
 };
 
 FReinstancingJob::FReinstancingJob(TSharedPtr<FBlueprintCompileReinstancer> InReinstancer)
@@ -2520,11 +2545,11 @@ void FBlueprintCompilationManagerImpl::ReinstanceBatch(TArray<FReinstancingJob>&
 		ArchetypeReferencers.Add(GUnrealEd->Trans);
 	}
 
-	for (const FReinstancingJob* ReinstancingJobPtr : ReinstancersPtr)
+	for (FReinstancingJob* ReinstancingJobPtr : ReinstancersPtr)
 	{
-		const FReinstancingJob& ReinstancingJob = *ReinstancingJobPtr;
+		FReinstancingJob& ReinstancingJob = *ReinstancingJobPtr;
 		UClass* OldClass = ReinstancingJob.OldToNew.Key;
-		if(OldClass)
+		if (OldClass)
 		{
 			SCOPED_LOADTIMER_ASSET_TEXT(*WriteToString<256>(TEXT("Reinstancing "), *GetPathNameSafe(ReinstancingJob.OldToNew.Value)));
 
@@ -2560,8 +2585,10 @@ void FBlueprintCompilationManagerImpl::ReinstanceBatch(TArray<FReinstancingJob>&
 			);
 
 			// for each archetype:
-			for(UObject* Archetype : ArchetypeObjects )
+			for (UObject* Archetype : ArchetypeObjects)
 			{
+				ReinstancingJob.OldArchetypeObjects.Add(Archetype);
+
 				// make sure we fix up references in the owner:
 				{
 					UObject* Iter = Archetype->GetOuter();
@@ -2599,14 +2626,48 @@ void FBlueprintCompilationManagerImpl::ReinstanceBatch(TArray<FReinstancingJob>&
 						Iter = Iter->GetOuter();
 					}
 				}
-				
+			}
+
+			// Sort all archetype between themselves, as one might depends on an other
+			// This happens when the class containing the archetype has derived classes.
+			Algo::TopologicalSort(ReinstancingJob.OldArchetypeObjects, [&ArchetypeObjects](const FReinstancingJob::FArchetypeInfo& OldArchetypeInfo)
+			{
+				TArray<UObject*> Dependencies;
+				if (OldArchetypeInfo.ArchetypeTemplate && !OldArchetypeInfo.ArchetypeTemplate->HasAnyFlags(RF_ClassDefaultObject))
+				{
+					if(ensureMsgf(ArchetypeObjects.Contains(OldArchetypeInfo.ArchetypeTemplate), TEXT("Expecting the template object to already be in the list of archetypes")))
+					{
+						Dependencies.Add(OldArchetypeInfo.ArchetypeTemplate);
+					}
+				}
+
+				return Dependencies;
+			});
+		}
+	}
+
+	// This loop finishes the reinstancing of archetypes after the entire Outer hierarchy has been updated with new instances:
+	for (const FReinstancingJob* ReinstancingJobPtr : ReinstancersPtr)
+	{
+		const FReinstancingJob& ReinstancingJob = *ReinstancingJobPtr;
+		UClass* OldClass = ReinstancingJob.OldToNew.Key;
+		if(OldClass)
+		{
+			SCOPED_LOADTIMER_ASSET_TEXT(*WriteToString<256>(TEXT("FinishReinstancing "), *GetPathNameSafe(ReinstancingJob.OldToNew.Value)));
+
+			UClass* NewClass = ReinstancingJob.OldToNew.Value;
+
+			for(const FReinstancingJob::FArchetypeInfo& OldArchetypeInfo : ReinstancingJob.OldArchetypeObjects)
+			{
+				UObject* OldInstance = OldArchetypeInfo.Archetype;
+
 				// move aside:
-				FName OriginalName = Archetype->GetFName();
-				UObject* OriginalOuter = Archetype->GetOuter();
-				EObjectFlags OriginalFlags = Archetype->GetFlags();
+				FName OriginalName = OldInstance->GetFName();
+				UObject* OriginalOuter = OldInstance->GetOuter();
+				EObjectFlags OriginalFlags = OldInstance->GetFlags();
 
-				UObject* Destination = GetTransientOuterForRename(Archetype->GetClass());
-				Archetype->Rename(
+				UObject* Destination = GetTransientOuterForRename(OldInstance->GetClass());
+				OldInstance->Rename(
 					nullptr,
 					// destination - this is the important part of this call. Moving the object 
 					// out of the way so we can reuse its name:
@@ -2617,37 +2678,25 @@ void FBlueprintCompilationManagerImpl::ReinstanceBatch(TArray<FReinstancingJob>&
 				// reconstruct
 				FMakeClassSpawnableOnScope TemporarilySpawnable(NewClass);
 				const EObjectFlags FlagMask = RF_Public | RF_ArchetypeObject | RF_Transactional | RF_Transient | RF_TextExportTransient | RF_InheritableComponentTemplate | RF_Standalone; //TODO: what about RF_RootSet?
-				UObject* NewArchetype = NewObject<UObject>(OriginalOuter, NewClass, OriginalName, OriginalFlags & FlagMask);
 
-				OldArchetypeToNewArchetype.Add(Archetype, NewArchetype);
+				UObject* Template = nullptr;
+				if (UObject** NewArchetypeTemplate = OldArchetypeInfo.ArchetypeTemplate ? OldArchetypeToNewArchetype.Find(OldArchetypeInfo.ArchetypeTemplate) : nullptr)
+				{
+					Template = *NewArchetypeTemplate;
+				}
+				UObject* NewArchetype = NewObject<UObject>(OriginalOuter, NewClass, OriginalName, OriginalFlags & FlagMask, Template);
+
+				OldArchetypeToNewArchetype.Add(OldInstance, NewArchetype);
 
 				// also map old *default* subobjects to new default subobjects:
-				BuildDSOMap(Archetype, NewArchetype, OldArchetypeToNewArchetype);
+				BuildDSOMap(OldInstance, NewArchetype, OldArchetypeToNewArchetype);
 
 				ArchetypeReferencers.Add(NewArchetype);
 
-				FLinkerLoad::PRIVATE_PatchNewObjectIntoExport(Archetype, NewArchetype);
-			}
-		}
-	}
-
-	// This loop finishes the reinstancing of archetypes after the entire Outer hierarchy has been updated with new instances:
-	for (const FReinstancingJob* ReinstancingJobPtr : ReinstancersPtr)
-	{
-		const FReinstancingJob& ReinstancingJob = *ReinstancingJobPtr;
-		UClass* OldClass = ReinstancingJob.OldToNew.Key;
-		if(OldClass)
-		{
-			SCOPED_LOADTIMER_ASSET_TEXT(*WriteToString<256>(TEXT("FinishReinstancing "), *GetPathNameSafe(ReinstancingJob.OldToNew.Value)));
-
-			TArray<UObject*> OldInstances;
-			GetObjectsOfClass( OldClass, OldInstances, false );
+				FLinkerLoad::PRIVATE_PatchNewObjectIntoExport(OldInstance, NewArchetype);
 
-			for(UObject* OldInstance : OldInstances)
-			{
-				UObject** NewInstance = OldArchetypeToNewArchetype.Find(OldInstance);
-				// NewInstance may be null in the case of deleted or EditorOnly (in -game) DSOs:
-				if(NewInstance && *NewInstance)
+				// NewArchetype may be null in the case of deleted or EditorOnly (in -game) DSOs:
+				if(NewArchetype)
 				{
 					// The new object hierarchy has been created, all of the old instances are in the transient package and new
 					// ones have taken their place. Reference members will mostly be pointing at *old* instances, and will get fixed
@@ -2655,7 +2704,7 @@ void FBlueprintCompilationManagerImpl::ReinstanceBatch(TArray<FReinstancingJob>&
 					const bool bUseDeltaSerialization = ReinstancingJob.Reinstancer.IsValid() ? ReinstancingJob.Reinstancer->bUseDeltaSerializationToCopyProperties : false;
 
 					TMap<UObject*, UObject*> CreatedInstanceMap;
-					FBlueprintCompileReinstancer::PreCreateSubObjectsForReinstantiation(InOutOldToNewClassMap, OldInstance, *NewInstance, CreatedInstanceMap);
+					FBlueprintCompileReinstancer::PreCreateSubObjectsForReinstantiation(InOutOldToNewClassMap, OldInstance, NewArchetype, CreatedInstanceMap);
 
 					// We only need to copy properties of the pre-created instances, the rest of the default sub object is done inside the UEditorEngine::CopyPropertiesForUnrelatedObjects
 					TMap<UObject*, UObject*> OldToNewInstanceMap(CreatedInstanceMap);
