Adding additional guard rails to the "Index All" action from Find-in-Blueprints. Editor setting that controls user permission is now an enum. The default setting allows mass loading blueprints, but not mass checkout and resave. As a result, mass write is opt-in - also improved the final warning for mass-write to emphasize you're about to checkout.

diff --git a/Engine/Source/Editor/BlueprintGraph/Private/BlueprintEditorSettings.cpp b/Engine/Source/Editor/BlueprintGraph/Private/BlueprintEditorSettings.cpp
--- a/Engine/Source/Editor/BlueprintGraph/Private/BlueprintEditorSettings.cpp
+++ b/Engine/Source/Editor/BlueprintGraph/Private/BlueprintEditorSettings.cpp
@@ -69,7 +69,7 @@ UBlueprintEditorSettings::UBlueprintEditorSettings(const FObjectInitializer& Obj
 	// Perf Settings
 	, NodeTemplateCacheCapMB(20.f)
 	// Find-in-Blueprints Settings
-	, bAllowIndexAllBlueprints(true)
+	, AllowIndexAllBlueprints(EFiBIndexAllPermission::LoadOnly)
 	// No category
 	, bShowInheritedVariables(false)
 	, bAlwaysShowInterfacesInOverrides(true)
diff --git a/Engine/Source/Editor/BlueprintGraph/Public/BlueprintEditorSettings.h b/Engine/Source/Editor/BlueprintGraph/Public/BlueprintEditorSettings.h
--- a/Engine/Source/Editor/BlueprintGraph/Public/BlueprintEditorSettings.h
+++ b/Engine/Source/Editor/BlueprintGraph/Public/BlueprintEditorSettings.h
@@ -12,6 +12,7 @@
 #include "Delegates/Delegate.h"
 #include "Engine/Blueprint.h"
 #include "Engine/DeveloperSettings.h"
+#include "FindInBlueprints.h"
 #include "Internationalization/Text.h"
 #include "Kismet2/Breakpoint.h"
 #include "Kismet2/KismetDebugUtilities.h"
@@ -264,9 +265,9 @@ public:
 
 	// Find-in-Blueprint Settings
 public:
-	/** Whether to enable the "Index All" button in the Find-in-Blueprints search window when blueprint assets with an out-of-date index (search metadata) are found */
+	/** Whether to enable the "Index All" action in the Find-in-Blueprints search window when blueprint assets with an out-of-date index (search metadata) are found and whether to allow automatic resaving. WARNING: Only allow "Index All" if your project is small enough that all assets can be loaded in memory at once. Only enable saving if you are allowed to potentially checkout and resave all assets. */
 	UPROPERTY(EditAnywhere, config, Category = FindInBlueprints)
-	bool bAllowIndexAllBlueprints;
+	EFiBIndexAllPermission AllowIndexAllBlueprints;
 
 public:
 	/** If set we'll show the inherited variables in the My Blueprint view. */
diff --git a/Engine/Source/Editor/Kismet/Private/FindInBlueprintManager.cpp b/Engine/Source/Editor/Kismet/Private/FindInBlueprintManager.cpp
--- a/Engine/Source/Editor/Kismet/Private/FindInBlueprintManager.cpp
+++ b/Engine/Source/Editor/Kismet/Private/FindInBlueprintManager.cpp
@@ -39,6 +39,7 @@
 #include "WorkspaceMenuStructure.h"
 #include "WorkspaceMenuStructureModule.h"
 
+#include "BlueprintEditorSettings.h"
 #include "Blueprint/BlueprintExtension.h"
 #include "Engine/SimpleConstructionScript.h"
 #include "Kismet2/BlueprintEditorUtils.h"
@@ -3311,10 +3312,24 @@ void FFindInBlueprintSearchManager::CacheAllAssets(TWeakPtr< SFindInBlueprints >
 			Args.Add(TEXT("PackageCount"), UnindexedAssets.Num() + BlueprintsToUpdate.Num());
 			Args.Add(TEXT("UnindexedCount"), UnindexedAssets.Num());
 			Args.Add(TEXT("OutOfDateCount"), BlueprintsToUpdate.Num());
-			const FText DialogDisplayText = FText::Format(LOCTEXT("CacheAllConfirmationMessage_UnindexedAndOutOfDate", "About to load {PackageCount} Blueprints ({UnindexedCount} unindexed/{OutOfDateCount} out-of-date). The editor may become unresponsive while these assets are loaded for indexing. Save your work before initiating this: broken assets and memory usage can affect editor stability. \n\nLoaded assets must be resaved to make this indexing permanent, otherwise their updated searchability is for this editor session only. Select 'Yes' to checkout, load and resave all Blueprints with an outdated index. Select 'No' to load them only."), Args);
-			const EAppReturnType::Type ReturnValue = FMessageDialog::Open(EAppMsgType::YesNoCancel, DialogDisplayText, DialogTitle);
 
-			// If Yes is chosen, checkout and save all Blueprints, if No is chosen, only load all Blueprints
+			// Retrieve from blueprint editor settings whether user is allowed to checkout and resave
+			const bool bCanCheckoutResaveAll = GetDefault<UBlueprintEditorSettings>()->AllowIndexAllBlueprints == EFiBIndexAllPermission::CheckoutAndResave;
+
+			// Present a prompt depending on whether checkout is allowed
+			EAppReturnType::Type ReturnValue;
+			if (bCanCheckoutResaveAll)
+			{
+				const FText DialogDisplayText = FText::Format(LOCTEXT("CacheAllConfirmationMessage_UnindexedAndOutOfDate_WithCheckout", "About to CHECKOUT and RESAVE {PackageCount} Blueprints ({UnindexedCount} unindexed/{OutOfDateCount} out-of-date)! The editor may become unresponsive while these assets are loaded for indexing. Save your work before initiating this: broken assets and memory usage can affect editor stability. \n\nLoaded assets must be resaved to make this indexing permanent, otherwise their updated searchability is for this editor session only. Select 'Yes' to checkout, load and resave all Blueprints with an outdated index. Select 'No' to load these Blueprints only without checking out and resaving them."), Args);
+				ReturnValue = FMessageDialog::Open(EAppMsgType::YesNoCancel, DialogDisplayText, DialogTitle);
+			}
+			else
+			{
+				const FText DialogDisplayText = FText::Format(LOCTEXT("CacheAllConfirmationMessage_UnindexedAndOutOfDate_LoadOnly", "About to load {PackageCount} Blueprints ({UnindexedCount} unindexed/{OutOfDateCount} out-of-date)! The editor may become unresponsive while these assets are loaded for indexing. Save your work before initiating this: broken assets and memory usage can affect editor stability. \n\nLoaded assets must be resaved to make this indexing permanent, otherwise their updated searchability is for this editor session only. Your editor settings disallow resaving all these assets from this window, see Blueprint Editor Settings: AllowIndexAllBlueprints."), Args);
+				ReturnValue = FMessageDialog::Open(EAppMsgType::OkCancel, DialogDisplayText, DialogTitle);
+			}
+
+			// If checkout is allowed and Yes is chosen, checkout and save all Blueprints. Otherwise only load all Blueprints. Cancel aborts everything.
 			if (ReturnValue != EAppReturnType::Cancel)
 			{
 				FailedToCachePaths.Empty();
@@ -3323,7 +3338,7 @@ void FFindInBlueprintSearchManager::CacheAllAssets(TWeakPtr< SFindInBlueprints >
 				TempUncachedAssets.Append(UnindexedAssets);
 				TempUncachedAssets.Append(BlueprintsToUpdate);
 
-				const bool bCheckOutAndSave = (ReturnValue == EAppReturnType::Yes);
+				const bool bCheckOutAndSave = bCanCheckoutResaveAll && (ReturnValue == EAppReturnType::Yes);
 				CacheParams.OpFlags = EFiBCacheOpFlags::ShowProgress | EFiBCacheOpFlags::AllowUserCancel | EFiBCacheOpFlags::AllowUserCloseProgress;
 				if (bCheckOutAndSave)
 				{
diff --git a/Engine/Source/Editor/Kismet/Private/FindInBlueprints.cpp b/Engine/Source/Editor/Kismet/Private/FindInBlueprints.cpp
--- a/Engine/Source/Editor/Kismet/Private/FindInBlueprints.cpp
+++ b/Engine/Source/Editor/Kismet/Private/FindInBlueprints.cpp
@@ -1638,14 +1638,30 @@ FText SFindInBlueprints::GetCacheBarStatusText() const
 	}
 	else
 	{
-		const FText IndexAllWarningText = LOCTEXT("IndexAllWarning", "Press \"Index All\" to load these assets right now. The editor may become unresponsive while these assets are loaded for indexing. Save your work before initiating this: broken assets and memory usage can affect editor stability. Alternatively, export the asset list to inspect which assets do not have optimal searchability.");
-		const FText IndexAllDisabledText = LOCTEXT("IndexAllDisabled", "Your Blueprint Editor settings disallow loading all these assets from this window. Export the asset list to inspect which assets do not have optimal searchability.");
-
 		const int32 UnindexedCount = FindInBlueprintManager.GetNumberUnindexedAssets();
 		Args.Add(TEXT("UnindexedCount"), UnindexedCount);
 		Args.Add(TEXT("OutOfDateCount"), OutOfDateWithLastSearchBPCount);
 		Args.Add(TEXT("Count"), UnindexedCount + OutOfDateWithLastSearchBPCount);
-		Args.Add(TEXT("Instruction"), CanCacheAllUnindexedBlueprints() ? IndexAllWarningText : IndexAllDisabledText);
+
+		// Show a different instruction depending on the "Index All" permission level in editor settings
+		const EFiBIndexAllPermission IndexAllPermission = GetDefault<UBlueprintEditorSettings>()->AllowIndexAllBlueprints;
+		const FText IndexAllDisabledText = LOCTEXT("IndexAllDisabled", "Your editor settings disallow loading all these assets from this window, see Blueprint Editor Settings: AllowIndexAllBlueprints. Export the asset list to inspect which assets do not have optimal searchability.");
+		const FText IndexAllWarningText_LoadOnly = LOCTEXT("IndexAllWarning_LoadOnly", "Press \"Index All\" to load these assets right now. The editor may become unresponsive while these assets are loaded for indexing. Save your work before initiating this: broken assets and memory usage can affect editor stability. Alternatively, export the asset list to inspect which assets do not have optimal searchability.");
+		const FText IndexAllWarningText_Checkout = LOCTEXT("IndexAllWarning_Checkout", "Press \"Index All\" to load, and optionally checkout and resave, these assets right now. The editor may become unresponsive while these assets are loaded for indexing. Save your work before initiating this: broken assets and memory usage can affect editor stability. Alternatively, export the asset list to inspect which assets do not have optimal searchability.");
+		switch (IndexAllPermission)
+		{
+		case EFiBIndexAllPermission::CheckoutAndResave:
+			Args.Add(TEXT("Instruction"), IndexAllWarningText_Checkout);
+			break;
+		case EFiBIndexAllPermission::LoadOnly:
+			Args.Add(TEXT("Instruction"), IndexAllWarningText_LoadOnly);
+			break;
+		case EFiBIndexAllPermission::None:
+			Args.Add(TEXT("Instruction"), IndexAllDisabledText);
+			break;
+		default:
+			ensureMsgf(false, TEXT("Unhandled case"));
+		}
 		
 		ReturnDisplayText = FText::Format(LOCTEXT("UncachedAssets", "Search incomplete: {Count} blueprints don't have an up-to-date index ({UnindexedCount} unindexed/{OutOfDateCount} out-of-date). These assets are searchable but some results may be missing. Load and resave these assets to improve their searchability. \n\n{Instruction}"), Args);
 
@@ -1674,7 +1690,7 @@ FText SFindInBlueprints::GetCacheBarCurrentAssetName() const
 
 bool SFindInBlueprints::CanCacheAllUnindexedBlueprints() const
 {
-	return GetDefault<UBlueprintEditorSettings>()->bAllowIndexAllBlueprints;
+	return GetDefault<UBlueprintEditorSettings>()->AllowIndexAllBlueprints != EFiBIndexAllPermission::None;
 }
 
 void SFindInBlueprints::OnCacheStarted(EFiBCacheOpType InOpType, EFiBCacheOpFlags InOpFlags)
diff --git a/Engine/Source/Editor/Kismet/Public/FindInBlueprints.h b/Engine/Source/Editor/Kismet/Public/FindInBlueprints.h
--- a/Engine/Source/Editor/Kismet/Public/FindInBlueprints.h
+++ b/Engine/Source/Editor/Kismet/Public/FindInBlueprints.h
@@ -261,6 +261,18 @@ enum class EFiBSearchBarWidget
 	ProgressBar,
 };
 
+// Whether the Find-in-Blueprints window allows the user to load and resave all assets with out-of-date Blueprint search metadata
+UENUM()
+enum class EFiBIndexAllPermission
+{
+	// Users may not automatically load all Blueprints with out-of-date search metadata
+	None,
+	// Users may automatically load all Blueprints with out-of-date search metadata, but not resave
+	LoadOnly,
+	// Users may automatically checkout, load and resave all Blueprints with out-of-date search metadata
+	CheckoutAndResave
+};
+
 /*Widget for searching for (functions/events) across all blueprints or just a single blueprint */
 class KISMET_API SFindInBlueprints: public SCompoundWidget
 {
