Cleanup pin type selector, improving performance in large projects.

diff --git a/Engine/Source/Editor/BlueprintGraph/Classes/EdGraphSchema_K2.h b/Engine/Source/Editor/BlueprintGraph/Classes/EdGraphSchema_K2.h
--- a/Engine/Source/Editor/BlueprintGraph/Classes/EdGraphSchema_K2.h
+++ b/Engine/Source/Editor/BlueprintGraph/Classes/EdGraphSchema_K2.h
@@ -467,6 +467,7 @@ public:
 
 	public:
 		const FEdGraphPinType& GetPinType(bool bForceLoadedSubCategoryObject);
+		const FEdGraphPinType& GetPinTypeNoResolve() const { return PinType; }
 		void SetPinSubTypeCategory(const FName SubCategory)
 		{
 			PinType.PinSubCategory = SubCategory;
diff --git a/Engine/Source/Editor/BlueprintGraph/Private/EdGraphSchema_K2.cpp b/Engine/Source/Editor/BlueprintGraph/Private/EdGraphSchema_K2.cpp
--- a/Engine/Source/Editor/BlueprintGraph/Private/EdGraphSchema_K2.cpp
+++ b/Engine/Source/Editor/BlueprintGraph/Private/EdGraphSchema_K2.cpp
@@ -640,7 +640,6 @@ UEdGraphSchema_K2::FPinTypeTreeInfo::FPinTypeTreeInfo(const FText& InFriendlyNam
 	PinType.PinCategory = CategoryName;
 
 	SubCategoryObjectAssetReference = SubCategoryObject;
-	PinType.PinSubCategoryObject = SubCategoryObjectAssetReference.ResolveObject();
 
 	bReadOnly = bInReadOnly;
 	CachedDescription = GenerateDescription();
diff --git a/Engine/Source/Editor/Kismet/Private/BlueprintEditor.cpp b/Engine/Source/Editor/Kismet/Private/BlueprintEditor.cpp
--- a/Engine/Source/Editor/Kismet/Private/BlueprintEditor.cpp
+++ b/Engine/Source/Editor/Kismet/Private/BlueprintEditor.cpp
@@ -428,7 +428,7 @@ namespace BlueprintEditorImpl
 			}
 
 			// First check pin type permissions
-			if (!FBlueprintActionDatabase::IsPinTypeAllowed(InItem->GetPinType(false), TopLevelAssetPath))
+			if (!FBlueprintActionDatabase::IsPinTypeAllowed(InItem->GetPinTypeNoResolve(), TopLevelAssetPath))
 			{
 				return false;
 			}
diff --git a/Engine/Source/Editor/Kismet/Private/BlueprintNamespaceHelper.cpp b/Engine/Source/Editor/Kismet/Private/BlueprintNamespaceHelper.cpp
--- a/Engine/Source/Editor/Kismet/Private/BlueprintNamespaceHelper.cpp
+++ b/Engine/Source/Editor/Kismet/Private/BlueprintNamespaceHelper.cpp
@@ -102,8 +102,7 @@ public:
 			return true;
 		}
 
-		const bool bForceLoadSubCategoryObject = false;
-		const FEdGraphPinType& PinType = InItem->GetPinType(bForceLoadSubCategoryObject);
+		const FEdGraphPinType& PinType = InItem->GetPinTypeNoResolve();
 
 		if (PinType.PinSubCategoryObject.IsValid() && !CachedNamespaceHelper->IsImportedObject(PinType.PinSubCategoryObject.Get()))
 		{
diff --git a/Engine/Source/Editor/Kismet/Private/BlueprintNamespaceUtilities.cpp b/Engine/Source/Editor/Kismet/Private/BlueprintNamespaceUtilities.cpp
--- a/Engine/Source/Editor/Kismet/Private/BlueprintNamespaceUtilities.cpp
+++ b/Engine/Source/Editor/Kismet/Private/BlueprintNamespaceUtilities.cpp
@@ -158,12 +158,7 @@ FString FBlueprintNamespaceUtilities::GetObjectNamespace(const UObject* InObject
 
 FString FBlueprintNamespaceUtilities::GetObjectNamespace(const FSoftObjectPath& InObjectPath)
 {
-	if (const UObject* Object = InObjectPath.ResolveObject())
-	{
-		return GetObjectNamespace(Object);
-	}
-
-	const bool bIncludeOnlyOnDiskAssets = true; // The resolve object failed so we know it isn't in memory
+	const bool bIncludeOnlyOnDiskAssets = false;
 	const FAssetRegistryModule& AssetRegistryModule = FModuleManager::LoadModuleChecked<FAssetRegistryModule>(TEXT("AssetRegistry"));
 	FAssetData AssetData = AssetRegistryModule.Get().GetAssetByObjectPath(InObjectPath, bIncludeOnlyOnDiskAssets);
 	if (!AssetData.IsValid())
diff --git a/Engine/Source/Editor/KismetWidgets/Private/SPinTypeSelector.cpp b/Engine/Source/Editor/KismetWidgets/Private/SPinTypeSelector.cpp
--- a/Engine/Source/Editor/KismetWidgets/Private/SPinTypeSelector.cpp
+++ b/Engine/Source/Editor/KismetWidgets/Private/SPinTypeSelector.cpp
@@ -1441,7 +1441,7 @@ bool SPinTypeSelector::GetChildrenWithSupportedTypes(const TArray<FPinTypeTreeIt
 
 		if (!bHasChildrenWithValidTypes)
 		{
-			bSupportsType = Schema->SupportsPinType(SchemaAction, Item->GetPinType(false));
+			bSupportsType = Schema->SupportsPinType(SchemaAction, Item->GetPinTypeNoResolve());
 		}
 		
 		if (bHasChildrenWithValidTypes || bSupportsType)
