Disable editing of the replication property in the BP details panel for Event Dispatchers. Event Dispatchers do not actually support replication.

diff --git a/Engine/Source/Editor/Kismet/Private/BlueprintDetailsCustomization.cpp b/Engine/Source/Editor/Kismet/Private/BlueprintDetailsCustomization.cpp
--- a/Engine/Source/Editor/Kismet/Private/BlueprintDetailsCustomization.cpp
+++ b/Engine/Source/Editor/Kismet/Private/BlueprintDetailsCustomization.cpp
@@ -2331,26 +2331,41 @@ bool FBlueprintVarActionDetails::ReplicationConditionEnabled() const
 bool FBlueprintVarActionDetails::ReplicationEnabled() const
 {
 	// Update FBlueprintVarActionDetails::ReplicationTooltip if you alter this function
-	// shat users can understand why replication settins are disabled!
+	// so that users can understand why replication settings are disabled!
 	bool bVariableCanBeReplicated = true;
 	const FProperty* const VariableProperty = CachedVariableProperty.Get();
 	if (VariableProperty)
 	{
-		// sets and maps cannot yet be replicated:
-		bVariableCanBeReplicated = CastField<FSetProperty>(VariableProperty) == nullptr && CastField<FMapProperty>(VariableProperty) == nullptr;
+		// sets and maps cannot yet be replicated, neither can Event Dispatchers:
+		bVariableCanBeReplicated =
+			CastField<FSetProperty>(VariableProperty) == nullptr &&
+			CastField<FMapProperty>(VariableProperty) == nullptr &&
+			CastField<FMulticastInlineDelegateProperty>(VariableProperty) == nullptr;
 	}
 	return bVariableCanBeReplicated && IsVariableInBlueprint();
 }
 
 FText FBlueprintVarActionDetails::ReplicationTooltip() const
 {
-	if(ReplicationEnabled())
+	if (ReplicationEnabled())
 	{
 		return LOCTEXT("VariableReplicate_Tooltip", "Should this Variable be replicated over the network?");
 	}
 	else
 	{
-		return LOCTEXT("VariableReplicateDisabled_Tooltip", "Set and Map properties cannot be replicated");
+		const FProperty* const VariableProperty = CachedVariableProperty.Get();
+		if (CastField<FSetProperty>(VariableProperty) || CastField<FMapProperty>(VariableProperty))
+		{
+			return LOCTEXT("VariableReplicateDisabledSetsAndMaps_Tooltip", "Set and Map properties cannot be replicated");
+		}
+		else if (CastField<FMulticastInlineDelegateProperty>(VariableProperty))
+		{
+			return LOCTEXT("VariableReplicateDisabledEventDispatchers_Tooltip", "Event Dispatcher properties cannot be replicated");
+		}
+		else
+		{
+			return LOCTEXT("VariableReplicateDisabledDefault_Tooltip", "This property type cannot be replicated");
+		}
 	}
 }
 
