UE-143412 - Using Collapse to Function inside of a thread safe graph does not mark the resulting function as thread safe

Since the nodes in a threadsafe function are threadsafe, it follows that the collapsed graph is also threadsafe. Not doing so creates a clumsy workflow: the Blueprint will fail to compile until the user manually marks the collapsed function as threadsafe.

#rnx

diff --git a/Engine/Source/Editor/Kismet/Private/BlueprintEditor.cpp b/Engine/Source/Editor/Kismet/Private/BlueprintEditor.cpp
--- a/Engine/Source/Editor/Kismet/Private/BlueprintEditor.cpp
+++ b/Engine/Source/Editor/Kismet/Private/BlueprintEditor.cpp
@@ -8799,6 +8799,24 @@ void FBlueprintEditor::CollapseNodesIntoGraph(UEdGraphNode* InGatewayNode, UK2No
 		}
 	}
 
+	if (UK2Node_FunctionEntry* EntryNode = Cast<UK2Node_FunctionEntry>(InEntryNode))
+	{
+		TArray<UK2Node_FunctionEntry*> SourceEntryNodes;
+		InSourceGraph->GetNodesOfClass(SourceEntryNodes);
+
+		UK2Node_FunctionEntry* SourceEntryNode =
+			(SourceEntryNodes.Num() == 1) ? SourceEntryNodes[0] : nullptr;
+
+		if (ensure(SourceEntryNode))
+		{
+			// If the source entry node is threadsafe, then our new entry node should also be threadsafe.
+			// It's implied that the source graph only has threadsafe nodes.
+			// This ensures that our newly collapsed graph can be still be used in the source graph.
+
+			EntryNode->MetaData.bThreadSafe = SourceEntryNode->MetaData.bThreadSafe;
+		}
+	}
+
 	const int32 CenterX = (NumNodes == 0) ? SumNodeX : SumNodeX / NumNodes;
 	const int32 CenterY = (NumNodes == 0) ? SumNodeY : SumNodeY / NumNodes;
 	const int32 MinusOffsetX = 160; //@TODO: Random magic numbers
