Fixed an ensure that was firing when setting default values for FText's in BP function local variables. This was occuring becuase the package was not being set for the DefaultValue

diff --git a/Engine/Source/Editor/Kismet/Private/BlueprintDetailsCustomization.cpp b/Engine/Source/Editor/Kismet/Private/BlueprintDetailsCustomization.cpp
--- a/Engine/Source/Editor/Kismet/Private/BlueprintDetailsCustomization.cpp
+++ b/Engine/Source/Editor/Kismet/Private/BlueprintDetailsCustomization.cpp
@@ -1092,6 +1092,7 @@ void FBlueprintVarActionDetails::CustomizeDetails( IDetailLayoutBuilder& DetailL
 				UK2Node_FunctionEntry* FuncEntry = EntryNodes[0];
 
 				TSharedPtr<FStructOnScope> StructData = MakeShareable(new FStructOnScope((UFunction*)StructScope));
+				StructData->SetPackage(BlueprintObj->GetPackage());
 
 				for (const FBPVariableDescription& LocalVar : FuncEntry->LocalVariables)
 				{
