Show icon in Widget Blueprint editor next to variables marked with BindWidget or BindWidgetOptional

diff --git a/Engine/Source/Editor/Kismet/Private/SBlueprintPalette.cpp b/Engine/Source/Editor/Kismet/Private/SBlueprintPalette.cpp
--- a/Engine/Source/Editor/Kismet/Private/SBlueprintPalette.cpp
+++ b/Engine/Source/Editor/Kismet/Private/SBlueprintPalette.cpp
@@ -103,6 +103,7 @@
 #include "UObject/WeakFieldPtr.h"
 #include "UObject/WeakObjectPtr.h"
 #include "UObject/WeakObjectPtrTemplates.h"
+#include "WidgetBlueprintEditorUtils.h"
 #include "Widgets/IToolTip.h"
 #include "Widgets/Images/SImage.h"
 #include "Widgets/Input/SCheckBox.h"
@@ -1407,6 +1408,34 @@ void SBlueprintPaletteItem::Construct(const FArguments& InArgs, FCreateWidgetFor
 				SNew(SPaletteItemVisibilityToggle, ActionPtr, InBlueprintEditor, InBlueprint)
 				.IsEnabled(bIsEditingEnabled)
 			];
+
+
+		if (TSharedPtr<FEdGraphSchemaAction> Action = ActionPtr.Pin())
+		{
+			if (GraphAction->GetTypeId() == FEdGraphSchemaAction_K2Var::StaticGetTypeId())
+			{
+				if (const FProperty* Property = StaticCastSharedPtr<FEdGraphSchemaAction_K2Var>(GraphAction)->GetProperty())
+				{
+					bool bIsOptional = false;
+					if (FWidgetBlueprintEditorUtils::IsBindWidgetProperty(Property, bIsOptional))
+					{
+						ActionBox.Get().AddSlot()
+							.AutoWidth()
+							.Padding(FMargin(6.0f, 0.0f, 3.0f, 0.0f))
+							.HAlign(HAlign_Right)
+							.VAlign(VAlign_Center)
+							[
+								SNew(SImage)
+									.Image(FAppStyle::Get().GetBrush("MainFrame.AddCodeToProject"))
+									.IsEnabled(!bIsOptional) // make grey to differentiate between optional and required
+									.ToolTipText(bIsOptional
+										? LOCTEXT("CppBindWidgetOptionalTooltip", "Widget marked with BindWidgetOptional in native.")
+										: LOCTEXT("CppBindWidgetTooltip", "Widget marked with BindWidget in native."))
+							];
+					}
+				}
+			}
+		}
 	}
 	else
 	{
