UE-143412 - Using Collapse to Function inside of a thread safe graph does not mark the resulting function as thread safe, causing a compile error until it is checked

Follow-up fix for threadsafe graphs. The previous fix in CL 30950233 failed to account for native functions that were marked as threadsafe. In this case, we need to check the metadata on the corresponding UFunction instead.

We already have a utilty function that handles both cases, which can be used here: HasFunctionBlueprintThreadSafeMetaData.

diff --git a/Engine/Source/Editor/Kismet/Private/BlueprintEditor.cpp b/Engine/Source/Editor/Kismet/Private/BlueprintEditor.cpp
--- a/Engine/Source/Editor/Kismet/Private/BlueprintEditor.cpp
+++ b/Engine/Source/Editor/Kismet/Private/BlueprintEditor.cpp
@@ -8798,20 +8798,11 @@ void FBlueprintEditor::CollapseNodesIntoGraph(UEdGraphNode* InGatewayNode, UK2No
 
 	if (UK2Node_FunctionEntry* EntryNode = Cast<UK2Node_FunctionEntry>(InEntryNode))
 	{
-		TArray<UK2Node_FunctionEntry*> SourceEntryNodes;
-		InSourceGraph->GetNodesOfClass(SourceEntryNodes);
+		// If the source entry node is threadsafe, then our new entry node should also be threadsafe.
+		// It's implied that the source graph only has threadsafe nodes.
+		// This ensures that our newly collapsed graph can be still be used in the source graph.
 
-		UK2Node_FunctionEntry* SourceEntryNode =
-			(SourceEntryNodes.Num() == 1) ? SourceEntryNodes[0] : nullptr;
-
-		if (ensure(SourceEntryNode))
-		{
-			// If the source entry node is threadsafe, then our new entry node should also be threadsafe.
-			// It's implied that the source graph only has threadsafe nodes.
-			// This ensures that our newly collapsed graph can be still be used in the source graph.
-
-			EntryNode->MetaData.bThreadSafe = SourceEntryNode->MetaData.bThreadSafe;
-		}
+		EntryNode->MetaData.bThreadSafe = K2Schema->IsGraphMarkedThreadSafe(InSourceGraph);
 	}
 
 	const int32 CenterX = (NumNodes == 0) ? SumNodeX : SumNodeX / NumNodes;
