Cleaned up SCD management on blueprint compile and asset reload

diff --git a/Engine/Source/Editor/Kismet/Private/BlueprintCompilationManager.cpp b/Engine/Source/Editor/Kismet/Private/BlueprintCompilationManager.cpp
--- a/Engine/Source/Editor/Kismet/Private/BlueprintCompilationManager.cpp
+++ b/Engine/Source/Editor/Kismet/Private/BlueprintCompilationManager.cpp
@@ -1363,13 +1363,14 @@ void FBlueprintCompilationManagerImpl::FlushCompilationQueueImpl(bool bSuppressB
 				if(BP->GeneratedClass)
 				{
 					BP->GeneratedClass->bLayoutChanging = true;
+					CompilerData.Reinstancer->SaveSparseClassData(BP->GeneratedClass);
 				}
 			}
 		}
 
 		SlowTask.EnterProgressFrame();
 
-		// STAGE XI: Reinstancing done, lets fix up child->parent pointers
+		// STAGE XI: Reinstancing done, lets fix up child->parent pointers and take ownership of SCD:
 		for (FCompilerData& CompilerData : CurrentlyCompilingBPs)
 		{
 			UBlueprint* BP = CompilerData.BP;
@@ -1377,6 +1378,10 @@ void FBlueprintCompilationManagerImpl::FlushCompilationQueueImpl(bool bSuppressB
 			{
 				BP->GeneratedClass->SetSuperStruct(BP->GeneratedClass->GetSuperClass()->GetAuthoritativeClass());
 			}
+			if(BP->GeneratedClass && CompilerData.Reinstancer.IsValid())
+			{
+				CompilerData.Reinstancer->TakeOwnershipOfSparseClassData(BP->GeneratedClass);
+			}
 		}
 
 		SlowTask.EnterProgressFrame();
@@ -1405,6 +1410,8 @@ void FBlueprintCompilationManagerImpl::FlushCompilationQueueImpl(bool bSuppressB
 					// Reset the flag, so if the user tries to use PIE it will warn them if the BP did not compile
 					BP->bDisplayCompilePIEWarning = true;
 		
+					// this will create FProperties for the UClass and generate the sparse class data
+					// if the compiler in question wants to:
 					FKismetCompilerContext& CompilerContext = *(CompilerData.Compiler);
 					CompilerContext.CompileClassLayout(EInternalCompilerFlags::PostponeLocalsGenerationUntilPhaseTwo);
 
@@ -1466,6 +1473,10 @@ void FBlueprintCompilationManagerImpl::FlushCompilationQueueImpl(bool bSuppressB
 					(	BPGC->ClassDefaultObject == nullptr || 
 						BPGC->ClassDefaultObject->GetClass() != BPGC) )
 				{
+					if (CompilerData.Reinstancer.IsValid())
+					{
+						CompilerData.Reinstancer->PropagateSparseClassDataToNewClass(BPGC);
+					}
 					// relink, generate CDO:
 					BPGC->bLayoutChanging = false;
 					BPGC->Bind();
@@ -2079,9 +2090,11 @@ void FBlueprintCompilationManagerImpl::ReparentHierarchies(const TMap<UClass*, U
 
 			for(UClass* DerivedClass : DerivedClasses)
 			{
-				if(DerivedClass->ClassDefaultObject == nullptr)
+				if (DerivedClass->ClassDefaultObject == nullptr && 
+					DerivedClass->GetSparseClassData(EGetSparseClassDataMethod::ReturnIfNull) == nullptr)
 				{
-					// on CDO->no other instances->no need to reinstance..
+					// no CDO->no other instances->no need to reinstance..
+					// and we have no sparse classs data to manage...
 					continue;
 				}
 
@@ -2127,11 +2140,21 @@ void FBlueprintCompilationManagerImpl::ReparentHierarchies(const TMap<UClass*, U
 			) );
 
 			FReinstancingJob& ReinstancingJob = Reinstancers.Last();
+			ReinstancingJob.Reinstancer->SaveSparseClassData(Class);
 			ensure(ReinstancingJob.Reinstancer->DuplicatedClass && ReinstancingJob.Reinstancer->ClassToReinstance);
 		}
 	}
-	
-	// Reparent and Link:
+
+	for (const FReinstancingJob& ReinstancingJob : Reinstancers)
+	{
+		if (ReinstancingJob.Reinstancer.IsValid())
+		{
+			ReinstancingJob.Reinstancer->TakeOwnershipOfSparseClassData(ReinstancingJob.OldToNew.Value);
+		}
+	}
+
+	// Reparent and Link - this is .. kind of pointless.. ReinstanceBatch should
+	// be doing this
 	TMap<UClass*, UClass*> OldClassToNewClassIncludingChildren = OldToNewClasses;
 	for(const FReinstancingJob& ReinstancingJob : Reinstancers)
 	{
@@ -2148,7 +2171,6 @@ void FBlueprintCompilationManagerImpl::ReparentHierarchies(const TMap<UClass*, U
 			UClass* const* NewParent = OldToNewClasses.Find(ClassToReinstance->GetSuperClass());
 			if(NewParent)
 			{
-				ClassToReinstance->ClearSparseClassDataStruct(false);
 				ClassToReinstance->SetSuperStruct(*NewParent);
 			}
 
@@ -2236,7 +2258,8 @@ void FBlueprintCompilationManagerImpl::ReinstanceBatch(TArray<FReinstancingJob>&
 	const auto FilterOutOfDateClasses = [](TArray<UClass*>& ClassList)
 	{
 		// Old versions of classes can be abandoned, classes without CDOs have no instances and don't require reinstancing
-		ClassList.RemoveAllSwap( [](UClass* Class) { return Class->HasAnyClassFlags(CLASS_NewerVersionExists) || Class->ClassDefaultObject == nullptr; } );
+		// but they may still require reparenting..
+		ClassList.RemoveAllSwap( [](UClass* Class) { return Class->HasAnyClassFlags(CLASS_NewerVersionExists); } );
 	};
 
 	const auto HasChildren = [FilterOutOfDateClasses](UClass* InClass) -> bool
@@ -2294,7 +2317,14 @@ void FBlueprintCompilationManagerImpl::ReinstanceBatch(TArray<FReinstancingJob>&
 			{
 				if(IsValid(ClassToReinstance))
 				{
-					ClassesToReinstance.Add(ClassToReinstance);
+					if (ClassToReinstance->ClassDefaultObject)
+					{
+						ClassesToReinstance.Add(ClassToReinstance);
+					}
+					else
+					{
+						ClassesToReparent.Add(ClassToReinstance);
+					}
 				}
 			}
 		}
@@ -2424,7 +2454,28 @@ void FBlueprintCompilationManagerImpl::ReinstanceBatch(TArray<FReinstancingJob>&
 		return Dependencies;
 	});
 
-	// 2. Copy defaults from old CDO - CDO may be missing if this class was reinstanced and relinked here,
+	// 2. Update Sparse Class Data
+	for (const FReinstancingJob* ReinstancingJobPtr : ReinstancersPtr)
+	{
+		TSharedPtr<FBlueprintCompileReinstancer> const& Reinstancer = ReinstancingJobPtr->Reinstancer;
+		if (!Reinstancer.IsValid())
+		{
+			continue; // no reinstancer, we're not responsible for recreating this class (e.g. it came from verse compile or asset reload)
+		}
+
+		// if the new class already has a sparse class data, that indicates
+		// it was generated by its compiler and we can discard the old data:
+		UClass* NewClass = ReinstancingJobPtr->OldToNew.Value;
+		if (!NewClass || NewClass->GetSparseClassData(EGetSparseClassDataMethod::ReturnIfNull) != nullptr)
+		{
+			continue;
+		}
+
+		// New class has not created a sparse class data, if the old class had one copy it over:
+		Reinstancer->PropagateSparseClassDataToNewClass(NewClass);
+	}
+
+	// 3. Copy defaults from old CDO - CDO may be missing if this class was reinstanced and relinked here,
 	// so use GetDefaultObject(true):
 	for (const FReinstancingJob* ReinstancingJobPtr : ReinstancersPtr)
 	{
@@ -2534,7 +2585,7 @@ void FBlueprintCompilationManagerImpl::ReinstanceBatch(TArray<FReinstancingJob>&
 
 	TMap<UObject*, UObject*> OldArchetypeToNewArchetype;
 
-	// 3. Update any remaining instances that are tagged as RF_ArchetypeObject or RF_InheritableComponentTemplate - 
+	// 4. Update any remaining instances that are tagged as RF_ArchetypeObject or RF_InheritableComponentTemplate - 
 	// we may need to do further sorting to ensure that interdependent archetypes are initialized correctly:
 	TSet<UObject*> ArchetypeReferencers;
 
@@ -2746,7 +2797,7 @@ void FBlueprintCompilationManagerImpl::ReinstanceBatch(TArray<FReinstancingJob>&
 	// Reassociate relevant property bags
 	UE::FPropertyBagRepository::Get().ReassociateObjects(OldArchetypeToNewArchetype);
 	
-	// 4. update known references to archetypes (e.g. component templates, WidgetTree). We don't want to run the normal 
+	// 5. update known references to archetypes (e.g. component templates, WidgetTree). We don't want to run the normal 
 	// reference finder to update these because searching the entire object graph is time consuming. Instead we just replace
 	// all references in our UBlueprint and its generated class:
 	for (const FReinstancingJob* ReinstancingJobPtr : ReinstancersPtr)
diff --git a/Engine/Source/Editor/UnrealEd/Private/Kismet2/KismetReinstanceUtilities.cpp b/Engine/Source/Editor/UnrealEd/Private/Kismet2/KismetReinstanceUtilities.cpp
--- a/Engine/Source/Editor/UnrealEd/Private/Kismet2/KismetReinstanceUtilities.cpp
+++ b/Engine/Source/Editor/UnrealEd/Private/Kismet2/KismetReinstanceUtilities.cpp
@@ -32,6 +32,8 @@
 #include "UObject/UObjectIterator.h"
 #include "Serialization/FindObjectReferencers.h"
 #include "Serialization/ArchiveReplaceObjectAndStructPropertyRef.h"
+#include "Serialization/ObjectReader.h"
+#include "Serialization/ObjectWriter.h"
 #include "BlueprintEditor.h"
 #include "Engine/Selection.h"
 #include "BlueprintEditorSettings.h"
@@ -471,6 +473,8 @@ FBlueprintCompileReinstancer::FBlueprintCompileReinstancer(UClass* InClassToRein
 	: ClassToReinstance(InClassToReinstance)
 	, DuplicatedClass(nullptr)
 	, OriginalCDO(nullptr)
+	, OriginalSCD(nullptr)
+	, OriginalSCDStruct(nullptr)
 	, bHasReinstanced(false)
 	, ReinstClassType(RCT_Unknown)
 	, ClassToReinstanceDefaultValuesCRC(0)
@@ -1238,6 +1242,90 @@ void FBlueprintCompileReinstancer::UpdateBytecodeReferences(
 	}
 }
 
+void FBlueprintCompileReinstancer::SaveSparseClassData(const UClass* ForClass)
+{
+	check(ForClass);
+	UClass* SuperClass = ForClass->GetSuperClass();
+	const void* SCD = const_cast<UClass*>(ForClass)->GetSparseClassData(EGetSparseClassDataMethod::ReturnIfNull);
+	if (!SuperClass || !SCD)
+	{
+		return; // null SuperClass should only be possible for UObject, but good to be complete
+	}
+
+	FObjectWriter Writer(SCDSnapshot);
+	UScriptStruct* SuperSCDType = SuperClass->GetSparseClassDataStruct();
+	const void* SuperSCD = SuperClass->GetSparseClassData(EGetSparseClassDataMethod::ReturnIfNull);
+
+	ForClass->GetSparseClassDataStruct()->SerializeTaggedProperties(
+		Writer,
+		(uint8*)SCD,
+		(UStruct*)SuperSCDType,
+		(uint8*)SuperSCD);
+}
+
+void FBlueprintCompileReinstancer::TakeOwnershipOfSparseClassData(UClass* ForClass)
+{
+	check(ForClass);
+	OriginalSCDStruct = ForClass->GetSparseClassDataStruct();
+	if (!OriginalSCDStruct)
+	{
+		return;
+	}
+
+	OriginalSCD = const_cast<void*>(ForClass->GetSparseClassData(EGetSparseClassDataMethod::ReturnIfNull));
+	if (OriginalSCDStruct->GetOuter() == ForClass)
+	{
+		OriginalSCDStruct->Rename(nullptr, DuplicatedClass, REN_DoNotDirty | REN_DontCreateRedirectors | REN_ForceNoResetLoaders | REN_NonTransactional);
+	}
+	// We own these now, remove ForClass's knowledge of the sparse class data - they
+	// will be freed when reinstancing is complete:
+	ForClass->SparseClassData = nullptr;
+	ForClass->SparseClassDataStruct = nullptr;
+}
+
+void FBlueprintCompileReinstancer::PropagateSparseClassDataToNewClass(UClass* NewClass)
+{
+	if (!OriginalSCD || 
+		NewClass->GetSparseClassData(EGetSparseClassDataMethod::ReturnIfNull))
+	{
+		return;
+	}
+
+	UScriptStruct* SparseClassDataStruct = OriginalSCDStruct;
+	if (UScriptStruct* NewSCD = NewClass->GetSparseClassDataStruct())
+	{
+		SparseClassDataStruct = NewSCD;
+	}
+
+	if (!IsValid(SparseClassDataStruct) ||
+		SparseClassDataStruct->GetOutermost() == GetTransientPackage())
+	{
+		return;
+	}
+
+	if (SparseClassDataStruct == OriginalSCDStruct && SparseClassDataStruct->GetOuter() == DuplicatedClass)
+	{
+		SparseClassDataStruct->Rename(nullptr, NewClass, REN_DoNotDirty | REN_DontCreateRedirectors | REN_ForceNoResetLoaders | REN_NonTransactional);
+	}
+	NewClass->SparseClassDataStruct = SparseClassDataStruct;
+	NewClass->CreateSparseClassData();
+
+	FObjectReader Reader(SCDSnapshot);
+	SparseClassDataStruct->SerializeTaggedProperties(
+		Reader,
+		(uint8*)NewClass->SparseClassData,
+		nullptr,
+		nullptr);
+
+	if (OriginalSCD && OriginalSCDStruct)
+	{
+		OriginalSCDStruct->DestroyStruct(OriginalSCD);
+		FMemory::Free(OriginalSCD);
+		OriginalSCD = nullptr;
+		OriginalSCDStruct = nullptr;
+	}
+}
+
 void FBlueprintCompileReinstancer::FinishUpdateBytecodeReferences(
 	const TSet<UBlueprint*>& DependentBPs,
 	const TMap<FFieldVariant, FFieldVariant>& FieldMappings)
diff --git a/Engine/Source/Editor/UnrealEd/Public/Kismet2/KismetReinstanceUtilities.h b/Engine/Source/Editor/UnrealEd/Public/Kismet2/KismetReinstanceUtilities.h
--- a/Engine/Source/Editor/UnrealEd/Public/Kismet2/KismetReinstanceUtilities.h
+++ b/Engine/Source/Editor/UnrealEd/Public/Kismet2/KismetReinstanceUtilities.h
@@ -119,6 +119,15 @@ protected:
 	/** The original CDO object for the class being actively reinstanced */
 	TObjectPtr<UObject>	OriginalCDO;
 
+	/** The original Sparse Class Data object for the class being actively reinstanced */
+	void* OriginalSCD;
+
+	/** The original SDO Struct for the class being actively reinstanced */
+	TObjectPtr<UScriptStruct> OriginalSCDStruct;
+
+	/** A snapshot of the SCD, currently delta serialized from its archetype - taken before ownership of SCDs is taken */
+	TArray<uint8> SCDSnapshot;
+
 	/** Children of this blueprint, which will need to be recompiled and relinked temporarily to maintain the class layout */
 	TArray<UBlueprint*> Children;
 
@@ -185,6 +194,15 @@ public:
 	/** Updates references to properties and functions of the class that has in the bytecode of dependent blueprints */
 	UNREALED_API void UpdateBytecodeReferences( TSet<UBlueprint*>& OutDependentBlueprints, TMap<FFieldVariant, FFieldVariant>& OutFieldMapping);
 
+	/** Populates SCDSnapshot, for use via PropagateSparseClassDataToNewClass */
+	UNREALED_API void SaveSparseClassData(const UClass* ForClass);
+
+	/** Instructs the reinstancer to take ownership of the sparse class data - doing so will make it difficult to identify archetype data */ 
+	UNREALED_API void TakeOwnershipOfSparseClassData(UClass* ForClass);
+
+	/** Copies an owned sparse class data instance to a new class and frees any owned SCD */
+	UNREALED_API void PropagateSparseClassDataToNewClass(UClass* NewClass);
+
 	/** Consumes the set and map populated by calls to UpdateBytecodeReferences */
 	static UNREALED_API void FinishUpdateBytecodeReferences( const TSet<UBlueprint*>& DependentBPs, const TMap<FFieldVariant, FFieldVariant>& FieldMappings);
 
diff --git a/Engine/Source/Runtime/CoreUObject/Private/UObject/Class.cpp b/Engine/Source/Runtime/CoreUObject/Private/UObject/Class.cpp
--- a/Engine/Source/Runtime/CoreUObject/Private/UObject/Class.cpp
+++ b/Engine/Source/Runtime/CoreUObject/Private/UObject/Class.cpp
@@ -6037,35 +6037,8 @@ void UClass::ClearSparseClassDataStruct(bool bInRecomplingOnLoad)
 { 
 	if (SparseClassDataStruct != nullptr)
 	{
-		// Find all subclasses and clear their sparse class data struct as well.
-		TArray<UClass*> SubClasses;
-		TArray<UClass*> SubClassesToClear;
-		GetDerivedClasses(this, SubClasses, true /* bRecursive */);
-		for (UClass* SubClass : SubClasses)
-		{
-			UScriptStruct* SubClassSparseClassDataStruct = SubClass->GetSparseClassDataStruct();
-			if (SubClassSparseClassDataStruct && SubClassSparseClassDataStruct->IsChildOf(SparseClassDataStruct))
-			{
-				SubClassesToClear.Add(SubClass);
-			}
-		}
-
-		auto ClearSparseClassDataStructInner = [bInRecomplingOnLoad](UClass* InClassToClear)
-		{
-			UScriptStruct* CurrentSparseClassDataStruct = InClassToClear->GetSparseClassDataStruct();
-			InClassToClear->CleanupSparseClassData();
-			InClassToClear->SparseClassDataStruct = nullptr;
-			CurrentSparseClassDataStruct->SetSuperStruct(nullptr);
-			const ERenameFlags RenameFlags = REN_DontCreateRedirectors | ((bInRecomplingOnLoad) ? REN_ForceNoResetLoaders : 0) | REN_NonTransactional | REN_DoNotDirty;
-			CurrentSparseClassDataStruct->Rename(nullptr, GetTransientPackage(), RenameFlags);
-		};
-
-		for(UClass* SubClassToClear : SubClassesToClear)
-		{
-			ClearSparseClassDataStructInner(SubClassToClear);
-		}
-
-		ClearSparseClassDataStructInner(this);
+		CleanupSparseClassData();
+		SparseClassDataStruct = nullptr;
 	}
 }
 
diff --git a/Engine/Source/Runtime/CoreUObject/Public/UObject/Class.h b/Engine/Source/Runtime/CoreUObject/Public/UObject/Class.h
--- a/Engine/Source/Runtime/CoreUObject/Public/UObject/Class.h
+++ b/Engine/Source/Runtime/CoreUObject/Public/UObject/Class.h
@@ -2740,6 +2740,7 @@ class UClass : public UStruct
 public:
 	friend class FRestoreClassInfo;
 	friend class FBlueprintEditorUtils;
+	friend class FBlueprintCompileReinstancer;
 
 	typedef void		(*ClassConstructorType)				(const FObjectInitializer&);
 	typedef UObject*	(*ClassVTableHelperCtorCallerType)	(FVTableHelper& Helper);
diff --git a/Engine/Source/Runtime/Engine/Private/Animation/AnimBlueprintGeneratedClass.cpp b/Engine/Source/Runtime/Engine/Private/Animation/AnimBlueprintGeneratedClass.cpp
--- a/Engine/Source/Runtime/Engine/Private/Animation/AnimBlueprintGeneratedClass.cpp
+++ b/Engine/Source/Runtime/Engine/Private/Animation/AnimBlueprintGeneratedClass.cpp
@@ -471,46 +471,6 @@ void UAnimBlueprintGeneratedClass::Link(FArchive& Ar, bool bRelinkExistingProper
 	{
 		MutablesStruct->Link(Ar, bRelinkExistingProperties);
 	}
-
-	// Link in sparse class data - it must mirror the generated class hierarchy
-	if(UAnimBlueprintGeneratedClass* ParentClass = Cast<UAnimBlueprintGeneratedClass>(GetSuperClass()))
-	{
-		if(UScriptStruct* ParentSparseClassDataStruct = ParentClass->GetSparseClassDataStruct())
-		{
-			if(SparseClassDataStruct)
-			{
-				// Make sure there are no duplicate SparseClassDataStructs in the parent heirarchy
-				bool bAllSparseClassDataStructsUnique = true;
-				TArray<UScriptStruct*, TInlineAllocator<4>> HierarchySparseClassDataStructs;
-				HierarchySparseClassDataStructs.Add(SparseClassDataStruct);
-				for (UClass* CurClass = ParentClass; CurClass; CurClass = CurClass->GetSuperClass())
-				{
-					if (UScriptStruct* CurSparseClassDataStruct = CurClass->GetSparseClassDataStruct())
-					{
-						if (HierarchySparseClassDataStructs.Contains(CurSparseClassDataStruct))
-						{
-							bAllSparseClassDataStructsUnique = false;
-							break;
-						}
-						else
-						{
-							HierarchySparseClassDataStructs.Add(CurSparseClassDataStruct);
-						}
-					}
-				}
-
-				if (bAllSparseClassDataStructsUnique)
-				{
-					// Ensure parent is linked before setting super
-					ParentSparseClassDataStruct->Link(Ar, bRelinkExistingProperties);
-					SparseClassDataStruct->SetSuperStruct(ParentSparseClassDataStruct);
-
-					// Link sparse class data now it is correctly parented
-					SparseClassDataStruct->Link(Ar, bRelinkExistingProperties);
-				}
-			}
-		}
-	}
 #endif	
 	
 	Super::Link(Ar, bRelinkExistingProperties);
