Fix crash when undoing creation of a custom event node

diff --git a/Engine/Source/Editor/Kismet/Private/BlueprintDetailsCustomization.cpp b/Engine/Source/Editor/Kismet/Private/BlueprintDetailsCustomization.cpp
--- a/Engine/Source/Editor/Kismet/Private/BlueprintDetailsCustomization.cpp
+++ b/Engine/Source/Editor/Kismet/Private/BlueprintDetailsCustomization.cpp
@@ -4336,7 +4336,7 @@ void FBlueprintGraphActionDetails::CustomizeDetails( IDetailLayoutBuilder& Detai
 					if (SelectedNode.IsValid())
 					{
 						UK2Node_CustomEvent const* SelectedCustomEvent = Cast<UK2Node_CustomEvent const>(SelectedNode.Get());
-						check(SelectedCustomEvent != NULL);
+						check(SelectedCustomEvent != nullptr);
 
 						bIsOverride = SelectedCustomEvent->IsOverride();
 					}
@@ -4359,10 +4359,10 @@ void FBlueprintGraphActionDetails::CustomizeDetails( IDetailLayoutBuilder& Detai
 				static bool CanSetReliabilityProperty(TWeakObjectPtr<UK2Node_EditablePinBase> SelectedNode)
 				{
 					bool bIsReliabilitySettingEnabled = false;
-					if (IsNotCustomEventOverride(SelectedNode))
+					if (IsNotCustomEventOverride(SelectedNode) && SelectedNode.IsValid())
 					{
 						UK2Node_CustomEvent const* SelectedCustomEvent = Cast<UK2Node_CustomEvent const>(SelectedNode.Get());
-						check(SelectedCustomEvent != NULL);
+						check(SelectedCustomEvent != nullptr);
 
 						bIsReliabilitySettingEnabled = ((SelectedCustomEvent->GetNetFlags() & FUNC_Net) != 0);
 					}
