Fixed "Fix Self Context Function References" window being the incorrect size on UI-scaled monitors.

diff --git a/Engine/Source/Editor/Kismet/Private/SFixupSelfContextDlg.cpp b/Engine/Source/Editor/Kismet/Private/SFixupSelfContextDlg.cpp
--- a/Engine/Source/Editor/Kismet/Private/SFixupSelfContextDlg.cpp
+++ b/Engine/Source/Editor/Kismet/Private/SFixupSelfContextDlg.cpp
@@ -177,8 +177,7 @@ bool SFixupSelfContextDialog::CreateModal(const TArray<UK2Node_CallFunction*>& N
 	Window = SNew(SWindow)
 		.Title(LOCTEXT("FixupReferencesTitle", "Fix Self Context Function References"))
 		.SizingRule(ESizingRule::UserSized)
-		.MinWidth(400.f)
-		.MinHeight(300.f)
+		.ClientSize(FVector2D(400.f, 300.f))
 		.SupportsMaximize(true)
 		.SupportsMinimize(false)
 		.HasCloseButton(false)
