Fix dependency order when compiling verse with prefabs that contains at lease 4 levels of inheritance!

diff --git a/Engine/Source/Editor/Kismet/Private/BlueprintCompilationManager.cpp b/Engine/Source/Editor/Kismet/Private/BlueprintCompilationManager.cpp
--- a/Engine/Source/Editor/Kismet/Private/BlueprintCompilationManager.cpp
+++ b/Engine/Source/Editor/Kismet/Private/BlueprintCompilationManager.cpp
@@ -2275,7 +2275,8 @@ void FBlueprintCompilationManagerImpl::ReinstanceBatch(TArray<FReinstancingJob>&
 		return ChildTypes.Num() > 0;
 	};
 
-	TSet<UClass*> ClassesToReparent;
+	TSet<UClass*> ClassesToReparentSet;
+	TArray<UClass*> ClassesToReparent;
 	TSet<UClass*> ClassesToReinstance;
 
 	// Reinstancers may contain *part* of a class hierarchy, so we first need to reparent any child types that 
@@ -2328,7 +2329,12 @@ void FBlueprintCompilationManagerImpl::ReinstanceBatch(TArray<FReinstancingJob>&
 					}
 					else
 					{
-						ClassesToReparent.Add(ClassToReinstance);
+						bool bAlreadyInSet = false;
+						ClassesToReparentSet.Add(ClassToReinstance, &bAlreadyInSet);
+						if (!bAlreadyInSet)
+						{
+							ClassesToReparent.Add(ClassToReinstance);
+						}
 					}
 				}
 			}
@@ -2344,12 +2350,38 @@ void FBlueprintCompilationManagerImpl::ReinstanceBatch(TArray<FReinstancingJob>&
 			{
 				if(IsValid(ClassToReparent))
 				{
-					ClassesToReparent.Add(ClassToReparent);
+					bool bAlreadyInSet = false;
+					ClassesToReparentSet.Add(ClassToReparent, &bAlreadyInSet);
+					if (!bAlreadyInSet)
+					{
+						ClassesToReparent.Add(ClassToReparent);
+					}
 				}
 			}
 		}
 	}
 
+	Algo::TopologicalSort(ClassesToReparent, [&InOutOldToNewClassMap](UClass* Class)
+	{
+		TArray<UClass*> Dependencies;
+
+		UClass* CurrentClass = Class;
+		while (UClass* SuperClass = CurrentClass->GetSuperClass())
+		{
+			if(UClass** NewSuperClass = InOutOldToNewClassMap.Find(SuperClass))
+			{
+				Dependencies.Add(*NewSuperClass);
+			}
+			else
+			{
+				Dependencies.Add(SuperClass);
+			}
+			CurrentClass = SuperClass;
+		}
+
+		return Dependencies;
+	});
+
 	for(UClass* Class : ClassesToReparent)
 	{
 		UClass** NewParent = InOutOldToNewClassMap.Find(Class->GetSuperClass());
