Made the "Override" button in the Blueprint Functions UI always visible.

diff --git a/Engine/Source/Editor/Kismet/Private/SMyBlueprint.cpp b/Engine/Source/Editor/Kismet/Private/SMyBlueprint.cpp
--- a/Engine/Source/Editor/Kismet/Private/SMyBlueprint.cpp
+++ b/Engine/Source/Editor/Kismet/Private/SMyBlueprint.cpp
@@ -880,7 +880,7 @@ TSharedRef<SWidget> SMyBlueprint::OnGetSectionWidget(TSharedRef<SWidget> RowWidg
 				[
 					SAssignNew(FunctionSectionButton, SComboButton)
 					.IsEnabled(this, &SMyBlueprint::IsEditingMode)
-					.Visibility(this, &SMyBlueprint::OnGetSectionTextVisibility, WeakRowWidget, InSectionID)
+					.Visibility(EVisibility::SelfHitTestInvisible)
 					.ForegroundColor(FAppStyle::GetSlateColor("DefaultForeground"))
 					.OnGetMenuContent(this, &SMyBlueprint::OnGetFunctionListMenu)
 					.ContentPadding(0.0f)
@@ -1021,25 +1021,6 @@ bool SMyBlueprint::HandleActionMatchesName(FEdGraphSchemaAction* InAction, const
 	return false;
 }
 
-EVisibility SMyBlueprint::OnGetSectionTextVisibility(TWeakPtr<SWidget> RowWidget, int32 InSectionID) const
-{
-	bool ShowText = RowWidget.Pin()->IsHovered();
-	if ( InSectionID == NodeSectionID::FUNCTION && FunctionSectionButton.IsValid() && FunctionSectionButton->IsOpen() )
-	{
-		ShowText = true;
-	}
-
-	// If the row is currently hovered, or a menu is being displayed for a button, keep the button expanded.
-	if ( ShowText )
-	{
-		return EVisibility::SelfHitTestInvisible;
-	}
-	else
-	{
-		return EVisibility::Hidden;
-	}
-}
-
 TSharedRef<SWidget> SMyBlueprint::OnGetFunctionListMenu()
 {
 	const bool bShouldCloseWindowAfterMenuSelection = true;
