Fix an old crash that occurs when two ghost nodes are replaced when a node is placed and autowired. The autowiring logic was buggy.

diff --git a/Engine/Source/Editor/Kismet/Private/BlueprintActionMenuItem.cpp b/Engine/Source/Editor/Kismet/Private/BlueprintActionMenuItem.cpp
--- a/Engine/Source/Editor/Kismet/Private/BlueprintActionMenuItem.cpp
+++ b/Engine/Source/Editor/Kismet/Private/BlueprintActionMenuItem.cpp
@@ -74,7 +74,7 @@ namespace FBlueprintMenuActionItemImpl
 	 * @param  SpawnedNodesBeginIndex    
 	 * @return 
 	 */
-	static UEdGraphNode* AutowireSpawnedNodes(UEdGraphPin* FromPin, const TArray<UEdGraphNode*>& GraphNodes, int32 const NodesBeginIndex);
+	static UEdGraphNode* AutowireSpawnedNodes(UEdGraphPin* FromPin, const TArray<UEdGraphNode*>& GraphNodes, TArray<UEdGraphNode*>& NewNodes);
 }
 
 //------------------------------------------------------------------------------
@@ -162,19 +162,8 @@ static bool FBlueprintMenuActionItemImpl::IsNodeLinked(UEdGraphPin* LeadingPin,
 }
 
 //------------------------------------------------------------------------------
-static UEdGraphNode* FBlueprintMenuActionItemImpl::AutowireSpawnedNodes(UEdGraphPin* FromPin, const TArray<UEdGraphNode*>& GraphNodes, int32 const NodesBeginIndex)
+static UEdGraphNode* FBlueprintMenuActionItemImpl::AutowireSpawnedNodes(UEdGraphPin* FromPin, const TArray<UEdGraphNode*>& GraphNodes, TArray<UEdGraphNode*>& OrderedNewNodes)
 {
-	int32 const CurrentGraphNodeCount = GraphNodes.Num();
-	int32 const NewNodeCount = CurrentGraphNodeCount - NodesBeginIndex;
-
-	TArray<UEdGraphNode*> OrderedNewNodes;
-	OrderedNewNodes.Reserve(NewNodeCount);
-	// @TODO: there's gotta be a better way to blit these in
-	for (int32 NodexIndex = NodesBeginIndex; NodexIndex < CurrentGraphNodeCount; ++NodexIndex)
-	{
-		OrderedNewNodes.Add(GraphNodes[NodexIndex]);
-	}
-
 	const EEdGraphPinDirection PinDirection = FromPin->Direction;
 	// should lhs come before rhs?
 	OrderedNewNodes.Sort([PinDirection](UEdGraphNode& Lhs, UEdGraphNode& Rhs)->bool
@@ -286,10 +275,12 @@ UEdGraphNode* FBlueprintActionMenuItem::PerformAction(UEdGraph* ParentGraph, UEd
 			}
 		}
 
-		const int32 PreInvokeNodeCount = ParentGraph->Nodes.Num();
+		const TSet<UEdGraphNode*> OldNodes(ParentGraph->Nodes);
 
 		bool bNewNode = false;
 		LastSpawnedNode = InvokeAction(Action, ParentGraph, ModifiedLocation, BindingsSubset, /*out*/ bNewNode);
+
+		TArray<UEdGraphNode*> NewNodes = TSet<UEdGraphNode*>(ParentGraph->Nodes).Difference(OldNodes).Array();
 		// could already be an existent node, so we have to add here (can't 
 		// catch it as we go through all new nodes)
 		NodesToFocus.Add(LastSpawnedNode);
@@ -300,7 +291,7 @@ UEdGraphNode* FBlueprintActionMenuItem::PerformAction(UEdGraph* ParentGraph, UEd
 		{
 			// make sure to auto-wire after we position the new node (in case
 			// the auto-wire creates a conversion node to put between them)
-			FBlueprintMenuActionItemImpl::AutowireSpawnedNodes(FromPin, ParentGraph->Nodes, PreInvokeNodeCount);
+			FBlueprintMenuActionItemImpl::AutowireSpawnedNodes(FromPin, ParentGraph->Nodes, NewNodes);
 		}
 
 		if (bNewNode)
