Fixing reinstancing issue with nested subobjects and transaction buffer not collecting all subobject info

diff --git a/Engine/Source/Editor/Kismet/Private/BlueprintCompilationManager.cpp b/Engine/Source/Editor/Kismet/Private/BlueprintCompilationManager.cpp
--- a/Engine/Source/Editor/Kismet/Private/BlueprintCompilationManager.cpp
+++ b/Engine/Source/Editor/Kismet/Private/BlueprintCompilationManager.cpp
@@ -2236,11 +2236,11 @@ void FBlueprintCompilationManagerImpl::ReparentHierarchies(const TMap<UClass*, U
 
 void FBlueprintCompilationManagerImpl::BuildDSOMap(UObject* OldObject, UObject* NewObject, TMap<UObject*, UObject*>& OutOldToNewDSO)
 {
-	// IsDefaultObject() unfortunately cannot be relied upon for archetypes, so we explicitly search for the flag:
+	// IsDefaultSubObject() unfortunately cannot be relied upon for archetypes, so we explicitly search for the flag:
 	TArray<UObject*> OldSubobjects;
 	ForEachObjectWithOuter(OldObject, [&OldSubobjects](UObject* Object)
 	{
-		if (Object->HasAnyFlags(RF_DefaultSubObject))
+		if (Object->HasAnyFlags(RF_DefaultSubObject|RF_ArchetypeObject))
 		{
 			OldSubobjects.Add(Object);
 		}
@@ -2648,6 +2648,8 @@ void FBlueprintCompilationManagerImpl::ReinstanceBatch(TArray<FReinstancingJob>&
 				OldClass->ClassDefaultObject != NewClass->ClassDefaultObject)
 			{
 				OldArchetypeToNewArchetype.Add(OldClass->ClassDefaultObject, NewClass->ClassDefaultObject);
+				// also map old *default* subobjects to new default subobjects:
+				BuildDSOMap(OldClass->ClassDefaultObject, NewClass->ClassDefaultObject, OldArchetypeToNewArchetype);
 			}
 
 			TArray<UObject*> ArchetypeObjects;
