Blueprint: FindInBlueprint now considers Extension Graphs for reference search

diff --git a/Engine/Source/Editor/Kismet/Private/FindInBlueprintManager.cpp b/Engine/Source/Editor/Kismet/Private/FindInBlueprintManager.cpp
--- a/Engine/Source/Editor/Kismet/Private/FindInBlueprintManager.cpp
+++ b/Engine/Source/Editor/Kismet/Private/FindInBlueprintManager.cpp
@@ -77,6 +77,7 @@ const FText FFindInBlueprintSearchTags::FiB_UberGraphs = LOCTEXT("Uber", "Uber")
 const FText FFindInBlueprintSearchTags::FiB_Functions = LOCTEXT("Functions", "Functions");
 const FText FFindInBlueprintSearchTags::FiB_Macros = LOCTEXT("Macros", "Macros");
 const FText FFindInBlueprintSearchTags::FiB_SubGraphs = LOCTEXT("Sub", "Sub");
+const FText FFindInBlueprintSearchTags::FiB_ExtensionGraphs = LOCTEXT("ExtGraphs", "ExtGraphs");
 const FText FFindInBlueprintSearchTags::FiB_Extensions = LOCTEXT("Extensions", "Extensions");
 
 const FText FFindInBlueprintSearchTags::FiB_Name = LOCTEXT("Name", "Name");
@@ -1283,6 +1284,17 @@ namespace BlueprintSearchMetaDataHelpers
 		// Gather normal event graphs
 		GatherGraphSearchData(InWriter, Blueprint, Blueprint->UbergraphPages, FFindInBlueprintSearchTags::FiB_UberGraphs, &SubGraphs);
 
+		// Gather extension graphs
+		for (UBlueprintExtension* Extension : Blueprint->GetExtensions())
+		{
+			if (Extension)
+			{
+				TArray<UEdGraph*> ExtensionGraphs;
+				Extension->GetAllGraphs(ExtensionGraphs);
+				GatherGraphSearchData(InWriter, Blueprint, ExtensionGraphs, FFindInBlueprintSearchTags::FiB_ExtensionGraphs, &SubGraphs);
+			}
+		}
+
 		// We have interface graphs and function graphs to put into the Functions category. We cannot do them separately, so we must compile the full list
 		{
 			TArray<UEdGraph*> CompleteGraphList;
diff --git a/Engine/Source/Editor/Kismet/Private/ImaginaryBlueprintData.cpp b/Engine/Source/Editor/Kismet/Private/ImaginaryBlueprintData.cpp
--- a/Engine/Source/Editor/Kismet/Private/ImaginaryBlueprintData.cpp
+++ b/Engine/Source/Editor/Kismet/Private/ImaginaryBlueprintData.cpp
@@ -646,6 +646,11 @@ bool FImaginaryBlueprint::TrySpecialHandleJsonValue(FText InKey, TSharedPtr< FJs
 		ParseGraph(InJsonValue, FFindInBlueprintSearchTags::FiB_SubGraphs.ToString(), GT_Ubergraph);
 		bResult = true;
 	}
+	else if (!InKey.CompareTo(FFindInBlueprintSearchTags::FiB_ExtensionGraphs))
+	{
+		ParseGraph(InJsonValue, FFindInBlueprintSearchTags::FiB_ExtensionGraphs.ToString(), GT_Ubergraph);
+		bResult = true;
+	}
 	else if(!InKey.CompareTo(FFindInBlueprintSearchTags::FiB_Components))
 	{
 		TArray<TSharedPtr< FJsonValue > > ComponentList = InJsonValue->AsArray();
diff --git a/Engine/Source/Editor/Kismet/Public/FindInBlueprintManager.h b/Engine/Source/Editor/Kismet/Public/FindInBlueprintManager.h
--- a/Engine/Source/Editor/Kismet/Public/FindInBlueprintManager.h
+++ b/Engine/Source/Editor/Kismet/Public/FindInBlueprintManager.h
@@ -90,6 +90,8 @@ struct KISMET_API FFindInBlueprintSearchTags
 	static const FText FiB_Macros;
 	/** Sub graph tag, for any sub-graphs in a Blueprint */
 	static const FText FiB_SubGraphs;
+	/** Extension graph tag, for any graphs in a Blueprint Extension */
+	static const FText FiB_ExtensionGraphs;
 	/** Blueprint extension tag. */
 	static const FText FiB_Extensions;
 
