Fixed a crash when reloading a blueprint when the diff window is open

diff --git a/Engine/Source/Editor/Kismet/Private/BlueprintDetailsCustomization.cpp b/Engine/Source/Editor/Kismet/Private/BlueprintDetailsCustomization.cpp
--- a/Engine/Source/Editor/Kismet/Private/BlueprintDetailsCustomization.cpp
+++ b/Engine/Source/Editor/Kismet/Private/BlueprintDetailsCustomization.cpp
@@ -6651,9 +6651,9 @@ float FBlueprintGlobalOptionsDetails::NamespacePropertyValueCustomization_MinDes
 
 UBlueprint* FBlueprintGlobalOptionsDetails::GetBlueprintObj() const
 {
-	if(BlueprintObjOverride)
+	if(UBlueprint* BP = BlueprintObjOverride.Get())
 	{
-		return BlueprintObjOverride;
+		return BP;
 	}
 	
 	if(BlueprintEditorPtr.IsValid())
diff --git a/Engine/Source/Editor/Kismet/Private/BlueprintDetailsCustomization.h b/Engine/Source/Editor/Kismet/Private/BlueprintDetailsCustomization.h
--- a/Engine/Source/Editor/Kismet/Private/BlueprintDetailsCustomization.h
+++ b/Engine/Source/Editor/Kismet/Private/BlueprintDetailsCustomization.h
@@ -939,7 +939,7 @@ public:
 
 	}
 	
-	FBlueprintGlobalOptionsDetails(UBlueprint* InBlueprintPtr)
+	FBlueprintGlobalOptionsDetails(TWeakObjectPtr<UBlueprint> InBlueprintPtr)
 		:BlueprintObjOverride(InBlueprintPtr)
 	{
 
@@ -953,7 +953,7 @@ public:
 
 	/** Diff functionality doesn't need access to the FBlueprintEditor so this allows creation without editor access
 	 *  but with limited functionality */
-	static TSharedRef<IDetailCustomization> MakeInstanceForDiff(UBlueprint* InBlueprintPtr)
+	static TSharedRef<IDetailCustomization> MakeInstanceForDiff(TWeakObjectPtr<UBlueprint> InBlueprintPtr)
 	{
 		return MakeShareable(new FBlueprintGlobalOptionsDetails(InBlueprintPtr));
 	}
@@ -1008,8 +1008,8 @@ private:
 	/** Weak reference to the Blueprint editor */
 	TWeakPtr<FBlueprintEditor> BlueprintEditorPtr;
 	
-	/** Weak reference to the Blueprint editor */
-	TObjectPtr<UBlueprint> BlueprintObjOverride;
+	/** Weak reference to a Blueprint override - used to customize a blueprint that is displayed without an editor, e.g. for diff */
+	TWeakObjectPtr<UBlueprint> BlueprintObjOverride;
 
 	/** Combo button used to choose a parent class */
 	TSharedPtr<SComboButton> ParentClassComboButton;
diff --git a/Engine/Source/Editor/Kismet/Private/DetailsDiff.cpp b/Engine/Source/Editor/Kismet/Private/DetailsDiff.cpp
--- a/Engine/Source/Editor/Kismet/Private/DetailsDiff.cpp
+++ b/Engine/Source/Editor/Kismet/Private/DetailsDiff.cpp
@@ -154,10 +154,11 @@ TSharedRef<IDetailsView> FDetailsDiff::CreateDetailsView(const UObject* InObject
 	DetailsView->SetIsPropertyEditingEnabledDelegate(FIsPropertyEditingEnabled::CreateStatic([]{return false; }));
 	if (InObject && InObject->IsA<UBlueprint>())
 	{
+		TWeakObjectPtr<UBlueprint> ForCapture = const_cast<UBlueprint*>(Cast<UBlueprint>(InObject));
 		// create a custom property layout so that sections like interfaces will be included in the diff view
 		const FOnGetDetailCustomizationInstance LayoutOptionDetails = FOnGetDetailCustomizationInstance::CreateStatic(
 			&FBlueprintGlobalOptionsDetails::MakeInstanceForDiff,
-			const_cast<UBlueprint *>(Cast<UBlueprint>(InObject))
+			ForCapture
 		);
 		DetailsView->RegisterInstancedCustomPropertyLayout(UBlueprint::StaticClass(), LayoutOptionDetails);
 	}
