Add LLM tags for the BP Subobject Editor Viewport client creation.

#rnx

diff --git a/Engine/Source/Editor/Kismet/Private/BlueprintEditor.cpp b/Engine/Source/Editor/Kismet/Private/BlueprintEditor.cpp
--- a/Engine/Source/Editor/Kismet/Private/BlueprintEditor.cpp
+++ b/Engine/Source/Editor/Kismet/Private/BlueprintEditor.cpp
@@ -3024,6 +3024,8 @@ void FBlueprintEditor::CreateSubobjectEditors()
 		.OnItemDoubleClicked(this, &FBlueprintEditor::OnComponentDoubleClicked)
 		.SubobjectClassListFilters(ClassFilters);
 	
+	LLM_SCOPE_BYNAME(TEXT("BPCreateSubobjectEditorViewportClient"));
+
 	SubobjectViewport = SAssignNew(SubobjectViewport, SSCSEditorViewport)
 		.BlueprintEditor(SharedThis(this));
 
