Removed unnecessary FBPCompileRequestInternal struct.

At one point, this struct had more fields, but it's now just an unnecessary wrapper around FBPCompileRequest.

Also made some minor whitespace fixes.

diff --git a/Engine/Source/Editor/Kismet/Private/BlueprintCompilationManager.cpp b/Engine/Source/Editor/Kismet/Private/BlueprintCompilationManager.cpp
--- a/Engine/Source/Editor/Kismet/Private/BlueprintCompilationManager.cpp
+++ b/Engine/Source/Editor/Kismet/Private/BlueprintCompilationManager.cpp
@@ -100,16 +100,6 @@ struct FReinstancingJob;
 struct FSkeletonFixupData;
 struct FCompilerData;
 
-struct FBPCompileRequestInternal
-{
-	FBPCompileRequestInternal(FBPCompileRequest InUserData)
-		: UserData(InUserData)
-	{
-	}
-
-	FBPCompileRequest UserData;
-};
-
 enum class EReparentClassOptions
 {
 	None = 0x0,
@@ -129,8 +119,8 @@ struct FBlueprintCompilationManagerImpl : public FGCObject
 
 	void RegisterCompilerExtension(TSubclassOf<UBlueprint> BlueprintType, UBlueprintCompilerExtension* Extension);
 
-	void QueueForCompilation(const FBPCompileRequestInternal& CompileJob);
-	void CompileSynchronouslyImpl(const FBPCompileRequestInternal& Request);
+	void QueueForCompilation(const FBPCompileRequest& CompileJob);
+	void CompileSynchronouslyImpl(const FBPCompileRequest& Request);
 	void FlushCompilationQueueImpl(bool bSuppressBroadcastCompiled, TArray<UBlueprint*>* BlueprintsCompiled, TArray<UBlueprint*>* BlueprintsCompiledOrSkeletonCompiled, FUObjectSerializeContext* InLoadContext, TMap<UClass*, TMap<UObject*, UObject*>>* OldToNewTemplates = nullptr);
 	void FixupDelegateProperties(const TArray<FCompilerData>& CurrentlyCompilingBPs);
 	void ProcessExtensions(const TArray<FCompilerData>& InCurrentlyCompilingBPs);
@@ -162,10 +152,10 @@ struct FBlueprintCompilationManagerImpl : public FGCObject
 	// Extension data, could be organized in many ways, but this provides an easy way
 	// to extend blueprint compilation after the graph has been pruned and functions
 	// have been generated (but before code is generated):
-	TMap<TObjectPtr<UClass>, TArray<TObjectPtr<UBlueprintCompilerExtension>> > CompilerExtensions;
+	TMap<TObjectPtr<UClass>, TArray<TObjectPtr<UBlueprintCompilerExtension>>> CompilerExtensions;
 
 	// Queued requests to be processed in the next FlushCompilationQueueImpl call:
-	TArray<FBPCompileRequestInternal> QueuedRequests;
+	TArray<FBPCompileRequest> QueuedRequests;
 	
 	// Data stored for reinstancing, which finishes much later than compilation,
 	// populated by FlushCompilationQueueImpl, cleared by FlushReinstancingQueueImpl:
@@ -210,15 +200,15 @@ FBlueprintCompilationManagerImpl::~FBlueprintCompilationManagerImpl()
 
 void FBlueprintCompilationManagerImpl::AddReferencedObjects(FReferenceCollector& Collector)
 {
-	for(auto& Extensions : CompilerExtensions)
+	for (auto& Extensions : CompilerExtensions)
 	{
 		Collector.AddReferencedObject(Extensions.Key);
 		Collector.AddReferencedObjects(Extensions.Value);
 	}
 
-	for( FBPCompileRequestInternal& Job : QueuedRequests )
+	for (FBPCompileRequest& Job : QueuedRequests)
 	{
-		Collector.AddReferencedObject(Job.UserData.BPToCompile);
+		Collector.AddReferencedObject(Job.BPToCompile);
 	}
 
 	Collector.AddReferencedObjects(ClassesToReinstance);
@@ -235,14 +225,14 @@ void FBlueprintCompilationManagerImpl::RegisterCompilerExtension(TSubclassOf<UBl
 	CompilerExtensions.FindOrAdd(BlueprintType).Emplace(Extension);
 }
 
-void FBlueprintCompilationManagerImpl::QueueForCompilation(const FBPCompileRequestInternal& CompileJob)
+void FBlueprintCompilationManagerImpl::QueueForCompilation(const FBPCompileRequest& CompileJob)
 {
 #if WITH_EDITOR
 	FScopeLock ScopeLock(&Lock);
 #endif
-	if(!CompileJob.UserData.BPToCompile->bQueuedForCompilation)
+	if (!CompileJob.BPToCompile->bQueuedForCompilation)
 	{
-		if(GCompilingBlueprint)
+		if (GCompilingBlueprint)
 		{
 			FString CurrentlyCompiling;
 			for (const TPair<TObjectPtr<UClass>, TObjectPtr<UClass>>& CompilerData : ClassesToReinstance)
@@ -255,42 +245,42 @@ void FBlueprintCompilationManagerImpl::QueueForCompilation(const FBPCompileReque
 			}
 			ensureMsgf(false, 
 				TEXT("Attempting to enqueue %s for compile while compiling: %s"), 
-				*CompileJob.UserData.BPToCompile->GetName(),
+				*CompileJob.BPToCompile->GetName(),
 				*CurrentlyCompiling);
 		}
 
-		CompileJob.UserData.BPToCompile->bQueuedForCompilation = true;
+		CompileJob.BPToCompile->bQueuedForCompilation = true;
 		QueuedRequests.Add(CompileJob);
 	}
 }
 
-void FBlueprintCompilationManagerImpl::CompileSynchronouslyImpl(const FBPCompileRequestInternal& Request)
+void FBlueprintCompilationManagerImpl::CompileSynchronouslyImpl(const FBPCompileRequest& Request)
 {
 	TRACE_CPUPROFILER_EVENT_SCOPE(CompileSynchronouslyImpl);
 
 #if WITH_EDITOR
 	FScopeLock ScopeLock(&Lock);
 #endif
-	Request.UserData.BPToCompile->bQueuedForCompilation = true;
+	Request.BPToCompile->bQueuedForCompilation = true;
 
-	const bool bIsRegeneratingOnLoad		= (Request.UserData.CompileOptions & EBlueprintCompileOptions::IsRegeneratingOnLoad				) != EBlueprintCompileOptions::None;
-	const bool bRegenerateSkeletonOnly		= (Request.UserData.CompileOptions & EBlueprintCompileOptions::RegenerateSkeletonOnly			) != EBlueprintCompileOptions::None;
-	const bool bSkipGarbageCollection		= (Request.UserData.CompileOptions & EBlueprintCompileOptions::SkipGarbageCollection			) != EBlueprintCompileOptions::None
+	const bool bIsRegeneratingOnLoad		= (Request.CompileOptions & EBlueprintCompileOptions::IsRegeneratingOnLoad				) != EBlueprintCompileOptions::None;
+	const bool bRegenerateSkeletonOnly		= (Request.CompileOptions & EBlueprintCompileOptions::RegenerateSkeletonOnly			) != EBlueprintCompileOptions::None;
+	const bool bSkipGarbageCollection		= (Request.CompileOptions & EBlueprintCompileOptions::SkipGarbageCollection				) != EBlueprintCompileOptions::None
 		|| bRegenerateSkeletonOnly;
-	const bool bBatchCompile				= (Request.UserData.CompileOptions & EBlueprintCompileOptions::BatchCompile						) != EBlueprintCompileOptions::None;
-	const bool bSkipReinstancing			= (Request.UserData.CompileOptions & EBlueprintCompileOptions::SkipReinstancing					) != EBlueprintCompileOptions::None;
-	const bool bSkipSaving					= (Request.UserData.CompileOptions & EBlueprintCompileOptions::SkipSave							) != EBlueprintCompileOptions::None;
-	const bool bFindAndReplaceCDOReferences	= (Request.UserData.CompileOptions & EBlueprintCompileOptions::IncludeCDOInReferenceReplacement	) != EBlueprintCompileOptions::None;
+	const bool bBatchCompile				= (Request.CompileOptions & EBlueprintCompileOptions::BatchCompile						) != EBlueprintCompileOptions::None;
+	const bool bSkipReinstancing			= (Request.CompileOptions & EBlueprintCompileOptions::SkipReinstancing					) != EBlueprintCompileOptions::None;
+	const bool bSkipSaving					= (Request.CompileOptions & EBlueprintCompileOptions::SkipSave							) != EBlueprintCompileOptions::None;
+	const bool bFindAndReplaceCDOReferences	= (Request.CompileOptions & EBlueprintCompileOptions::IncludeCDOInReferenceReplacement	) != EBlueprintCompileOptions::None;
 
 	ensure(!bIsRegeneratingOnLoad); // unexpected code path, compile on load handled with different function call
 	ensure(!bSkipReinstancing); // This is an internal option, should not go through CompileSynchronouslyImpl
 	ensure(QueuedRequests.Num() == 0);
 
 	// Wipe the PreCompile log, any generated messages are now irrelevant
-	Request.UserData.BPToCompile->PreCompileLog.Reset();
+	Request.BPToCompile->PreCompileLog.Reset();
 
 	// Reset the flag, so if the user tries to use PIE it will warn them if the BP did not compile
-	Request.UserData.BPToCompile->bDisplayCompilePIEWarning = true;
+	Request.BPToCompile->bDisplayCompilePIEWarning = true;
 	
 	// Do not want to run this code without the editor present nor when running commandlets.
 	// We do not want to regenerate a search Guid during loads, nothing has changed in the Blueprint
@@ -298,7 +288,7 @@ void FBlueprintCompilationManagerImpl::CompileSynchronouslyImpl(const FBPCompile
 	// We would like to regenerated it when a skeleton changes, but it is too expensive:
 	if (GEditor && GIsEditor && !bIsRegeneratingOnLoad && !bRegenerateSkeletonOnly)
 	{
-		FFindInBlueprintSearchManager::Get().AddOrUpdateBlueprintSearchMetadata(Request.UserData.BPToCompile);
+		FFindInBlueprintSearchManager::Get().AddOrUpdateBlueprintSearchMetadata(Request.BPToCompile);
 	}
 
 	QueuedRequests.Add(Request);
@@ -312,10 +302,10 @@ void FBlueprintCompilationManagerImpl::CompileSynchronouslyImpl(const FBPCompile
 	FlushCompilationQueueImpl(bSuppressBroadcastCompiled, &CompiledBlueprints, &SkeletonCompiledBlueprints, nullptr, bFindAndReplaceCDOReferences ? &OldToNewTemplates : nullptr);
 	FlushReinstancingQueueImpl(bFindAndReplaceCDOReferences, bFindAndReplaceCDOReferences ? &OldToNewTemplates : nullptr);
 	
-	if (FBlueprintEditorUtils::IsLevelScriptBlueprint(Request.UserData.BPToCompile) && !bRegenerateSkeletonOnly)
+	if (FBlueprintEditorUtils::IsLevelScriptBlueprint(Request.BPToCompile) && !bRegenerateSkeletonOnly)
 	{
 		// When the Blueprint is recompiled, then update the bound events for level scripting
-		ULevelScriptBlueprint* LevelScriptBP = CastChecked<ULevelScriptBlueprint>(Request.UserData.BPToCompile);
+		ULevelScriptBlueprint* LevelScriptBP = CastChecked<ULevelScriptBlueprint>(Request.BPToCompile);
 
 		// ULevel::OnLevelScriptBlueprintChanged needs to be run after the CDO has
 		// been updated as it respawns the actor:
@@ -329,7 +319,7 @@ void FBlueprintCompilationManagerImpl::CompileSynchronouslyImpl(const FBPCompile
 		}
 	}
 
-	if ( GEditor && !bRegenerateSkeletonOnly)
+	if (GEditor && !bRegenerateSkeletonOnly)
 	{
 		TRACE_CPUPROFILER_EVENT_SCOPE(BroadcastBlueprintReinstanced)
 
@@ -339,9 +329,9 @@ void FBlueprintCompilationManagerImpl::CompileSynchronouslyImpl(const FBPCompile
 		GEditor->BroadcastBlueprintReinstanced();
 	}
 	
-	ensure(Request.UserData.BPToCompile->bQueuedForCompilation == false);
+	ensure(Request.BPToCompile->bQueuedForCompilation == false);
 
-	if(!bSkipGarbageCollection)
+	if (!bSkipGarbageCollection)
 	{
 		TGuardValue<bool> GuardTemplateNameFlag(GIsGCingAfterBlueprintCompile, true);
 		CollectGarbage(GARBAGE_COLLECTION_KEEPFLAGS);
@@ -365,12 +355,12 @@ void FBlueprintCompilationManagerImpl::CompileSynchronouslyImpl(const FBPCompile
 	{
 		TRACE_CPUPROFILER_EVENT_SCOPE(BroadCastCompiled);
 
-		for(UBlueprint* BP : SkeletonCompiledBlueprints)
+		for (UBlueprint* BP : SkeletonCompiledBlueprints)
 		{
 			BP->BroadcastCompiled();
 		}
 
-		if(GEditor)
+		if (GEditor)
 		{
 			GEditor->BroadcastBlueprintCompiled();	
 		}
@@ -615,16 +605,16 @@ void FBlueprintCompilationManagerImpl::FlushCompilationQueueImpl(bool bSuppressB
 		TArray<UBlueprint*> BlueprintsToRecompile;
 
 		// First add any dependents of macro libraries that are being compiled:
-		for(const FBPCompileRequestInternal& CompileJob : QueuedRequests)
+		for (const FBPCompileRequest& CompileJob : QueuedRequests)
 		{
-			if ((CompileJob.UserData.CompileOptions & 
+			if ((CompileJob.CompileOptions & 
 				(	EBlueprintCompileOptions::RegenerateSkeletonOnly)
 				) != EBlueprintCompileOptions::None)
 			{
 				continue;
 			}
 
-			UBlueprint* BP = CompileJob.UserData.BPToCompile;
+			UBlueprint* BP = CompileJob.BPToCompile;
 
 			if(!BP->bHasBeenRegenerated && BP->GetLinker())
 			{
@@ -636,7 +626,7 @@ void FBlueprintCompilationManagerImpl::FlushCompilationQueueImpl(bool bSuppressB
 
 			FBlueprintEditorUtils::EnsureCachedDependenciesUpToDate(BP);
 
-			if ((CompileJob.UserData.CompileOptions & 
+			if ((CompileJob.CompileOptions & 
 				(	EBlueprintCompileOptions::IsRegeneratingOnLoad)
 				) != EBlueprintCompileOptions::None)
 			{
@@ -673,9 +663,9 @@ void FBlueprintCompilationManagerImpl::FlushCompilationQueueImpl(bool bSuppressB
 		SlowTask.EnterProgressFrame();
 
 		// then make sure any normal blueprints have their bytecode dependents recompiled, this is in case a function signature changes:
-		for(const FBPCompileRequestInternal& CompileJob : QueuedRequests)
+		for (const FBPCompileRequest& CompileJob : QueuedRequests)
 		{
-			if ((CompileJob.UserData.CompileOptions & EBlueprintCompileOptions::RegenerateSkeletonOnly) != EBlueprintCompileOptions::None)
+			if ((CompileJob.CompileOptions & EBlueprintCompileOptions::RegenerateSkeletonOnly) != EBlueprintCompileOptions::None)
 			{
 				continue;
 			}
@@ -685,8 +675,8 @@ void FBlueprintCompilationManagerImpl::FlushCompilationQueueImpl(bool bSuppressB
 			// added or removed. Below (Stage VIII) we skip further compilation for blueprints 
 			// that are being bytecode compiled, but their dependencies have not changed:
 			TArray<UBlueprint*> DependentBlueprints;
-			FBlueprintEditorUtils::GetDependentBlueprints(CompileJob.UserData.BPToCompile, DependentBlueprints);
-			for(UBlueprint* DependentBlueprint : DependentBlueprints)
+			FBlueprintEditorUtils::GetDependentBlueprints(CompileJob.BPToCompile, DependentBlueprints);
+			for (UBlueprint* DependentBlueprint : DependentBlueprints)
 			{
 				if(!IsQueuedForCompilation(DependentBlueprint))
 				{
@@ -711,10 +701,10 @@ void FBlueprintCompilationManagerImpl::FlushCompilationQueueImpl(bool bSuppressB
 		SlowTask.EnterProgressFrame();
 
 		// STAGE II: Filter out data only and interface blueprints:
-		for(int32 I = 0; I < QueuedRequests.Num(); ++I)
+		for (int32 I = 0; I < QueuedRequests.Num(); ++I)
 		{
-			FBPCompileRequestInternal& QueuedJob = QueuedRequests[I];
-			UBlueprint* QueuedBP = QueuedJob.UserData.BPToCompile;
+			FBPCompileRequest& QueuedJob = QueuedRequests[I];
+			UBlueprint* QueuedBP = QueuedJob.BPToCompile;
 
 			ensure(!QueuedBP->GeneratedClass ||
 				!QueuedBP->GeneratedClass->ClassDefaultObject ||
@@ -723,7 +713,7 @@ void FBlueprintCompilationManagerImpl::FlushCompilationQueueImpl(bool bSuppressB
 			bool bHasPendingUberGraphFrame = false;
 			UBlueprintGeneratedClass* BPGC = Cast<UBlueprintGeneratedClass>(QueuedBP->GeneratedClass);
 
-			if(BPGC)
+			if (BPGC)
 			{
 				if( BPGC->SimpleConstructionScript &&
 					BPGC->SimpleConstructionScript->GetSceneRootComponentTemplate(true) == nullptr)
@@ -737,9 +727,9 @@ void FBlueprintCompilationManagerImpl::FlushCompilationQueueImpl(bool bSuppressB
 			bool bSkipCompile = false;
 			const UClass* ParentClass = QueuedBP->ParentClass;
 			const bool bHasClassAndMatchesParent = BPGC && BPGC->GetSuperClass() == ParentClass;
-			if( bHasClassAndMatchesParent &&
-				FBlueprintEditorUtils::IsDataOnlyBlueprint(QueuedBP) && !QueuedBP->bHasBeenRegenerated && 
-				QueuedBP->GetLinker() && !bDefaultComponentMustBeAdded && !bHasPendingUberGraphFrame )
+			if ( bHasClassAndMatchesParent &&
+				 FBlueprintEditorUtils::IsDataOnlyBlueprint(QueuedBP) && !QueuedBP->bHasBeenRegenerated && 
+				 QueuedBP->GetLinker() && !bDefaultComponentMustBeAdded && !bHasPendingUberGraphFrame )
 			{
 				// consider skipping the compile operation for this DOB:
 				if (ParentClass && ParentClass->HasAllClassFlags(CLASS_Native))
@@ -755,14 +745,14 @@ void FBlueprintCompilationManagerImpl::FlushCompilationQueueImpl(bool bSuppressB
 				}
 			}
 
-			if(bSkipCompile)
+			if (bSkipCompile)
 			{
 				CurrentlyCompilingBPs.Emplace(
 					FCompilerData(
 						QueuedBP, 
 						ECompilationManagerJobType::SkeletonOnly, 
-						QueuedJob.UserData.ClientResultsLog, 
-						QueuedJob.UserData.CompileOptions,
+						QueuedJob.ClientResultsLog, 
+						QueuedJob.CompileOptions,
 						false
 					)
 				);
@@ -788,7 +778,7 @@ void FBlueprintCompilationManagerImpl::FlushCompilationQueueImpl(bool bSuppressB
 			else
 			{
 				ECompilationManagerJobType JobType = ECompilationManagerJobType::Normal;
-				if ((QueuedJob.UserData.CompileOptions & EBlueprintCompileOptions::RegenerateSkeletonOnly) != EBlueprintCompileOptions::None)
+				if ((QueuedJob.CompileOptions & EBlueprintCompileOptions::RegenerateSkeletonOnly) != EBlueprintCompileOptions::None)
 				{
 					JobType = ECompilationManagerJobType::SkeletonOnly;
 				}
@@ -797,8 +787,8 @@ void FBlueprintCompilationManagerImpl::FlushCompilationQueueImpl(bool bSuppressB
 					FCompilerData(
 						QueuedBP, 
 						JobType, 
-						QueuedJob.UserData.ClientResultsLog, 
-						QueuedJob.UserData.CompileOptions, 
+						QueuedJob.ClientResultsLog, 
+						QueuedJob.CompileOptions, 
 						false
 					)
 				);
@@ -809,21 +799,21 @@ void FBlueprintCompilationManagerImpl::FlushCompilationQueueImpl(bool bSuppressB
 
 		SlowTask.EnterProgressFrame();
 
-		for(UBlueprint* BP : BlueprintsToRecompile)
+		for (UBlueprint* BP : BlueprintsToRecompile)
 		{
 			// make sure all children are at least re-linked:
-			if(UClass* OldSkeletonClass = BP->SkeletonGeneratedClass)
+			if (UClass* OldSkeletonClass = BP->SkeletonGeneratedClass)
 			{
 				TArray<UClass*> SkeletonClassesToReparentList;
 				// Has to be recursive gather of children because instances of a UClass will cache information about
 				// classes that are above their immediate parent (e.g. ClassConstructor):
 				GetDerivedClasses(OldSkeletonClass, SkeletonClassesToReparentList);
 		
-				for(UClass* ChildClass : SkeletonClassesToReparentList)
+				for (UClass* ChildClass : SkeletonClassesToReparentList)
 				{
-					if(UBlueprint* ChildBlueprint = UBlueprint::GetBlueprintFromClass(ChildClass))
+					if (UBlueprint* ChildBlueprint = UBlueprint::GetBlueprintFromClass(ChildClass))
 					{
-						if(!IsQueuedForCompilation(ChildBlueprint))
+						if (!IsQueuedForCompilation(ChildBlueprint))
 						{
 							ChildBlueprint->bQueuedForCompilation = true;
 							ensure(ChildBlueprint->bHasBeenRegenerated);
@@ -2050,9 +2040,9 @@ void FBlueprintCompilationManagerImpl::VerifyNoQueuedRequests(const TArray<FComp
 	if (QueuedRequests.Num() != 0)
 	{
 		FString QueuedBlueprints = TEXT("");
-		for (const FBPCompileRequestInternal& Request: QueuedRequests)
+		for (const FBPCompileRequest& Request: QueuedRequests)
 		{
-			QueuedBlueprints += Request.UserData.BPToCompile->GetName() + TEXT(" ");
+			QueuedBlueprints += Request.BPToCompile->GetName() + TEXT(" ");
 		}
 		FString CompilingBlueprints = TEXT("");
 		for (const FCompilerData& CompilerData : CurrentlyCompilingBPs)
@@ -2072,15 +2062,15 @@ void FBlueprintCompilationManagerImpl::ReparentHierarchies(const TMap<UClass*, U
 	TRACE_CPUPROFILER_EVENT_SCOPE(ReparentHierarchies);
 
 	// something has decided to replace instances of a class. We need to update all the children of those types:
-	TArray< UClass* > ClassesOrdered;
+	TArray<UClass*> ClassesOrdered;
 	// Map used to distinguish between new classes and classes that need to be reinstanced (reparented) via a new reinstancer:
 	TMap<UClass*, UClass*> NewToOldClasses;
 	{
 		TSet<UClass*> Classes;
-		for(const TPair<UClass*, UClass*>& OldToNewClass : OldToNewClasses)
+		for (const TPair<UClass*, UClass*>& OldToNewClass : OldToNewClasses)
 		{
 			// classes with no CDO do not need to be reinstanced:
-			if(OldToNewClass.Key->ClassDefaultObject == nullptr)
+			if (OldToNewClass.Key->ClassDefaultObject == nullptr)
 			{
 				continue;
 			}
@@ -2093,7 +2083,7 @@ void FBlueprintCompilationManagerImpl::ReparentHierarchies(const TMap<UClass*, U
 			// update things like the ClassConstructor pointer in case it changed:
 			GetDerivedClasses(OldToNewClass.Key, DerivedClasses);
 
-			for(UClass* DerivedClass : DerivedClasses)
+			for (UClass* DerivedClass : DerivedClasses)
 			{
 				if (DerivedClass->ClassDefaultObject == nullptr && 
 					DerivedClass->GetSparseClassData(EGetSparseClassDataMethod::ReturnIfNull) == nullptr)
@@ -2125,10 +2115,10 @@ void FBlueprintCompilationManagerImpl::ReparentHierarchies(const TMap<UClass*, U
 
 	// create reinstancing jobs, no need to create a reinstancer when there is a new UClass* available (e.g. asset reload, hot reload):
 	TArray<FReinstancingJob> Reinstancers;
-	for(UClass* Class : ClassesOrdered)
+	for (UClass* Class : ClassesOrdered)
 	{
 		UClass* const* OldClass = NewToOldClasses.Find(Class);
-		if(OldClass)
+		if (OldClass)
 		{
 			Reinstancers.Push( FReinstancingJob( TPair<UClass*, UClass*>(*OldClass, Class) ) );
 		}
@@ -2161,7 +2151,7 @@ void FBlueprintCompilationManagerImpl::ReparentHierarchies(const TMap<UClass*, U
 	// Reparent and Link - this is .. kind of pointless.. ReinstanceBatch should
 	// be doing this
 	TMap<UClass*, UClass*> OldClassToNewClassIncludingChildren = OldToNewClasses;
-	for(const FReinstancingJob& ReinstancingJob : Reinstancers)
+	for (const FReinstancingJob& ReinstancingJob : Reinstancers)
 	{
 		UClass* ClassToReinstance = ReinstancingJob.OldToNew.Value;
 
@@ -2199,7 +2189,7 @@ void FBlueprintCompilationManagerImpl::ReparentHierarchies(const TMap<UClass*, U
 		OldClassToNewClassDerivedTypes = OldClassToNewClassIncludingChildren;
 	}
 	
-	for(const FReinstancingJob& ReinstancingJob : Reinstancers)
+	for (const FReinstancingJob& ReinstancingJob : Reinstancers)
 	{
 		OldClassToNewClassDerivedTypes.Add(ReinstancingJob.OldToNew);
 	}
@@ -2210,16 +2200,16 @@ void FBlueprintCompilationManagerImpl::ReparentHierarchies(const TMap<UClass*, U
 
 	// Make sure we don't replace old instances that are in the *callers* old to new TMap!
 	TSet<UObject*> OldObjects;
-	for(TPair<UClass*, UClass*> OldToNew : OldClassToNewClassDerivedTypes)
+	for (TPair<UClass*, UClass*> OldToNew : OldClassToNewClassDerivedTypes)
 	{
 		ensure(OldToNew.Value->HasAnyClassFlags(CLASS_TokenStreamAssembled));
 
 		TArray< UObject* > OldObjectsOfType;
 		GetObjectsOfClass(OldToNew.Key, OldObjectsOfType);
 
-		for(UObject* Obj : OldObjectsOfType)
+		for (UObject* Obj : OldObjectsOfType)
 		{
-			if(Obj->HasAnyFlags(RF_NewerVersionExists))
+			if (Obj->HasAnyFlags(RF_NewerVersionExists))
 			{
 				OldObjects.Add(Obj);
 			}
