Remove pointless typedef related to blueprint context menu building
#rnx

diff --git a/Engine/Source/Editor/Kismet/Private/BlueprintActionMenuBuilder.cpp b/Engine/Source/Editor/Kismet/Private/BlueprintActionMenuBuilder.cpp
--- a/Engine/Source/Editor/Kismet/Private/BlueprintActionMenuBuilder.cpp
+++ b/Engine/Source/Editor/Kismet/Private/BlueprintActionMenuBuilder.cpp
@@ -249,8 +249,6 @@ FBlueprintActionUiSpec FBlueprintActionMenuItemFactory::GetActionUiSignature(FBl
 
 namespace FBlueprintActionMenuBuilderImpl
 {
-	typedef TArray< TSharedPtr<FEdGraphSchemaAction> > MenuItemList;
-
 	/** Defines a sub-section of the overall blueprint menu (filter, heading, etc.) */
 	struct FMenuSectionDefinition
 	{
@@ -280,7 +278,7 @@ namespace FBlueprintActionMenuBuilderImpl
 		 * @param  DatabaseAction	The node-spawner that the new menu item should wrap.
 		 * @return An empty TSharedPtr if the action was filtered out, otherwise a newly allocated FBlueprintActionMenuItem.
 		 */
-		MenuItemList MakeMenuItems(FBlueprintActionInfo& DatabaseAction);
+		TArray< TSharedPtr<FEdGraphSchemaAction> > MakeMenuItems(FBlueprintActionInfo& DatabaseAction);
 
 		/**
 		 * 
@@ -289,7 +287,7 @@ namespace FBlueprintActionMenuBuilderImpl
 		 * @param  Bindings	
 		 * @return 
 		 */
-		void AddBoundMenuItems(FBlueprintActionInfo& DatabaseAction, TArray<FFieldVariant> const& Bindings, MenuItemList& MenuItemsOut);
+		void AddBoundMenuItems(FBlueprintActionInfo& DatabaseAction, TArray<FFieldVariant> const& Bindings, TArray< TSharedPtr<FEdGraphSchemaAction> >& MenuItemsOut);
 		
 		/**
 		 * Clears out any consolidated properties that this may have been 
@@ -394,7 +392,7 @@ int32 FBlueprintActionMenuBuilderImpl::FMenuSectionDefinition::GetSectionSortOrd
 }
 // 
 //------------------------------------------------------------------------------
-void FBlueprintActionMenuBuilderImpl::FMenuSectionDefinition::AddBoundMenuItems(FBlueprintActionInfo& DatabaseActionInfo, TArray<FFieldVariant> const& PerspectiveBindings, MenuItemList& MenuItemsOut)
+void FBlueprintActionMenuBuilderImpl::FMenuSectionDefinition::AddBoundMenuItems(FBlueprintActionInfo& DatabaseActionInfo, TArray<FFieldVariant> const& PerspectiveBindings, TArray< TSharedPtr<FEdGraphSchemaAction> >& MenuItemsOut)
 {
 	TRACE_CPUPROFILER_EVENT_SCOPE(FMenuSectionDefinition::AddBoundMenuItems);
 
@@ -459,7 +457,7 @@ void FBlueprintActionMenuBuilderImpl::FMenuSectionDefinition::AddBoundMenuItems(
 }
 
 //------------------------------------------------------------------------------
-FBlueprintActionMenuBuilderImpl::MenuItemList FBlueprintActionMenuBuilderImpl::FMenuSectionDefinition::MakeMenuItems(FBlueprintActionInfo& DatabaseAction)
+TArray< TSharedPtr<FEdGraphSchemaAction> > FBlueprintActionMenuBuilderImpl::FMenuSectionDefinition::MakeMenuItems(FBlueprintActionInfo& DatabaseAction)
 {
 	TRACE_CPUPROFILER_EVENT_SCOPE(FMenuSectionDefinition::MakeMenuItems);
 
@@ -510,7 +508,7 @@ FBlueprintActionMenuBuilderImpl::MenuItemList FBlueprintActionMenuBuilderImpl::F
 		}
 	}
 
-	FBlueprintActionMenuBuilderImpl::MenuItemList MenuItems;
+	TArray< TSharedPtr<FEdGraphSchemaAction> > MenuItems;
 	if (UnBoundMenuEntry.IsValid())
 	{
 		MenuItems.Add(UnBoundMenuEntry);
