StateTree: Added support for description and icon for BP nodes.

diff --git a/Engine/Plugins/Runtime/StateTree/Source/StateTreeEditorModule/Private/K2Node_StateTreeNodeGetPropertyDescription.cpp b/Engine/Plugins/Runtime/StateTree/Source/StateTreeEditorModule/Private/K2Node_StateTreeNodeGetPropertyDescription.cpp
new file mode 100644
--- /dev/null
+++ b/Engine/Plugins/Runtime/StateTree/Source/StateTreeEditorModule/Private/K2Node_StateTreeNodeGetPropertyDescription.cpp
@@ -0,0 +1,93 @@
+﻿// Copyright Epic Games, Inc. All Rights Reserved.
+
+#include "K2Node_StateTreeNodeGetPropertyDescription.h"
+#include "Blueprint/StateTreeNodeBlueprintBase.h"
+#include "EdGraphSchema_K2.h"
+#include "BlueprintActionDatabaseRegistrar.h"
+#include "BlueprintNodeSpawner.h"
+#include "K2Node_CallFunction.h"
+#include "KismetCompiler.h"
+
+#define LOCTEXT_NAMESPACE "K2Node_StateTreeNodeGetPropertyDescription"
+
+
+void UK2Node_StateTreeNodeGetPropertyDescription::AllocateDefaultPins()
+{
+	CreatePin(EGPD_Output, UEdGraphSchema_K2::PC_Text, UEdGraphSchema_K2::PN_ReturnValue);
+	
+	Super::AllocateDefaultPins();
+}
+
+FText UK2Node_StateTreeNodeGetPropertyDescription::GetNodeTitle(ENodeTitleType::Type TitleType) const
+{
+	UBlueprint* Blueprint = HasValidBlueprint() ? GetBlueprint() : nullptr;
+	const FProperty* Property = const_cast<FMemberReference&>(Variable).ResolveMember<FProperty>(Blueprint);
+	const FText SelectedPropertyName = Property ? Property->GetDisplayNameText() : LOCTEXT("None", "<None>");
+
+	return FText::Format(LOCTEXT("NodeTitle", "Get Description for {0}"), SelectedPropertyName);
+}
+
+FText UK2Node_StateTreeNodeGetPropertyDescription::GetTooltipText() const
+{
+	return LOCTEXT("NodeTooltip", "Returns text describing the specified member variable.");
+}
+
+FText UK2Node_StateTreeNodeGetPropertyDescription::GetMenuCategory() const
+{
+	return LOCTEXT("NodeCategory", "StateTree");
+}
+
+void UK2Node_StateTreeNodeGetPropertyDescription::GetMenuActions(FBlueprintActionDatabaseRegistrar& ActionRegistrar) const
+{
+	Super::GetMenuActions(ActionRegistrar);
+	UClass* Action = GetClass();
+	if (ActionRegistrar.IsOpenForRegistration(Action))
+    {
+        UBlueprintNodeSpawner* Spawner = UBlueprintNodeSpawner::Create(Action);
+        ActionRegistrar.AddBlueprintAction(Action, Spawner);
+    }
+}
+
+void UK2Node_StateTreeNodeGetPropertyDescription::ValidateNodeDuringCompilation(FCompilerResultsLog& MessageLog) const
+{
+	Super::ValidateNodeDuringCompilation(MessageLog);
+
+	const UClass* BlueprintClass = GetBlueprintClassFromNode();
+	if (!BlueprintClass->IsChildOf<UStateTreeNodeBlueprintBase>())
+	{
+		const FText ErrorText = LOCTEXT("InvalidSelfType", "This blueprint (self) is not a 'State Tree Blueprint Node'.");
+		MessageLog.Error(*ErrorText.ToString(), this);
+	}
+
+	UBlueprint* Blueprint = HasValidBlueprint() ? GetBlueprint() : nullptr;
+	const FProperty* Property = const_cast<FMemberReference&>(Variable).ResolveMember<FProperty>(Blueprint);
+	if (!Property)
+	{
+		const FText ErrorText = FText::Format(LOCTEXT("InvalidProperty", "Cannot find property '{0}'."), FText::FromName(Variable.GetMemberName()));
+		MessageLog.Error(*ErrorText.ToString(), this);
+	}
+}
+
+void UK2Node_StateTreeNodeGetPropertyDescription::ExpandNode(FKismetCompilerContext& CompilerContext, UEdGraph* SourceGraph)
+{
+	UBlueprint* Blueprint = HasValidBlueprint() ? GetBlueprint() : nullptr;
+
+	// Property name
+	const FProperty* Property = Variable.ResolveMember<FProperty>(Blueprint);
+	FString SelectedPropertyName = Property ? Property->GetName() : TEXT("");
+	
+	UK2Node_CallFunction* CallGetPropertyDescription = CompilerContext.SpawnIntermediateNode<UK2Node_CallFunction>(this, SourceGraph);
+	CallGetPropertyDescription->SetFromFunction(UStateTreeNodeBlueprintBase::StaticClass()->FindFunctionByName(GET_FUNCTION_NAME_CHECKED(UStateTreeNodeBlueprintBase, GetPropertyDescriptionByPropertyName)));
+	CallGetPropertyDescription->AllocateDefaultPins();
+
+	UEdGraphPin* PropertyNamePin = CallGetPropertyDescription->FindPinChecked(TEXT("PropertyName"));
+	check(PropertyNamePin);
+	PropertyNamePin->DefaultValue = SelectedPropertyName;
+
+	UEdGraphPin* OrgReturnPin = FindPinChecked(UEdGraphSchema_K2::PN_ReturnValue);
+	UEdGraphPin* NewReturnPin = CallGetPropertyDescription->GetReturnValuePin();
+	check(NewReturnPin);
+	CompilerContext.MovePinLinksToIntermediate(*OrgReturnPin, *NewReturnPin);
+}
+
+#undef LOCTEXT_NAMESPACE
diff --git a/Engine/Plugins/Runtime/StateTree/Source/StateTreeEditorModule/Private/K2Node_StateTreeNodeGetPropertyDescription.h b/Engine/Plugins/Runtime/StateTree/Source/StateTreeEditorModule/Private/K2Node_StateTreeNodeGetPropertyDescription.h
new file mode 100644
--- /dev/null
+++ b/Engine/Plugins/Runtime/StateTree/Source/StateTreeEditorModule/Private/K2Node_StateTreeNodeGetPropertyDescription.h
@@ -0,0 +1,39 @@
+﻿// Copyright Epic Games, Inc. All Rights Reserved.
+
+#pragma once
+#include "K2Node.h"
+#include "Engine/MemberReference.h"
+#include "K2Node_StateTreeNodeGetPropertyDescription.generated.h"
+
+class FKismetCompilerContext;
+
+/**
+ * Returns description for a specific property in the class.
+ * If the property has a binding, the binding string will be returned.
+ * Otherwise the current value of the property is returned.
+ */
+UCLASS()
+class UK2Node_StateTreeNodeGetPropertyDescription : public UK2Node
+{
+	GENERATED_BODY()
+
+public:
+	virtual FText GetNodeTitle(ENodeTitleType::Type TitleType) const override;
+	virtual FText GetTooltipText() const override;
+	virtual FText GetMenuCategory() const override;
+
+	virtual void GetMenuActions(FBlueprintActionDatabaseRegistrar& ActionRegistrar) const override;
+	virtual bool IsNodePure() const override { return true; }
+	virtual bool ShouldShowNodeProperties() const override { return true; }
+	virtual void ExpandNode(FKismetCompilerContext& CompilerContext, UEdGraph* SourceGraph) override;
+	virtual void ValidateNodeDuringCompilation(FCompilerResultsLog& MessageLog) const override;
+
+protected:
+	virtual void AllocateDefaultPins() override;
+
+	/** Property of the class to describe. */
+	UPROPERTY(EditAnywhere, Category = "Node", meta=(PropertyReference, AllowOnlyVisibleProperties))
+	FMemberReference Variable;
+
+	friend class FKCHandler_StateTreeNodeGetPropertyDescription;
+};
diff --git a/Engine/Plugins/Runtime/StateTree/Source/StateTreeModule/Private/Blueprint/StateTreeConditionBlueprintBase.cpp b/Engine/Plugins/Runtime/StateTree/Source/StateTreeModule/Private/Blueprint/StateTreeConditionBlueprintBase.cpp
--- a/Engine/Plugins/Runtime/StateTree/Source/StateTreeModule/Private/Blueprint/StateTreeConditionBlueprintBase.cpp
+++ b/Engine/Plugins/Runtime/StateTree/Source/StateTreeModule/Private/Blueprint/StateTreeConditionBlueprintBase.cpp
@@ -47,10 +47,33 @@ bool FStateTreeBlueprintConditionWrapper::TestCondition(FStateTreeExecutionConte
 #if WITH_EDITOR
 FText FStateTreeBlueprintConditionWrapper::GetDescription(const FGuid& ID, FStateTreeDataView InstanceDataView, const IStateTreeBindingLookup& BindingLookup, EStateTreeNodeFormatting Formatting) const
 {
-	if (ConditionClass)
+	FText Description;
+	if (const UStateTreeConditionBlueprintBase* Instance = InstanceDataView.GetPtr<UStateTreeConditionBlueprintBase>())
 	{
-		return ConditionClass->GetDisplayNameText();
+		Description = Instance->GetDescription(ID, InstanceDataView, BindingLookup, Formatting);
 	}
-	return FText::GetEmpty();
+	if (Description.IsEmpty() && ConditionClass)
+	{
+		Description = ConditionClass->GetDisplayNameText();
+	}
+	return Description;
+}
+
+FName FStateTreeBlueprintConditionWrapper::GetIconName() const
+{
+	if (const UStateTreeNodeBlueprintBase* NodeCDO = GetDefault<const UStateTreeNodeBlueprintBase>(ConditionClass))
+	{
+		return NodeCDO->GetIconName();
+	}
+	return FStateTreeConditionBase::GetIconName();
+}
+
+FColor FStateTreeBlueprintConditionWrapper::GetIconColor() const
+{
+	if (const UStateTreeNodeBlueprintBase* NodeCDO = GetDefault<const UStateTreeNodeBlueprintBase>(ConditionClass))
+	{
+		return NodeCDO->GetIconColor();
+	}
+	return FStateTreeConditionBase::GetIconColor();
 }
 #endif
\ No newline at end of file
diff --git a/Engine/Plugins/Runtime/StateTree/Source/StateTreeModule/Private/Blueprint/StateTreeEvaluatorBlueprintBase.cpp b/Engine/Plugins/Runtime/StateTree/Source/StateTreeModule/Private/Blueprint/StateTreeEvaluatorBlueprintBase.cpp
--- a/Engine/Plugins/Runtime/StateTree/Source/StateTreeModule/Private/Blueprint/StateTreeEvaluatorBlueprintBase.cpp
+++ b/Engine/Plugins/Runtime/StateTree/Source/StateTreeModule/Private/Blueprint/StateTreeEvaluatorBlueprintBase.cpp
@@ -77,10 +77,33 @@ void FStateTreeBlueprintEvaluatorWrapper::Tick(FStateTreeExecutionContext& Conte
 #if WITH_EDITOR
 FText FStateTreeBlueprintEvaluatorWrapper::GetDescription(const FGuid& ID, FStateTreeDataView InstanceDataView, const IStateTreeBindingLookup& BindingLookup, EStateTreeNodeFormatting Formatting) const
 {
-	if (EvaluatorClass)
+	FText Description;
+	if (const UStateTreeEvaluatorBlueprintBase* Instance = InstanceDataView.GetPtr<UStateTreeEvaluatorBlueprintBase>())
 	{
-		return EvaluatorClass->GetDisplayNameText();
+		Description = Instance->GetDescription(ID, InstanceDataView, BindingLookup, Formatting);
 	}
-	return FText::GetEmpty();
+	if (Description.IsEmpty() && EvaluatorClass)
+	{
+		Description = EvaluatorClass->GetDisplayNameText();
+	}
+	return Description;
+}
+
+FName FStateTreeBlueprintEvaluatorWrapper::GetIconName() const
+{
+	if (const UStateTreeNodeBlueprintBase* NodeCDO = GetDefault<const UStateTreeNodeBlueprintBase>(EvaluatorClass))
+	{
+		return NodeCDO->GetIconName();
+	}
+	return FStateTreeEvaluatorBase::GetIconName();
+}
+
+FColor FStateTreeBlueprintEvaluatorWrapper::GetIconColor() const
+{
+	if (const UStateTreeNodeBlueprintBase* NodeCDO = GetDefault<const UStateTreeNodeBlueprintBase>(EvaluatorClass))
+	{
+		return NodeCDO->GetIconColor();
+	}
+	return FStateTreeEvaluatorBase::GetIconColor();
 }
 #endif
\ No newline at end of file
diff --git a/Engine/Plugins/Runtime/StateTree/Source/StateTreeModule/Private/Blueprint/StateTreeNodeBlueprintBase.cpp b/Engine/Plugins/Runtime/StateTree/Source/StateTreeModule/Private/Blueprint/StateTreeNodeBlueprintBase.cpp
--- a/Engine/Plugins/Runtime/StateTree/Source/StateTreeModule/Private/Blueprint/StateTreeNodeBlueprintBase.cpp
+++ b/Engine/Plugins/Runtime/StateTree/Source/StateTreeModule/Private/Blueprint/StateTreeNodeBlueprintBase.cpp
@@ -9,6 +9,13 @@
 
 #include UE_INLINE_GENERATED_CPP_BY_NAME(StateTreeNodeBlueprintBase)
 
+#define LOCTEXT_NAMESPACE "StateTree"
+
+#if WITH_EDITOR
+FGuid UStateTreeNodeBlueprintBase::CachedNodeID;
+const IStateTreeBindingLookup* UStateTreeNodeBlueprintBase::CachedBindingLookup = nullptr;
+#endif
+
 UWorld* UStateTreeNodeBlueprintBase::GetWorld() const
 {
 	// The items are duplicated as the State Tree execution context as outer, so this should be essentially the same as GetWorld() on StateTree context.
@@ -149,3 +156,58 @@ DEFINE_FUNCTION(UStateTreeNodeBlueprintBase::execGetPropertyReference)
 		Stack.MostRecentProperty = nullptr;
 	}
 }
+
+#if WITH_EDITOR
+FText UStateTreeNodeBlueprintBase::GetDescription(const FGuid& ID, FStateTreeDataView InstanceDataView, const IStateTreeBindingLookup& BindingLookup, EStateTreeNodeFormatting Formatting) const
+{
+	FText Result;
+	
+	const FGuid OldCachedNodeID = CachedNodeID;
+	const IStateTreeBindingLookup* OldCachedBindingLookup = CachedBindingLookup;
+
+	CachedNodeID = ID;
+	CachedBindingLookup = &BindingLookup;
+
+	Result = Description;
+	if (Result.IsEmpty())
+	{
+		Result = ReceiveGetDescription(Formatting);
+	}
+	if (Result.IsEmpty())
+	{
+		Result = GetClass()->GetDisplayNameText();
+	}
+
+	CachedNodeID = OldCachedNodeID;
+	CachedBindingLookup = OldCachedBindingLookup;
+	
+	return Result;
+}
+#endif
+
+FText UStateTreeNodeBlueprintBase::GetPropertyDescriptionByPropertyName(FName PropertyName) const
+{
+	FText Result;
+#if WITH_EDITOR
+	// Try property binding first
+	if (CachedBindingLookup)
+	{
+		const FStateTreePropertyPath Path(CachedNodeID, PropertyName);
+		Result = CachedBindingLookup->GetBindingSourceDisplayName(Path);
+	}
+
+	// No binding, get the value.
+	if (Result.IsEmpty())
+	{
+		if (const FProperty* Property = GetClass()->FindPropertyByName(PropertyName))
+		{
+			FString	Value;
+			Property->ExportText_InContainer(0, Value, this, this, nullptr, PPF_PropertyWindow | PPF_BlueprintDebugView);
+			Result = FText::FromString(Value);
+		}
+	}
+#endif
+	return Result;
+}
+
+#undef LOCTEXT_NAMESPACE
diff --git a/Engine/Plugins/Runtime/StateTree/Source/StateTreeModule/Private/Blueprint/StateTreeTaskBlueprintBase.cpp b/Engine/Plugins/Runtime/StateTree/Source/StateTreeModule/Private/Blueprint/StateTreeTaskBlueprintBase.cpp
--- a/Engine/Plugins/Runtime/StateTree/Source/StateTreeModule/Private/Blueprint/StateTreeTaskBlueprintBase.cpp
+++ b/Engine/Plugins/Runtime/StateTree/Source/StateTreeModule/Private/Blueprint/StateTreeTaskBlueprintBase.cpp
@@ -149,10 +149,33 @@ EStateTreeRunStatus FStateTreeBlueprintTaskWrapper::Tick(FStateTreeExecutionCont
 #if WITH_EDITOR
 FText FStateTreeBlueprintTaskWrapper::GetDescription(const FGuid& ID, FStateTreeDataView InstanceDataView, const IStateTreeBindingLookup& BindingLookup, EStateTreeNodeFormatting Formatting) const
 {
-	if (TaskClass)
+	FText Description;
+	if (const UStateTreeTaskBlueprintBase* Instance = InstanceDataView.GetPtr<UStateTreeTaskBlueprintBase>())
 	{
-		return TaskClass->GetDisplayNameText();
+		Description = Instance->GetDescription(ID, InstanceDataView, BindingLookup, Formatting);
 	}
-	return FText::GetEmpty();
+	if (Description.IsEmpty() && TaskClass)
+	{
+		Description = TaskClass->GetDisplayNameText();
+	}
+	return Description;
+}
+
+FName FStateTreeBlueprintTaskWrapper::GetIconName() const
+{
+	if (const UStateTreeNodeBlueprintBase* NodeCDO = GetDefault<const UStateTreeNodeBlueprintBase>(TaskClass))
+	{
+		return NodeCDO->GetIconName();
+	}
+	return FStateTreeTaskBase::GetIconName();
+}
+
+FColor FStateTreeBlueprintTaskWrapper::GetIconColor() const
+{
+	if (const UStateTreeNodeBlueprintBase* NodeCDO = GetDefault<const UStateTreeNodeBlueprintBase>(TaskClass))
+	{
+		return NodeCDO->GetIconColor();
+	}
+	return FStateTreeTaskBase::GetIconColor();
 }
 #endif
\ No newline at end of file
diff --git a/Engine/Plugins/Runtime/StateTree/Source/StateTreeModule/Public/Blueprint/StateTreeConditionBlueprintBase.h b/Engine/Plugins/Runtime/StateTree/Source/StateTreeModule/Public/Blueprint/StateTreeConditionBlueprintBase.h
--- a/Engine/Plugins/Runtime/StateTree/Source/StateTreeModule/Public/Blueprint/StateTreeConditionBlueprintBase.h
+++ b/Engine/Plugins/Runtime/StateTree/Source/StateTreeModule/Public/Blueprint/StateTreeConditionBlueprintBase.h
@@ -42,6 +42,8 @@ struct STATETREEMODULE_API FStateTreeBlueprintConditionWrapper : public FStateTr
 	virtual bool TestCondition(FStateTreeExecutionContext& Context) const override;
 #if WITH_EDITOR
 	virtual FText GetDescription(const FGuid& ID, FStateTreeDataView InstanceDataView, const IStateTreeBindingLookup& BindingLookup, EStateTreeNodeFormatting Formatting = EStateTreeNodeFormatting::Text) const override;
+	virtual FName GetIconName() const override;
+	virtual FColor GetIconColor() const override;
 #endif
 	
 	UPROPERTY()
diff --git a/Engine/Plugins/Runtime/StateTree/Source/StateTreeModule/Public/Blueprint/StateTreeEvaluatorBlueprintBase.h b/Engine/Plugins/Runtime/StateTree/Source/StateTreeModule/Public/Blueprint/StateTreeEvaluatorBlueprintBase.h
--- a/Engine/Plugins/Runtime/StateTree/Source/StateTreeModule/Public/Blueprint/StateTreeEvaluatorBlueprintBase.h
+++ b/Engine/Plugins/Runtime/StateTree/Source/StateTreeModule/Public/Blueprint/StateTreeEvaluatorBlueprintBase.h
@@ -55,6 +55,8 @@ struct STATETREEMODULE_API FStateTreeBlueprintEvaluatorWrapper : public FStateTr
 	virtual void Tick(FStateTreeExecutionContext& Context, const float DeltaTime) const override;
 #if WITH_EDITOR
 	virtual FText GetDescription(const FGuid& ID, FStateTreeDataView InstanceDataView, const IStateTreeBindingLookup& BindingLookup, EStateTreeNodeFormatting Formatting = EStateTreeNodeFormatting::Text) const override;
+	virtual FName GetIconName() const override;
+	virtual FColor GetIconColor() const override;
 #endif
 	
 	UPROPERTY()
diff --git a/Engine/Plugins/Runtime/StateTree/Source/StateTreeModule/Public/Blueprint/StateTreeNodeBlueprintBase.h b/Engine/Plugins/Runtime/StateTree/Source/StateTreeModule/Public/Blueprint/StateTreeNodeBlueprintBase.h
--- a/Engine/Plugins/Runtime/StateTree/Source/StateTreeModule/Public/Blueprint/StateTreeNodeBlueprintBase.h
+++ b/Engine/Plugins/Runtime/StateTree/Source/StateTreeModule/Public/Blueprint/StateTreeNodeBlueprintBase.h
@@ -3,6 +3,7 @@
 #pragma once
 
 #include "StateTreeExecutionTypes.h"
+#include "StateTreeNodeBase.h"
 #include "UObject/ObjectMacros.h"
 #include "StateTreeNodeBlueprintBase.generated.h"
 
@@ -55,7 +56,30 @@ public:
 	UFUNCTION()
 	bool IsPropertyRefValid(const FStateTreeBlueprintPropertyRef& PropertyRef) const;
 
+	/** @return text describing the property, either direct value or binding description. Used internally. */
+	UFUNCTION(BlueprintCallable, Category = "StateTree", meta=( BlueprintInternalUseOnly="true" ))
+	FText GetPropertyDescriptionByPropertyName(FName PropertyName) const;
+
+#if WITH_EDITOR
+	FText GetDescription(const FGuid& ID, FStateTreeDataView InstanceDataView, const IStateTreeBindingLookup& BindingLookup, EStateTreeNodeFormatting Formatting) const;
+
+	FName GetIconName() const
+	{
+		return IconName;
+	}
+	
+	FColor GetIconColor() const
+	{
+		return IconColor;
+	}
+#endif
+	
 protected:
+
+	/** Event to implement to get node description. */
+	UFUNCTION(BlueprintImplementableEvent, meta = (DisplayName = "Get Description"))
+	FText ReceiveGetDescription(EStateTreeNodeFormatting Formatting) const;
+
 	virtual UWorld* GetWorld() const override;
 	AActor* GetOwnerActor(const FStateTreeExecutionContext& Context) const;
 
@@ -89,6 +113,31 @@ private:
 
 	/** Cached State where the node is processed on. */
 	mutable FStateTreeStateHandle CachedState;
+
+#if WITH_EDITORONLY_DATA
+	/** Description of the node. */
+	UPROPERTY(EditDefaultsOnly, Category="Description")
+	FText Description;
+
+	/**
+	 * Name of the icon in format:
+	 *		StyleSetName | StyleName [ | SmallStyleName | StatusOverlayStyleName]
+	 *		SmallStyleName and StatusOverlayStyleName are optional.
+	 *		Example: "StateTreeEditorStyle|Node.Animation"
+	 */
+	UPROPERTY(EditDefaultsOnly, Category="Description")
+	FName IconName;
+
+	/** Color of the icon. */
+	UPROPERTY(EditDefaultsOnly, Category="Description")
+	FColor IconColor = UE::StateTree::Colors::Grey;
+#endif // 	WITH_EDITORONLY_DATA
+
+#if WITH_EDITOR
+	/** Cached values used during editor to make some BP nodes simpler to use. */
+	static FGuid CachedNodeID;
+	static const IStateTreeBindingLookup* CachedBindingLookup;
+#endif	
 };
 
 #if UE_ENABLE_INCLUDE_ORDER_DEPRECATED_IN_5_2
diff --git a/Engine/Plugins/Runtime/StateTree/Source/StateTreeModule/Public/Blueprint/StateTreeTaskBlueprintBase.h b/Engine/Plugins/Runtime/StateTree/Source/StateTreeModule/Public/Blueprint/StateTreeTaskBlueprintBase.h
--- a/Engine/Plugins/Runtime/StateTree/Source/StateTreeModule/Public/Blueprint/StateTreeTaskBlueprintBase.h
+++ b/Engine/Plugins/Runtime/StateTree/Source/StateTreeModule/Public/Blueprint/StateTreeTaskBlueprintBase.h
@@ -148,6 +148,8 @@ struct STATETREEMODULE_API FStateTreeBlueprintTaskWrapper : public FStateTreeTas
 	virtual EStateTreeRunStatus Tick(FStateTreeExecutionContext& Context, const float DeltaTime) const override;
 #if WITH_EDITOR
 	virtual FText GetDescription(const FGuid& ID, FStateTreeDataView InstanceDataView, const IStateTreeBindingLookup& BindingLookup, EStateTreeNodeFormatting Formatting = EStateTreeNodeFormatting::Text) const override;
+	virtual FName GetIconName() const override;
+	virtual FColor GetIconColor() const override;
 #endif	
 
 	UPROPERTY()
diff --git a/Engine/Plugins/Runtime/StateTree/Source/StateTreeModule/Public/StateTreeNodeBase.h b/Engine/Plugins/Runtime/StateTree/Source/StateTreeModule/Public/StateTreeNodeBase.h
--- a/Engine/Plugins/Runtime/StateTree/Source/StateTreeModule/Public/StateTreeNodeBase.h
+++ b/Engine/Plugins/Runtime/StateTree/Source/StateTreeModule/Public/StateTreeNodeBase.h
@@ -17,7 +17,7 @@ struct IStateTreeBindingLookup;
  * - Bold text should generally be used for actions, like name a of a task "<b>Play Animation</> {AnimName}".
  * - Subdued should be generally used for secondary/structural information, like "{Left} <s>equals</> {Right}".
  */
-UENUM()
+UENUM(BlueprintType)
 enum EStateTreeNodeFormatting : uint8
 {
 	/**
diff --git a/Engine/Source/Editor/Kismet/Private/BlueprintMemberReferenceCustomization.cpp b/Engine/Source/Editor/Kismet/Private/BlueprintMemberReferenceCustomization.cpp
--- a/Engine/Source/Editor/Kismet/Private/BlueprintMemberReferenceCustomization.cpp
+++ b/Engine/Source/Editor/Kismet/Private/BlueprintMemberReferenceCustomization.cpp
@@ -72,11 +72,12 @@ void FBlueprintMemberReferenceDetails::CustomizeHeader(TSharedRef<IPropertyHandl
 			return;
 		}
 		
+		const bool bAllowOnlyVisibleProperties = InStructPropertyHandle->HasMetaData("AllowOnlyVisibleProperties");
 		const bool bFunctionReference = InStructPropertyHandle->HasMetaData("FunctionReference");
 		const bool bPropertyReference = InStructPropertyHandle->HasMetaData("PropertyReference");
-		if(bPropertyReference || !bFunctionReference)
+		if(!bPropertyReference && !bFunctionReference)
 		{
-			// Only function references are supported right now
+			// Only property and function references are supported right now
 			return;
 		}
 
@@ -96,9 +97,19 @@ void FBlueprintMemberReferenceDetails::CustomizeHeader(TSharedRef<IPropertyHandl
 			return DefaultBindingName;
 		};
 		
-		auto OnCanBindProperty = [](FProperty* InPropertyOnCanBindProperty, TConstArrayView<FBindingChainElement> InBindingChain)
+		auto OnCanBindProperty = [bPropertyReference, bAllowOnlyVisibleProperties](FProperty* InPropertyOnCanBindProperty, TConstArrayView<FBindingChainElement> InBindingChain)
 		{
-			return true;
+			if (bPropertyReference)
+			{
+				if (bAllowOnlyVisibleProperties && InPropertyOnCanBindProperty)
+				{
+					return InPropertyOnCanBindProperty->HasAnyPropertyFlags(CPF_Edit)
+							&& (!InPropertyOnCanBindProperty->HasAnyPropertyFlags(CPF_NativeAccessSpecifierPrivate | CPF_NativeAccessSpecifierProtected)
+								|| InPropertyOnCanBindProperty->GetBoolMetaData(FBlueprintMetadata::MD_AllowPrivateAccess));
+				}
+				return true;
+			}
+			return false;
 		};
 
 		auto OnGoToBinding = [InStructPropertyHandle, Blueprint](FName InPropertyName)
@@ -125,23 +136,27 @@ void FBlueprintMemberReferenceDetails::CustomizeHeader(TSharedRef<IPropertyHandl
 			return false;
 		};
 		
-		auto OnCanGotoBinding = [InStructPropertyHandle](FName InPropertyName)
+		auto OnCanGotoBinding = [InStructPropertyHandle, Blueprint](FName InPropertyName)
 		{
 			void* StructData = nullptr;
 			const FPropertyAccess::Result Result = InStructPropertyHandle->GetValueData(StructData);
 			if(Result == FPropertyAccess::Success)
 			{
 				check(StructData);
-				FMemberReference* MemberReference = static_cast<FMemberReference*>(StructData);
-				return MemberReference->GetMemberName() != NAME_None;
+				const FMemberReference* MemberReference = static_cast<FMemberReference*>(StructData);
+				// Only support to go to functions for now.
+				const bool bHasValidFunction = MemberReference->ResolveMember<UFunction>(Blueprint->SkeletonGeneratedClass) != nullptr;
+				
+				return bHasValidFunction;
 			}
 			
 			return false;
 		};
 		
-		auto OnCanBindFunction = [PrototypeFunction](UFunction* InFunction)
+		auto OnCanBindFunction = [PrototypeFunction, bFunctionReference](UFunction* InFunction)
 		{
-			if(PrototypeFunction != nullptr)
+			if (PrototypeFunction != nullptr
+				&& bFunctionReference)
 			{
 				return PrototypeFunction->IsSignatureCompatibleWith(InFunction)
 					&& FBlueprintEditorUtils::HasFunctionBlueprintThreadSafeMetaData(PrototypeFunction) == FBlueprintEditorUtils::HasFunctionBlueprintThreadSafeMetaData(InFunction);
@@ -160,15 +175,29 @@ void FBlueprintMemberReferenceDetails::CustomizeHeader(TSharedRef<IPropertyHandl
 				
 				check(StructData);
 				FMemberReference* MemberReference = static_cast<FMemberReference*>(StructData);
-				UFunction* Function = InBindingChain[0].Field.Get<UFunction>();
-				UClass* OwnerClass = Function ? Function->GetOwnerClass() : nullptr;
-				bool bSelfContext = false;
-				if(OwnerClass != nullptr)
+				
+				if (UFunction* Function = InBindingChain[0].Field.Get<UFunction>())
 				{
-					bSelfContext = (Blueprint->GeneratedClass != nullptr && Blueprint->GeneratedClass->IsChildOf(OwnerClass)) ||
-									(Blueprint->SkeletonGeneratedClass != nullptr && Blueprint->SkeletonGeneratedClass->IsChildOf(OwnerClass));
+					UClass* OwnerClass = Function ? Function->GetOwnerClass() : nullptr;
+					bool bSelfContext = false;
+					if(OwnerClass != nullptr)
+					{
+						bSelfContext = (Blueprint->GeneratedClass != nullptr && Blueprint->GeneratedClass->IsChildOf(OwnerClass)) ||
+										(Blueprint->SkeletonGeneratedClass != nullptr && Blueprint->SkeletonGeneratedClass->IsChildOf(OwnerClass));
+					}
+					MemberReference->SetFromField<UFunction>(Function, bSelfContext);
+				}
+				else if (const FProperty* Property = InBindingChain[0].Field.Get<FProperty>())
+				{
+					UClass* OwnerClass = Property ? Property->GetOwnerClass() : nullptr;
+					bool bSelfContext = false;
+					if(OwnerClass != nullptr)
+					{
+						bSelfContext = (Blueprint->GeneratedClass != nullptr && Blueprint->GeneratedClass->IsChildOf(OwnerClass)) ||
+										(Blueprint->SkeletonGeneratedClass != nullptr && Blueprint->SkeletonGeneratedClass->IsChildOf(OwnerClass));
+					}
+					MemberReference->SetFromField<FProperty>(Property, bSelfContext);
 				}
-				MemberReference->SetFromField<UFunction>(Function, bSelfContext);
 
 				InStructPropertyHandle->NotifyPostChange(EPropertyChangeType::ValueSet);
 			}
@@ -218,7 +247,7 @@ void FBlueprintMemberReferenceDetails::CustomizeHeader(TSharedRef<IPropertyHandl
 			}
 		};
 
-		auto CurrentBindingText = [bFunctionReference, Blueprint, InStructPropertyHandle]()
+		auto CurrentBindingText = [Blueprint, InStructPropertyHandle]()
 		{
 			void* StructData = nullptr;
 			const FPropertyAccess::Result Result = InStructPropertyHandle->GetValueData(StructData);
@@ -226,18 +255,16 @@ void FBlueprintMemberReferenceDetails::CustomizeHeader(TSharedRef<IPropertyHandl
 			{
 				check(StructData);
 				FMemberReference* MemberReference = static_cast<FMemberReference*>(StructData);
-				if(bFunctionReference)
+
+				if (UFunction* Function = MemberReference->ResolveMember<UFunction>(Blueprint->SkeletonGeneratedClass))
 				{
-					UFunction* Function = MemberReference->ResolveMember<UFunction>(Blueprint->SkeletonGeneratedClass);
-					if(Function)
-					{
-						return FText::FromName(Function->GetFName());
-					}
-					else
-					{
-						return FText::FromName(MemberReference->GetMemberName());
-					}
+					return FText::FromName(Function->GetFName());
 				}
+				else if (const FProperty* Property = MemberReference->ResolveMember<FProperty>(Blueprint->SkeletonGeneratedClass))
+				{
+					return FText::FromName(Property->GetFName());
+				}
+				return FText::FromName(MemberReference->GetMemberName());
 			}
 			else if(Result == FPropertyAccess::MultipleValues)
 			{
@@ -264,6 +291,52 @@ void FBlueprintMemberReferenceDetails::CustomizeHeader(TSharedRef<IPropertyHandl
 
 			return FText::GetEmpty();
 		};
+
+		auto CurrentBindingImage = [Blueprint, InStructPropertyHandle]() -> const FSlateBrush*
+		{
+			void* StructData = nullptr;
+			const FPropertyAccess::Result Result = InStructPropertyHandle->GetValueData(StructData);
+			if(Result == FPropertyAccess::Success)
+			{
+				check(StructData);
+				FMemberReference* MemberReference = static_cast<FMemberReference*>(StructData);
+
+				if (MemberReference->ResolveMember<UFunction>(Blueprint->SkeletonGeneratedClass))
+				{
+					return FAppStyle::GetBrush("GraphEditor.Function_16x");
+				}
+				else if (MemberReference->ResolveMember<FProperty>(Blueprint->SkeletonGeneratedClass))
+				{
+					return FAppStyle::GetBrush("Kismet.Tabs.Variables");
+				}
+			}
+			return nullptr;
+		};
+
+		auto CurrentBindingColor = [Blueprint, InStructPropertyHandle]() -> FLinearColor
+		{
+			void* StructData = nullptr;
+			const FPropertyAccess::Result Result = InStructPropertyHandle->GetValueData(StructData);
+			if(Result == FPropertyAccess::Success)
+			{
+				check(StructData);
+				FMemberReference* MemberReference = static_cast<FMemberReference*>(StructData);
+
+				if (MemberReference->ResolveMember<UFunction>(Blueprint->SkeletonGeneratedClass))
+				{
+					return FAppStyle::GetSlateColor("Colors.Foreground").GetSpecifiedColor();
+				}
+				else if (const FProperty* Property = MemberReference->ResolveMember<FProperty>(Blueprint->SkeletonGeneratedClass))
+				{
+					const UEdGraphSchema_K2* Schema = GetDefault<UEdGraphSchema_K2>();
+					check(Schema);
+					FEdGraphPinType PinType;
+					Schema->ConvertPropertyToPinType(Property, PinType);
+					return Schema->GetPinTypeColor(PinType);
+				}
+			}
+			return FAppStyle::GetSlateColor("Colors.Foreground").GetSpecifiedColor();
+		};
 	
 		FPropertyBindingWidgetArgs Args;
 		Args.BindableSignature = PrototypeFunction;
@@ -279,13 +352,13 @@ void FBlueprintMemberReferenceDetails::CustomizeHeader(TSharedRef<IPropertyHandl
 		Args.OnNewFunctionBindingCreated = FOnNewFunctionBindingCreated::CreateLambda(OnNewFunctionBindingCreated);
 		Args.CurrentBindingText = MakeAttributeLambda(CurrentBindingText);
 		Args.CurrentBindingToolTipText = MakeAttributeLambda(CurrentBindingToolTipText);
-		Args.CurrentBindingImage = FAppStyle::GetBrush("GraphEditor.Function_16x");
-		Args.CurrentBindingColor = FAppStyle::GetSlateColor("Colors.Foreground").GetSpecifiedColor();
-		Args.bGeneratePureBindings = false;
+		Args.CurrentBindingImage = MakeAttributeLambda(CurrentBindingImage);
+		Args.CurrentBindingColor = MakeAttributeLambda(CurrentBindingColor);
+		Args.bGeneratePureBindings = bPropertyReference;
 		Args.bAllowFunctionBindings = bFunctionReference;
 		Args.bAllowFunctionLibraryBindings = bAllowFunctionLibraryReferences;
-		Args.bAllowPropertyBindings = false;
-		Args.bAllowNewBindings = true;
+		Args.bAllowPropertyBindings = bPropertyReference;
+		Args.bAllowNewBindings = bFunctionReference;
 		Args.bAllowArrayElementBindings = false;
 		Args.bAllowUObjectFunctions = false;
 		Args.bAllowStructFunctions = false;
