Added more stats to the Blueprint compiler.

In addition to more scoped events, we've hooked into the TraceLog API to report the name of a Blueprint to Insights.

Also, reorganized GTimeCompiling and GTimeReinstancing. When we flush compilation, we'll continue to report compilation time. However, reinstancing in a flush only refers to how long ReinstanceBatch took. When we flush the reinstancing queue, we'll only report how long that flush took.

#rnx

diff --git a/Engine/Source/Editor/Kismet/Kismet.Build.cs b/Engine/Source/Editor/Kismet/Kismet.Build.cs
--- a/Engine/Source/Editor/Kismet/Kismet.Build.cs
+++ b/Engine/Source/Editor/Kismet/Kismet.Build.cs
@@ -59,6 +59,7 @@ public class Kismet : ModuleRules
 				"SubobjectEditor",
 				"SubobjectDataInterface",
 				"ToolWidgets",
+				"TraceLog"
 			}
 			);
 
diff --git a/Engine/Source/Editor/Kismet/Private/BlueprintCompilationManager.cpp b/Engine/Source/Editor/Kismet/Private/BlueprintCompilationManager.cpp
--- a/Engine/Source/Editor/Kismet/Private/BlueprintCompilationManager.cpp
+++ b/Engine/Source/Editor/Kismet/Private/BlueprintCompilationManager.cpp
@@ -29,6 +29,7 @@
 #include "Kismet2/KismetEditorUtilities.h"
 #include "Kismet2/KismetReinstanceUtilities.h"
 #include "KismetCompiler.h"
+#include "Logging/StructuredLog.h"
 #include "Misc/ScopedSlowTask.h"
 #include "Misc/DataValidation.h"
 #include "Misc/PackageAccessTrackingOps.h"
@@ -37,6 +38,8 @@
 #include "Serialization/ArchiveReplaceObjectRef.h"
 #include "ProfilingDebugging/LoadTimeTracker.h"
 #include "TickableEditorObject.h"
+#include "Trace/Trace.h"
+#include "Trace/Trace.inl"
 #include "UObject/FortniteMainBranchObjectVersion.h"
 #include "UObject/MetaData.h"
 #include "UObject/ReferenceChainSearch.h"
@@ -172,7 +175,7 @@ struct FBlueprintCompilationManagerImpl : public FGCObject
 	bool bGeneratedClassLayoutReady;
 
 #if WITH_EDITOR
-	// Used to avoid reinstanciation on the GT while compiling on the loading thread
+	// Used to avoid reinstantiation on the GT while compiling on the loading thread
 	FCriticalSection Lock;
 #endif
 };
@@ -385,9 +388,6 @@ void FBlueprintCompilationManagerImpl::CompileSynchronouslyImpl(const FBPCompile
 	OldCDOs.Empty();
 }
 
-static double GTimeCompiling = 0.f;
-static double GTimeReinstancing = 0.f;
-
 enum class ECompilationManagerJobType
 {
 	Normal,
@@ -578,6 +578,11 @@ namespace UE::Kismet::BlueprintCompilationManager::Private
 	}
 }
 
+UE_TRACE_EVENT_BEGIN(Cpu, FlushCompilationQueue, NoSync)
+	UE_TRACE_EVENT_FIELD(UE::Trace::WideString, BlueprintPath)
+	UE_TRACE_EVENT_FIELD(int32, TotalBlueprints)
+UE_TRACE_EVENT_END()
+
 void FBlueprintCompilationManagerImpl::FlushCompilationQueueImpl(bool bSuppressBroadcastCompiled, TArray<UBlueprint*>* BlueprintsCompiled, TArray<UBlueprint*>* BlueprintsCompiledOrSkeletonCompiled, FUObjectSerializeContext* InLoadContext, TMap<UClass*, TMap<UObject*, UObject*>>* OldToNewTemplates /* = nullptr*/)
 {
 	TRACE_CPUPROFILER_EVENT_SCOPE(FlushCompilationQueueImpl);
@@ -598,8 +603,10 @@ void FBlueprintCompilationManagerImpl::FlushCompilationQueueImpl(bool bSuppressB
 	SlowTask.MakeDialogDelayed(1.0f);
 
 	TArray<FCompilerData> CurrentlyCompilingBPs;
-	{ // begin GTimeCompiling scope 
-		FScopedDurationTimer SetupTimer(GTimeCompiling); 
+	double TimeCompiling = 0.0;
+	{ // begin TimeCompiling scope 
+
+		FScopedDurationTimer SetupTimer(TimeCompiling); 
 
 		// STAGE I: Add any related blueprints that were not compiled, then add any children so that they will be relinked:
 		TArray<UBlueprint*> BlueprintsToRecompile;
@@ -875,6 +882,15 @@ void FBlueprintCompilationManagerImpl::FlushCompilationQueueImpl(bool bSuppressB
 		BlueprintsToRecompile.Empty();
 		QueuedRequests.Empty();
 
+#if CPUPROFILERTRACE_ENABLED
+		const FString& FirstBlueprintName = 
+			(CurrentlyCompilingBPs.Num() > 0) ? CurrentlyCompilingBPs[0].BP->GetFullName() : TEXT("<No Blueprints to compile>");
+
+		UE_TRACE_LOG_SCOPED_T(Cpu, FlushCompilationQueue, CpuChannel) <<
+			FlushCompilationQueue.BlueprintPath(*FirstBlueprintName) <<
+			FlushCompilationQueue.TotalBlueprints(CurrentlyCompilingBPs.Num());
+#endif
+
 		SlowTask.EnterProgressFrame();
 
 		// STAGE III: Sort into correct compilation order. We want to compile root types before their derived (child) types:
@@ -1561,11 +1577,12 @@ void FBlueprintCompilationManagerImpl::FlushCompilationQueueImpl(bool bSuppressB
 
 			ensure(BPGC == nullptr || BPGC->ClassDefaultObject->GetClass() == BPGC);
 		}
-	} // end GTimeCompiling scope
+	} // end TimeCompiling scope
 
 	SlowTask.EnterProgressFrame();
 
 	// STAGE XIV: Now we can finish the first stage of the reinstancing operation, moving old classes to new classes:
+	double TimeReinstancing = 0.0;
 	{
 		{
 			TRACE_CPUPROFILER_EVENT_SCOPE(MoveOldClassesToNewClasses);
@@ -1583,7 +1600,7 @@ void FBlueprintCompilationManagerImpl::FlushCompilationQueueImpl(bool bSuppressB
 				}
 			}
 
-			FScopedDurationTimer ReinstTimer(GTimeReinstancing);
+			FScopedDurationTimer ReinstTimer(TimeReinstancing);
 			ReinstanceBatch(Reinstancers, MutableView(ClassesToReinstance), InLoadContext, OldToNewTemplates);
 
 			// We purposefully do not remove the OldCDOs yet, need to keep them in memory past first GC
@@ -1807,9 +1824,13 @@ void FBlueprintCompilationManagerImpl::FlushCompilationQueueImpl(bool bSuppressB
 
 	SlowTask.EnterProgressFrame();
 
-	UE_LOG(LogBlueprint, Display, TEXT("Time Compiling: %f, Time Reinstancing: %f"),  GTimeCompiling, GTimeReinstancing);
-	//GTimeCompiling = 0.0;
-	//GTimeReinstancing = 0.0;
+	UE_LOGFMT(LogBlueprint, Display, "Finished compiling {0} Blueprint(s):", CurrentlyCompilingBPs.Num());
+	for (const FCompilerData& CompilerData : CurrentlyCompilingBPs)
+	{
+		UE_LOGFMT(LogBlueprint, Display, "\t{0}", CompilerData.BP->GetFullName());
+	}
+	UE_LOGFMT(LogBlueprint, Display, "\tTime Compiling: {0} ms. Time Reinstancing: {1} ms.", TimeCompiling * 1000.0, TimeReinstancing * 1000.0);
+
 	VerifyNoQueuedRequests(CurrentlyCompilingBPs);
 }
 
@@ -1926,20 +1947,22 @@ void FBlueprintCompilationManagerImpl::FlushReinstancingQueueImpl(bool bFindAndR
 	FScopeLock ScopeLock(&Lock);
 #endif
 
-	if(GCompilingBlueprint)
+	if (GCompilingBlueprint)
 	{
 		return;
 	}
 
 	TGuardValue<bool> GuardTemplateNameFlag(GCompilingBlueprint, true);
 	// we can finalize reinstancing now:
-	if(ClassesToReinstance.Num() == 0)
+	if (ClassesToReinstance.Num() == 0)
 	{
 		return;
 	}
 
+	double TimeReinstancing = 0.0;
+	int NumClassesToReinstance = ClassesToReinstance.Num();
 	{
-		FScopedDurationTimer ReinstTimer(GTimeReinstancing);
+		FScopedDurationTimer ReinstTimer(TimeReinstancing);
 		
 		TGuardValue<bool> ReinstancingGuard(GIsReinstancing, true);
 		
@@ -1952,7 +1975,7 @@ void FBlueprintCompilationManagerImpl::FlushReinstancingQueueImpl(bool bFindAndR
 		Options.OldToNewTemplates = OldToNewTemplates;
 		FBlueprintCompileReinstancer::BatchReplaceInstancesOfClass(ClassesToReinstanceOwned, Options);
 
-		// Special case when we run on ALT, we want to cleanup all classes flagged for reinstanciation right away.
+		// Special case when we run on ALT, we want to cleanup all classes flagged for reinstantiation right away.
 		const bool bIsInActualAsyncLoadingThread = IsInAsyncLoadingThread() && !IsInGameThread();
 		if (IsAsyncLoading() && (!IsAsyncLoadingMultithreaded() || !bIsInActualAsyncLoadingThread))
 		{
@@ -1960,7 +1983,7 @@ void FBlueprintCompilationManagerImpl::FlushReinstancingQueueImpl(bool bFindAndR
 			// async loaded. Those instances will need to be reinstanced once they finish
 			// loading, there's no race here because if any instances are created after
 			// we check ClassHasInstancesAsyncLoading they will be created with the new class:
-			for( TMap<UClass*, UClass*>::TIterator It(ClassesToReinstanceOwned); It; ++It )
+			for (TMap<UClass*, UClass*>::TIterator It(ClassesToReinstanceOwned); It; ++It)
 			{
 				if (!ClassHasInstancesAsyncLoading(It->Key))
 				{
@@ -1970,7 +1993,7 @@ void FBlueprintCompilationManagerImpl::FlushReinstancingQueueImpl(bool bFindAndR
 					It.RemoveCurrent();
 				}
 			}
-			// preserve any pairs that are currently loading:
+			// Preserve any pairs that are currently loading:
 			ClassesToReinstance = ObjectPtrWrap(MoveTemp(ClassesToReinstanceOwned));
 		}
 		else
@@ -1992,7 +2015,8 @@ void FBlueprintCompilationManagerImpl::FlushReinstancingQueueImpl(bool bFindAndR
 	FBlueprintSupport::ValidateNoExternalRefsToSkeletons();
 #endif
 
-	UE_LOG(LogBlueprint, Display, TEXT("Time Compiling: %f, Time Reinstancing: %f"),  GTimeCompiling, GTimeReinstancing);
+	UE_LOGFMT(LogBlueprint, Display, "Finished reinstancing for {0} class(es):", NumClassesToReinstance);
+	UE_LOGFMT(LogBlueprint, Display, "\tTime Reinstancing: {0} ms.", TimeReinstancing * 1000.0);
 }
 
 bool FBlueprintCompilationManagerImpl::HasBlueprintsToCompile() const
diff --git a/Engine/Source/Editor/KismetCompiler/Private/KismetCompiler.cpp b/Engine/Source/Editor/KismetCompiler/Private/KismetCompiler.cpp
--- a/Engine/Source/Editor/KismetCompiler/Private/KismetCompiler.cpp
+++ b/Engine/Source/Editor/KismetCompiler/Private/KismetCompiler.cpp
@@ -73,14 +73,16 @@ FSimpleMulticastDelegate FKismetCompilerContext::OnPostCompile;
 // Stats for this module
 DECLARE_CYCLE_STAT(TEXT("Create Schema"), EKismetCompilerStats_CreateSchema, STATGROUP_KismetCompiler );
 DECLARE_CYCLE_STAT(TEXT("Create Function List"), EKismetCompilerStats_CreateFunctionList, STATGROUP_KismetCompiler );
-DECLARE_CYCLE_STAT(TEXT("Expansion"), EKismetCompilerStats_Expansion, STATGROUP_KismetCompiler )
+DECLARE_CYCLE_STAT(TEXT("Expansion"), EKismetCompilerStats_Expansion, STATGROUP_KismetCompiler)
+DECLARE_CYCLE_STAT(TEXT("Validate"), EKismetCompilerStats_Validate, STATGROUP_KismetCompiler )
 DECLARE_CYCLE_STAT(TEXT("Expand Node"), EKismetCompilerStats_ExpandNode, STATGROUP_KismetCompiler )
 DECLARE_CYCLE_STAT(TEXT("Post Expansion Step"), EKismetCompilerStats_PostExpansionStep, STATGROUP_KismetCompiler )
 DECLARE_CYCLE_STAT(TEXT("Process uber"), EKismetCompilerStats_ProcessUbergraph, STATGROUP_KismetCompiler );
 DECLARE_CYCLE_STAT(TEXT("Process func"), EKismetCompilerStats_ProcessFunctionGraph, STATGROUP_KismetCompiler );
 DECLARE_CYCLE_STAT(TEXT("Generate Function Graph"), EKismetCompilerStats_GenerateFunctionGraphs, STATGROUP_KismetCompiler );
 DECLARE_CYCLE_STAT(TEXT("Precompile Function"), EKismetCompilerStats_PrecompileFunction, STATGROUP_KismetCompiler );
-DECLARE_CYCLE_STAT(TEXT("Compile Function"), EKismetCompilerStats_CompileFunction, STATGROUP_KismetCompiler );
+DECLARE_CYCLE_STAT(TEXT("Compile Function"), EKismetCompilerStats_CompileFunction, STATGROUP_KismetCompiler);
+DECLARE_CYCLE_STAT(TEXT("Create Locals and Register Nets"), EKismetCompilerStats_CreateLocalsAndRegisterNets, STATGROUP_KismetCompiler );
 DECLARE_CYCLE_STAT(TEXT("Postcompile Function"), EKismetCompilerStats_PostcompileFunction, STATGROUP_KismetCompiler );
 DECLARE_CYCLE_STAT(TEXT("Finalization"), EKismetCompilerStats_FinalizationWork, STATGROUP_KismetCompiler );
 DECLARE_CYCLE_STAT(TEXT("Code Gen"), EKismetCompilerStats_CodeGenerationTime, STATGROUP_KismetCompiler);
@@ -346,7 +348,7 @@ namespace UE::KismetCompiler::Private
 	// that's what the underlying delegate signature requires.
 	// Failure to do so will lead to incorrect UFunction signatures in the skeleton class,
 	// which can cause compilation errors.
-	static void MatchNodesToDelegateSignatures(UBlueprint* BP)
+	static void MatchNodesToDelegateSignatures(const UBlueprint* BP)
 	{
 		check(BP);
 
@@ -3852,6 +3854,8 @@ void FKismetCompilerContext::DetermineNodeExecLinks(UEdGraphNode* SourceNode, TM
 
 void FKismetCompilerContext::CreateLocalsAndRegisterNets(FKismetFunctionContext& Context, FField**& FunctionPropertyStorageLocation)
 {
+	BP_SCOPED_COMPILER_EVENT_STAT(EKismetCompilerStats_CreateLocalsAndRegisterNets);
+
 	// Create any user defined variables, this must occur before registering nets so that the properties are in place
 	CreateUserDefinedLocalVariablesForFunction(Context, FunctionPropertyStorageLocation);
 
@@ -4072,6 +4076,8 @@ void FKismetCompilerContext::CreateAndProcessUbergraph()
 
 		// Do some cursory validation (pin types match, inputs to outputs, pins never point to their parent node, etc...)
 		{
+			BP_SCOPED_COMPILER_EVENT_STAT(EKismetCompilerStats_Validate);
+
 			UbergraphContext = new FKismetFunctionContext(MessageLog, Schema, NewClass, Blueprint);
 			FunctionList.Add(UbergraphContext);
 			UbergraphContext->SourceGraph = ConsolidatedEventGraph;
