Provide a way to specify default compile option when compiling a BPClass

diff --git a/Engine/Source/Editor/Kismet/Private/BlueprintCompilationManager.cpp b/Engine/Source/Editor/Kismet/Private/BlueprintCompilationManager.cpp
--- a/Engine/Source/Editor/Kismet/Private/BlueprintCompilationManager.cpp
+++ b/Engine/Source/Editor/Kismet/Private/BlueprintCompilationManager.cpp
@@ -3718,6 +3718,19 @@ void MoveSkelCDOAside(UClass* Class, TMap<UClass*, UClass*>& OutOldToNewMap)
 	}
 };
 
+// ***************************************************************
+// FBPCompileRequest
+// ***************************************************************
+
+FBPCompileRequest::FBPCompileRequest(UBlueprint* InBPToCompile, EBlueprintCompileOptions InCompileOptions, FCompilerResultsLog* InClientResultsLog): BPToCompile(InBPToCompile)
+	, CompileOptions(InBPToCompile ? InCompileOptions | InBPToCompile->GetDefaultCompileOptions() : InCompileOptions)
+	, ClientResultsLog(InClientResultsLog)
+{
+}
+
+// ***************************************************************
+// FBlueprintCompilationManager
+// ***************************************************************
 void FBlueprintCompilationManager::Initialize()
 {
 	if(!BPCMImpl)
diff --git a/Engine/Source/Editor/Kismet/Public/BlueprintCompilationManager.h b/Engine/Source/Editor/Kismet/Public/BlueprintCompilationManager.h
--- a/Engine/Source/Editor/Kismet/Public/BlueprintCompilationManager.h
+++ b/Engine/Source/Editor/Kismet/Public/BlueprintCompilationManager.h
@@ -20,12 +20,7 @@ struct FUObjectSerializeContext;
 
 struct FBPCompileRequest
 {
-	explicit FBPCompileRequest(UBlueprint* InBPToCompile, EBlueprintCompileOptions InCompileOptions, FCompilerResultsLog* InClientResultsLog )
-		: BPToCompile(InBPToCompile)
-		, CompileOptions(InCompileOptions)
-		, ClientResultsLog(InClientResultsLog)
-	{
-	}
+	explicit KISMET_API FBPCompileRequest(UBlueprint* InBPToCompile, EBlueprintCompileOptions InCompileOptions, FCompilerResultsLog* InClientResultsLog);
 
 	// BP that needs to be compiled:
 	TObjectPtr<UBlueprint> BPToCompile;
diff --git a/Engine/Source/Runtime/Engine/Classes/Engine/Blueprint.h b/Engine/Source/Runtime/Engine/Classes/Engine/Blueprint.h
--- a/Engine/Source/Runtime/Engine/Classes/Engine/Blueprint.h
+++ b/Engine/Source/Runtime/Engine/Classes/Engine/Blueprint.h
@@ -22,6 +22,7 @@
 
 #include "Blueprint.generated.h"
 
+enum class EBlueprintCompileOptions;
 class ITargetPlatform;
 class UActorComponent;
 class UEdGraph;
@@ -850,6 +851,9 @@ public:
 		return true;
 	}
 
+	// Compile this prefab, overriding this methods give a chance to derived classes to specify extra default compilation flags
+	ENGINE_API virtual EBlueprintCompileOptions GetDefaultCompileOptions() const;
+
 	/** Sets the current object being debugged */
 	ENGINE_API virtual void SetObjectBeingDebugged(UObject* NewObject);
 
diff --git a/Engine/Source/Runtime/Engine/Private/Blueprint.cpp b/Engine/Source/Runtime/Engine/Private/Blueprint.cpp
--- a/Engine/Source/Runtime/Engine/Private/Blueprint.cpp
+++ b/Engine/Source/Runtime/Engine/Private/Blueprint.cpp
@@ -869,6 +869,11 @@ UClass* UBlueprint::GetBlueprintClass() const
 	return UBlueprintGeneratedClass::StaticClass();
 }
 
+EBlueprintCompileOptions UBlueprint::GetDefaultCompileOptions() const
+{
+	return EBlueprintCompileOptions::None;
+}
+
 void UBlueprint::SetObjectBeingDebugged(UObject* NewObject)
 {
 	// Unregister the old object (even if PendingKill)
