PR #11754: Use the actor label when displaying 'Add event for {0}' action in the level blueprint actions

diff --git a/Engine/Source/Editor/Kismet/Private/BlueprintActionMenuUtils.cpp b/Engine/Source/Editor/Kismet/Private/BlueprintActionMenuUtils.cpp
--- a/Engine/Source/Editor/Kismet/Private/BlueprintActionMenuUtils.cpp
+++ b/Engine/Source/Editor/Kismet/Private/BlueprintActionMenuUtils.cpp
@@ -396,7 +396,20 @@ static void BlueprintActionMenuUtilsImpl::AddLevelActorSections(FBlueprintAction
 
 	if (LevelActorsFilter.Context.SelectedObjects.Num() == 1)
 	{
-		FText const ActorName = FText::FromName(LevelActorsFilter.Context.SelectedObjects.Last().GetFName());
+		auto GetDisplayNameText = [](const FFieldVariant& Element) -> const FText
+		{
+			// If we have an actor, then use it's name
+			if (const AActor* const Actor = Element.Get<AActor>())
+			{
+				return FText::FromString(Actor->GetActorLabel());
+			}
+
+			// otherwise, fall back to just using its FName
+			return FText::FromName(Element.GetFName());
+		};
+		
+		FText const ActorName = GetDisplayNameText(LevelActorsFilter.Context.SelectedObjects.Last());
+		
 		FuncSectionHeading  = FText::Format(LOCTEXT("SingleActorFuncCategory", "Call Function on {0}"), ActorName);
 		EventSectionHeading = FText::Format(LOCTEXT("SingleActorEventCategory", "Add Event for {0}"), ActorName);
 	}
