UE-117208 - Comment Node doesn't update in the Bookmarks tab after being renamed in the Event Graph

Anytime a node's title is changed, we need to refresh the bookmarks tab.

diff --git a/Engine/Source/Editor/Kismet/Private/BlueprintEditor.cpp b/Engine/Source/Editor/Kismet/Private/BlueprintEditor.cpp
--- a/Engine/Source/Editor/Kismet/Private/BlueprintEditor.cpp
+++ b/Engine/Source/Editor/Kismet/Private/BlueprintEditor.cpp
@@ -10002,6 +10002,10 @@ void FBlueprintEditor::OnNodeTitleCommitted(const FText& NewText, ETextCommit::T
 		const FScopedTransaction Transaction( NSLOCTEXT( "K2_RenameNode", "RenameNode", "Rename Node" ) );
 		NodeBeingChanged->Modify();
 		NodeBeingChanged->OnRenameNode(NewText.ToString());
+		if (BookmarksWidget.IsValid())
+		{
+			BookmarksWidget->RefreshBookmarksTree();
+		}
 	}
 }
 
