Cache archetype before reinstancing a class, as it is impossible during the process as the old archeyptes and instances get renamed during the process.
This change only affect object with OverridableSerialization enabled(prefabs so far)

diff --git a/Engine/Source/Editor/Kismet/Private/BlueprintCompilationManager.cpp b/Engine/Source/Editor/Kismet/Private/BlueprintCompilationManager.cpp
--- a/Engine/Source/Editor/Kismet/Private/BlueprintCompilationManager.cpp
+++ b/Engine/Source/Editor/Kismet/Private/BlueprintCompilationManager.cpp
@@ -2679,6 +2679,7 @@ void FBlueprintCompilationManagerImpl::ReinstanceBatch(TArray<FReinstancingJob>&
 				{ 
 					bool bIsArchetype = 
 						Obj->HasAnyFlags(RF_ArchetypeObject|RF_InheritableComponentTemplate)
+						|| (FOverridableManager::Get().IsEnabled(*Obj) && Obj->GetOutermostObject()->HasAnyFlags(RF_ClassDefaultObject))
 						|| Obj->GetTypedOuter<UBlueprintGeneratedClass>()
 						|| Obj->GetTypedOuter<UBlueprint>();
 					// remove if this is not an archetype or its already in the transient package, note
@@ -2728,9 +2729,22 @@ void FBlueprintCompilationManagerImpl::ReinstanceBatch(TArray<FReinstancingJob>&
 							}
 
 							// this handles nested subobjects:
-							TArray<UObject*> ContainedObjects;
-							GetObjectsWithOuter(Iter, ContainedObjects);
-							ArchetypeReferencers.Append(ContainedObjects);
+							if (FOverridableManager::Get().IsEnabled(*Archetype))
+							{
+								ForEachObjectWithOuter(Iter, [Archetype, &ArchetypeReferencers](UObject* SubObject)
+								{
+									if (SubObject != Archetype && !SubObject->IsIn(Archetype))
+									{
+										ArchetypeReferencers.Add(SubObject);
+									}
+								});
+							}
+							else
+							{
+								TArray<UObject*> ContainedObjects;
+								GetObjectsWithOuter(Iter, ContainedObjects);
+								ArchetypeReferencers.Append(ContainedObjects);
+							}
 						}
 						Iter = Iter->GetOuter();
 					}
@@ -2781,6 +2795,23 @@ void FBlueprintCompilationManagerImpl::ReinstanceBatch(TArray<FReinstancingJob>&
 				UObject* OriginalOuter = OldInstance->GetOuter();
 				EObjectFlags OriginalFlags = OldInstance->GetFlags();
 
+				// We need to cache the archetype of the instances of this archetype
+				// as it will not be possible to get them afterwards as it gets renamed
+				if (FOverridableManager::Get().IsEnabled(*OldInstance))
+				{
+					TArray<UObject*> ArchetypeInstances;
+					OldInstance->GetArchetypeInstances(ArchetypeInstances);
+					for (UObject* ArchetypeInstance : ArchetypeInstances)
+					{
+						FOverridableManager::Get().CacheArchetype(*ArchetypeInstance);
+
+						ForEachObjectWithOuter(ArchetypeInstance, [](UObject* SubObject)
+						{
+							FOverridableManager::Get().CacheArchetype(*SubObject);
+						});
+					}
+				}
+
 				UObject* Destination = GetTransientOuterForRename(OldInstance->GetClass());
 				OldInstance->Rename(
 					nullptr,
diff --git a/Engine/Source/Editor/UnrealEd/Private/Kismet2/KismetReinstanceUtilities.cpp b/Engine/Source/Editor/UnrealEd/Private/Kismet2/KismetReinstanceUtilities.cpp
--- a/Engine/Source/Editor/UnrealEd/Private/Kismet2/KismetReinstanceUtilities.cpp
+++ b/Engine/Source/Editor/UnrealEd/Private/Kismet2/KismetReinstanceUtilities.cpp
@@ -2711,6 +2711,23 @@ void FBlueprintCompileReinstancer::ReplaceInstancesOfClass_Inner(const TMap<UCla
 						}
 						return Dependencies;
 					});
+
+					// We need to cache the archetype of the objects about to be replaced 
+					// as it will not be possible to get them during this process as it renames the objects
+					for (UObject* OldObject : ObjectsToReplace)
+					{
+						if(!IsValid(OldObject))
+						{
+							continue;
+						}
+
+						FOverridableManager::Get().CacheArchetype(*OldObject);
+
+						ForEachObjectWithOuter(OldObject, [](UObject* SubObject)
+						{
+							FOverridableManager::Get().CacheArchetype(*SubObject);
+						});
+					}
 				}
 				
 				// Then fix 'real' (non archetype) instances of the class
diff --git a/Engine/Source/Runtime/CoreUObject/Private/UObject/OverridableManager.cpp b/Engine/Source/Runtime/CoreUObject/Private/UObject/OverridableManager.cpp
--- a/Engine/Source/Runtime/CoreUObject/Private/UObject/OverridableManager.cpp
+++ b/Engine/Source/Runtime/CoreUObject/Private/UObject/OverridableManager.cpp
@@ -53,6 +53,14 @@ void FOverridableManager::InheritEnabledFrom(UObject& Object, const UObject* Def
 #endif // WITH_EDITORONLY_DATA
 }
 
+void FOverridableManager::CacheArchetype(UObject& Object)
+{
+	if (FOverriddenPropertySet* OverriddenProperties = OverriddenObjectAnnotations.Find(Object))
+	{
+		OverriddenProperties->CacheArchetype();
+	}
+}
+
 bool FOverridableManager::NeedSubObjectTemplateInstantiation(const UObject& Object)
 {
 #ifdef WITH_EDITORONLY_DATA
diff --git a/Engine/Source/Runtime/CoreUObject/Private/UObject/OverriddenPropertySet.cpp b/Engine/Source/Runtime/CoreUObject/Private/UObject/OverriddenPropertySet.cpp
--- a/Engine/Source/Runtime/CoreUObject/Private/UObject/OverriddenPropertySet.cpp
+++ b/Engine/Source/Runtime/CoreUObject/Private/UObject/OverriddenPropertySet.cpp
@@ -1026,6 +1026,14 @@ void FOverriddenPropertySet::RemoveOverriddenSubProperties(FOverriddenPropertyNo
 	PropertyNode.SubPropertyNodeKeys.Empty();
 }
 
+void FOverriddenPropertySet::CacheArchetype()
+{
+	if (!CachedArchetype)
+	{
+		CachedArchetype = Owner->GetArchetype();
+	}
+}
+
 EOverriddenPropertyOperation FOverriddenPropertySet::GetOverriddenPropertyOperation(const FPropertyChangedEvent& PropertyEvent, const FEditPropertyChain::TDoubleLinkedListNode* PropertyNode, bool* bOutInheritedOperation /*= nullptr*/) const
 {
 	if (const FOverriddenPropertyNode* RootNode = OverriddenPropertyNodes.FindByHash(GetTypeHash(RootNodeID), RootNodeID))
@@ -1189,6 +1197,13 @@ void FOverriddenPropertySet::Reset()
 
 void FOverriddenPropertySet::HandleObjectsReInstantiated(const TMap<UObject*, UObject*>& Map)
 {
+	// When the saved archetype is set, it is an indicator this object is about to be replaced
+	// So no need to replace any ptr, otherwise we might not be able to reconstitute the right information
+	if (CachedArchetype)
+	{
+		return;
+	}
+
 	for (FOverriddenPropertyNode& Node : OverriddenPropertyNodes)
 	{
 		Node.NodeID.HandleObjectsReInstantiated(Map);
diff --git a/Engine/Source/Runtime/CoreUObject/Private/UObject/UObjectArchetype.cpp b/Engine/Source/Runtime/CoreUObject/Private/UObject/UObjectArchetype.cpp
--- a/Engine/Source/Runtime/CoreUObject/Private/UObject/UObjectArchetype.cpp
+++ b/Engine/Source/Runtime/CoreUObject/Private/UObject/UObjectArchetype.cpp
@@ -13,6 +13,7 @@
 #include "UObject/UObjectAnnotation.h"
 #include "Stats/StatsMisc.h"
 #include "HAL/IConsoleManager.h"
+#include "UObject/OverridableManager.h"
 #include "UObject/UObjectArchetypeHelper.h"
 
 #define UE_CACHE_ARCHETYPE (1 && !WITH_EDITORONLY_DATA)
@@ -203,6 +204,16 @@ UObject* GetArchetypeImpl(const UObject* InObject, const FObjectArchetypeHelper:
 			return Archetype;
 		}
 	}
+
+	if (const FOverriddenPropertySet* OverriddenProperties = FOverridableManager::Get().GetOverriddenProperties(*InObject))
+	{
+		// Use the cached archetype if set
+		if (UObject* CacheArchetype = OverriddenProperties->GetCachedArchetype())
+		{
+			return CacheArchetype;
+		}
+	}
+
 #endif
 
 	bool bUseUpToDateClass = false;
diff --git a/Engine/Source/Runtime/CoreUObject/Public/UObject/OverridableManager.h b/Engine/Source/Runtime/CoreUObject/Public/UObject/OverridableManager.h
--- a/Engine/Source/Runtime/CoreUObject/Public/UObject/OverridableManager.h
+++ b/Engine/Source/Runtime/CoreUObject/Public/UObject/OverridableManager.h
@@ -127,6 +127,14 @@ public:
 	 * @param DefaultData its default object */
 	COREUOBJECT_API void InheritEnabledFrom(UObject& Object, const UObject* DefaultData);
 
+	/**
+	 * Caches the archetype for this object if it has overridable serialization enabled.
+	 * NOTE: ONLY CACHE THE ARCHETYPE POINTER ON OBJECTS THAT ARE ABOUT TO BE REINSTANTIATED.
+	 * This is only a utility method to ensure the archetype is not lost during reinstantiation
+	 * as it renames and changes outers which makes impossible to retrieve it during the process.
+	 * @param Object top cache its archetype */
+	COREUOBJECT_API void CacheArchetype(UObject& Object);
+
 	/**
 	 * Return true if this object needs subobject template instantiation
 	 * @param Object to be querying about */
diff --git a/Engine/Source/Runtime/CoreUObject/Public/UObject/OverriddenPropertySet.h b/Engine/Source/Runtime/CoreUObject/Public/UObject/OverriddenPropertySet.h
--- a/Engine/Source/Runtime/CoreUObject/Public/UObject/OverriddenPropertySet.h
+++ b/Engine/Source/Runtime/CoreUObject/Public/UObject/OverriddenPropertySet.h
@@ -292,6 +292,18 @@ public:
 	void Reset();
 	void HandleObjectsReInstantiated(const TMap<UObject*, UObject*>& Map);
 
+	/**
+	 * Caches the archetype of this object for future reference */
+	void CacheArchetype();
+
+	/**
+	 * Gets the archetype that was cached when the CacheArchetype method was called
+	 * @return the cached Archetype if any */
+	UObject* GetCachedArchetype() const
+	{
+		return CachedArchetype;
+	}
+
 protected:
 
 	FOverriddenPropertyNode& FindOrAddNode(FOverriddenPropertyNode& ParentPropertyNode, FOverriddenPropertyNodeID NodeID);
@@ -316,6 +328,8 @@ private:
 
 	static inline FOverriddenPropertyNodeID RootNodeID = FOverriddenPropertyNodeID(FName(TEXT("root")));
 
+	/** Archetype is cache when it is or its archetype is about to be replaced **/
+	TObjectPtr<UObject> CachedArchetype = nullptr;
 public:
 	bool bNeedsSubobjectTemplateInstantiation = false;
 };
