UserDefinedStruct: Add support for meta data storage and add min/max clamp/UI value for numeric fields

diff --git a/Engine/Source/Editor/Kismet/Private/UserDefinedStructureEditor.cpp b/Engine/Source/Editor/Kismet/Private/UserDefinedStructureEditor.cpp
--- a/Engine/Source/Editor/Kismet/Private/UserDefinedStructureEditor.cpp
+++ b/Engine/Source/Editor/Kismet/Private/UserDefinedStructureEditor.cpp
@@ -1001,6 +1001,39 @@ public:
 		}
 	}
 
+	// Value Range
+	EVisibility IsValueRangeOptionVisible() const
+	{
+		if (const TSharedPtr<FUserDefinedStructureDetails> StructureDetailsSP = StructureDetails.Pin())
+		{
+			return FStructureEditorUtils::CanEditValueRange(StructureDetailsSP->GetUserDefinedStruct(), FieldGuid)
+				? EVisibility::Visible
+				: EVisibility::Collapsed;
+		}
+		return EVisibility::Collapsed;
+	}
+
+	// Meta Data
+	FText OnGetMetaData(const FName Key) const
+	{
+		const TSharedPtr<FUserDefinedStructureDetails> Details = StructureDetails.Pin();
+		if (!Details)
+		{
+			return FText();
+		}
+
+		const FString* Value = FStructureEditorUtils::GetMetaData(Details->GetUserDefinedStruct(), FieldGuid, Key);
+		return Value ? FText::FromString(*Value) : FText();
+	}
+
+	void OnMetaDataCommitted(const FText& NewText, ETextCommit::Type, const FName Key)
+	{
+		if (const TSharedPtr<FUserDefinedStructureDetails> StructureDetailsSP = StructureDetails.Pin())
+		{
+			FStructureEditorUtils::SetMetaData(StructureDetailsSP->GetUserDefinedStruct(), FieldGuid, Key, NewText.ToString());
+		}
+	}
+
 	/** IDetailCustomNodeBuilder Interface*/
 	virtual void SetOnRebuildChildren( FSimpleDelegate InOnRegenerateChildren ) override 
 	{
@@ -1182,6 +1215,78 @@ public:
 			.IsChecked(this, &FUserDefinedStructureFieldLayout::OnGet3dWidgetEnabled)
 		]
 		.Visibility(TAttribute<EVisibility>::Create(TAttribute<EVisibility>::FGetter::CreateSP(this, &FUserDefinedStructureFieldLayout::Is3dWidgetOptionVisible)));
+
+		ChildrenBuilder.AddCustomRow(LOCTEXT("ClampRange", "Value Range"))
+		.NameContent()
+		[
+			SNew(STextBlock)
+			.Text(LOCTEXT("ValueRangeLabel", "Value Range"))
+			.Font(IDetailLayoutBuilder::GetDetailFont())
+		]
+		.ValueContent()
+		[
+			SNew(SHorizontalBox)
+			+SHorizontalBox::Slot()
+			.FillWidth(1)
+			[
+				SNew(SEditableTextBox)
+				.OnTextCommitted(this, &FUserDefinedStructureFieldLayout::OnMetaDataCommitted, FStructVariableMetaData::ClampMin)
+				.Text(this, &FUserDefinedStructureFieldLayout::OnGetMetaData, FStructVariableMetaData::ClampMin)
+				.Font(IDetailLayoutBuilder::GetDetailFont())
+			]
+			+SHorizontalBox::Slot()
+			.AutoWidth()
+			[
+				SNew(STextBlock)
+				.Text(LOCTEXT("Min .. Max Separator", " .. "))
+				.Font(IDetailLayoutBuilder::GetDetailFont())
+			]
+			+SHorizontalBox::Slot()
+			.FillWidth(1)
+			[
+				SNew(SEditableTextBox)
+				.OnTextCommitted(this, &FUserDefinedStructureFieldLayout::OnMetaDataCommitted, FStructVariableMetaData::ClampMax)
+				.Text(this, &FUserDefinedStructureFieldLayout::OnGetMetaData, FStructVariableMetaData::ClampMax)
+				.Font(IDetailLayoutBuilder::GetDetailFont())
+			]
+		]
+		.Visibility(TAttribute<EVisibility>::Create(TAttribute<EVisibility>::FGetter::CreateSP(this, &FUserDefinedStructureFieldLayout::IsValueRangeOptionVisible)));
+
+		ChildrenBuilder.AddCustomRow(LOCTEXT("UIRange", "Slider Range"))
+		.NameContent()
+		[
+			SNew(STextBlock)
+			.Text(LOCTEXT("SliderRangeLabel", "Slider Range"))
+			.Font(IDetailLayoutBuilder::GetDetailFont())
+		]
+		.ValueContent()
+		[
+			SNew(SHorizontalBox)
+			+SHorizontalBox::Slot()
+			.FillWidth(1)
+			[
+				SNew(SEditableTextBox)
+				.OnTextCommitted(this, &FUserDefinedStructureFieldLayout::OnMetaDataCommitted, FStructVariableMetaData::UIMin)
+				.Text(this, &FUserDefinedStructureFieldLayout::OnGetMetaData, FStructVariableMetaData::UIMin)
+				.Font(IDetailLayoutBuilder::GetDetailFont())
+			]
+			+SHorizontalBox::Slot()
+			.AutoWidth()
+			[
+				SNew(STextBlock)
+				.Text(LOCTEXT("Min .. Max Separator", " .. "))
+				.Font(IDetailLayoutBuilder::GetDetailFont())
+			]
+			+SHorizontalBox::Slot()
+			.FillWidth(1)
+			[
+				SNew(SEditableTextBox)
+				.OnTextCommitted(this, &FUserDefinedStructureFieldLayout::OnMetaDataCommitted, FStructVariableMetaData::UIMax)
+				.Text(this, &FUserDefinedStructureFieldLayout::OnGetMetaData, FStructVariableMetaData::UIMax)
+				.Font(IDetailLayoutBuilder::GetDetailFont())
+			]
+		]
+		.Visibility(TAttribute<EVisibility>::Create(TAttribute<EVisibility>::FGetter::CreateSP(this, &FUserDefinedStructureFieldLayout::IsValueRangeOptionVisible)));
 	}
 
 	virtual void Tick( float DeltaTime ) override {}
diff --git a/Engine/Source/Editor/KismetCompiler/Private/UserDefinedStructureCompilerUtils.cpp b/Engine/Source/Editor/KismetCompiler/Private/UserDefinedStructureCompilerUtils.cpp
--- a/Engine/Source/Editor/KismetCompiler/Private/UserDefinedStructureCompilerUtils.cpp
+++ b/Engine/Source/Editor/KismetCompiler/Private/UserDefinedStructureCompilerUtils.cpp
@@ -285,6 +285,7 @@ struct FUserDefinedStructureCompilerInner
 			}
 			VarProperty->SetMetaData(TEXT("DisplayName"), *VarDesc.FriendlyName);
 			VarProperty->SetMetaData(FBlueprintMetadata::MD_Tooltip, *VarDesc.ToolTip);
+			VarProperty->AppendMetaData(VarDesc.MetaData);
 			VarProperty->RepNotifyFunc = NAME_None;
 
 			if (!VarDesc.DefaultValue.IsEmpty())
diff --git a/Engine/Source/Editor/UnrealEd/Classes/UserDefinedStructure/UserDefinedStructEditorData.h b/Engine/Source/Editor/UnrealEd/Classes/UserDefinedStructure/UserDefinedStructEditorData.h
--- a/Engine/Source/Editor/UnrealEd/Classes/UserDefinedStructure/UserDefinedStructEditorData.h
+++ b/Engine/Source/Editor/UnrealEd/Classes/UserDefinedStructure/UserDefinedStructEditorData.h
@@ -15,6 +15,18 @@
 
 class ITransactionObjectAnnotation;
 
+struct UNREALED_API FStructVariableMetaData
+{
+	static const FName ClampMin;
+	static const FName ClampMax;
+	static const FName UIMin;
+	static const FName UIMax;
+
+private:
+	// This struct should never be instantiated
+	FStructVariableMetaData();
+};
+
 USTRUCT()
 struct FStructVariableDescription
 {
@@ -82,6 +94,9 @@ struct FStructVariableDescription
 	UPROPERTY()
 	FString ToolTip;
 
+	UPROPERTY()
+	TMap<FName, FString> MetaData;
+
 	UNREALED_API bool SetPinType(const struct FEdGraphPinType& VarType);
 
 	UNREALED_API FEdGraphPinType ToPinType() const;
diff --git a/Engine/Source/Editor/UnrealEd/Private/Kismet2/StructureEditorUtils.cpp b/Engine/Source/Editor/UnrealEd/Private/Kismet2/StructureEditorUtils.cpp
--- a/Engine/Source/Editor/UnrealEd/Private/Kismet2/StructureEditorUtils.cpp
+++ b/Engine/Source/Editor/UnrealEd/Private/Kismet2/StructureEditorUtils.cpp
@@ -955,6 +955,73 @@ bool FStructureEditorUtils::Is3dWidgetEnabled(const UUserDefinedStruct* Struct,
 	return VarDesc && VarDesc->bEnable3dWidget && FEdMode::CanCreateWidgetForStructure(PropertyStruct);
 }
 
+bool FStructureEditorUtils::CanEditValueRange(const UUserDefinedStruct* Struct, FGuid VarGuid)
+{
+	if (const FStructVariableDescription* VarDesc = GetVarDescByGuid(Struct, VarGuid))
+	{
+		if (const FProperty* Property = FindFProperty<FProperty>(Struct, VarDesc->VarName))
+		{
+			return Property->IsA(FNumericProperty::StaticClass());
+		}
+	}
+	return false;
+}
+
+bool FStructureEditorUtils::SetMetaData(UUserDefinedStruct* Struct, FGuid VarGuid, FName Key, const FString& Value)
+{
+	FStructVariableDescription* VarDesc = GetVarDescByGuid(Struct, VarGuid);
+	if (!VarDesc)
+	{
+		return false;
+	}
+
+	const FString* CurrentValue = VarDesc->MetaData.Find(Key);
+
+	if (Value.IsEmpty())
+	{
+		if (!CurrentValue)
+		{
+			// No new or old value, nothing to do
+			return false;
+		}
+
+		const FScopedTransaction Transaction(LOCTEXT("RemoveMetaData", "Unset Meta Data"));
+		ModifyStructData(Struct);
+
+		VarDesc->MetaData.Remove(Key);
+		if (FProperty* Property = FindFProperty<FProperty>(Struct, VarDesc->VarName))
+		{
+			Property->RemoveMetaData(Key);
+		}
+
+		OnStructureChanged(Struct);
+		return true;
+	}
+
+	if (CurrentValue && CurrentValue->Equals(Value, ESearchCase::CaseSensitive))
+	{
+		// There's both old and new value, and they are the same, so nothing to do
+		return false;
+	}
+
+	const FScopedTransaction Transaction(LOCTEXT("ChangeMetaData", "Set Meta Data"));
+	ModifyStructData(Struct);
+
+	VarDesc->MetaData.Add(Key, Value);
+	if (FProperty* Property = FindFProperty<FProperty>(Struct, VarDesc->VarName))
+	{
+		Property->SetMetaData(Key, *Value);
+	}
+
+	OnStructureChanged(Struct);
+	return true;
+}
+
+const FString* FStructureEditorUtils::GetMetaData(const UUserDefinedStruct* Struct, FGuid VarGuid, FName Key)
+{
+	return GetVarDescByGuid(Struct, VarGuid)->MetaData.Find(Key);
+}
+
 FGuid FStructureEditorUtils::GetGuidForProperty(const FProperty* Property)
 {
 	const UUserDefinedStruct* UDStruct = Property ? Cast<const UUserDefinedStruct>(Property->GetOwnerStruct()) : nullptr;
diff --git a/Engine/Source/Editor/UnrealEd/Private/UserDefinedStructEditorData.cpp b/Engine/Source/Editor/UnrealEd/Private/UserDefinedStructEditorData.cpp
--- a/Engine/Source/Editor/UnrealEd/Private/UserDefinedStructEditorData.cpp
+++ b/Engine/Source/Editor/UnrealEd/Private/UserDefinedStructEditorData.cpp
@@ -14,6 +14,11 @@
 
 #define LOCTEXT_NAMESPACE "UserDefinedStructEditorData"
 
+const FName FStructVariableMetaData::ClampMin("ClampMin");
+const FName FStructVariableMetaData::ClampMax("ClampMax");
+const FName FStructVariableMetaData::UIMin("UIMin");
+const FName FStructVariableMetaData::UIMax("UIMax");
+
 void FStructVariableDescription::PostSerialize(const FArchive& Ar)
 {
 	if (ContainerType == EPinContainerType::None)
diff --git a/Engine/Source/Editor/UnrealEd/Public/Kismet2/StructureEditorUtils.h b/Engine/Source/Editor/UnrealEd/Public/Kismet2/StructureEditorUtils.h
--- a/Engine/Source/Editor/UnrealEd/Public/Kismet2/StructureEditorUtils.h
+++ b/Engine/Source/Editor/UnrealEd/Public/Kismet2/StructureEditorUtils.h
@@ -140,6 +140,13 @@ public:
 
 	static UNREALED_API bool Is3dWidgetEnabled(const UUserDefinedStruct* Struct, FGuid VarGuid);
 
+	//Value Range
+	static UNREALED_API bool CanEditValueRange(const UUserDefinedStruct* Struct, FGuid VarGuid);
+
+	//META DATA
+	static UNREALED_API bool SetMetaData(UUserDefinedStruct* Struct, FGuid VarGuid, FName Key, const FString& Value);
+	static UNREALED_API const FString* GetMetaData(const UUserDefinedStruct* Struct, FGuid VarGuid, FName Key);
+
 	//GUID AND VAR DESC
 	static UNREALED_API TArray<FStructVariableDescription>& GetVarDesc(UUserDefinedStruct* Struct);
 
