Avoid emitting math calls when the native function is in a module that is not guaranteed to be compatable with the callers
#rnx

diff --git a/Engine/Source/Editor/Kismet/Internal/ProjectUtilities/BuildTargetSet.h b/Engine/Source/Editor/Kismet/Internal/ProjectUtilities/BuildTargetSet.h
new file mode 100644
--- /dev/null
+++ b/Engine/Source/Editor/Kismet/Internal/ProjectUtilities/BuildTargetSet.h
@@ -0,0 +1,55 @@
+// Copyright Epic Games, Inc. All Rights Reserved.
+ 
+#pragma once
+
+#include "Containers/UnrealString.h"
+#include "Misc/EnumClassFlags.h"
+
+class UClass;
+class UFunction;
+
+namespace UE::ProjectUtilities
+{
+	/** 
+	 * FBuildTargetSet is a value type that represents a set of build targets, e.g. Server|Client|Editor or Client|Server. 
+	 * The underlying target types are meant to be abstracted away, instead the set can be constructed from a piece of 
+	 * reflection data (e.g. a UFunction to call or UClass to use)
+	 */
+	struct FBuildTargetSet
+	{
+		/**
+		 * Returns a FBuildTargetSet representing the set of targets that the caller supports, that the callee does not.
+		 * 
+		 * @param Caller	A class that intends to call the Callee
+		 * @param Callee	The function being called
+		 * @return			FBuildTargetSet representing the unsupported targets that a caller supports that the 
+		 *					callee does not - useful for detecting a caller that may exist on some target (e.g. client) 
+		 *					using a function that will not exist on said target.
+		 */
+		static KISMET_API FBuildTargetSet GetCallerTargetsUnsupportedByCallee(const UClass* Caller, const UFunction* Callee);
+
+		/** @return			A String representation of this, useful for messaging to a user */
+		KISMET_API FString LexToString() const;
+
+		/** Consistent compiler generated default comparison operations */
+		friend auto operator<=>(const FBuildTargetSet&, const FBuildTargetSet&) = default;
+	private:
+		// This internal state may become more restrictive as we add validations - e.g.
+		// for platforms or specialty programs that go beyond the server/client/editor 
+		// paradigm:
+		enum class EBuildTargetFlags : uint32
+		{
+			None = 0,
+			Server = 1 << 0,
+			Client = 1 << 1,
+			Editor = 1 << 2,
+		};
+		FRIEND_ENUM_CLASS_FLAGS(EBuildTargetFlags)
+		static KISMET_API FString LexToStringImpl(EBuildTargetFlags Flags);
+		static EBuildTargetFlags GetSupportedTargetsForNativeClass(const UClass* NativeClass);
+		static EBuildTargetFlags GetCallerTargetsUnsupportedByCalleeImpl(EBuildTargetFlags CallerTargets, EBuildTargetFlags CalleeTargets);
+
+		EBuildTargetFlags BuildTargetFlags = EBuildTargetFlags::None;
+	};
+	ENUM_CLASS_FLAGS(FBuildTargetSet::EBuildTargetFlags);
+} 
diff --git a/Engine/Source/Editor/Kismet/Private/ProjectUtilities/BuildTargetSet.cpp b/Engine/Source/Editor/Kismet/Private/ProjectUtilities/BuildTargetSet.cpp
new file mode 100644
--- /dev/null
+++ b/Engine/Source/Editor/Kismet/Private/ProjectUtilities/BuildTargetSet.cpp
@@ -0,0 +1,250 @@
+// Copyright Epic Games, Inc. All Rights Reserved.
+
+#include "ProjectUtilities/BuildTargetSet.h"
+#include "Interfaces/IProjectManager.h"
+#include "Interfaces/IPluginManager.h"
+#include "ProjectDescriptor.h"
+#include "Kismet2/BlueprintEditorUtils.h"
+
+namespace UE::Private
+{
+	static TMap<FName, const FModuleDescriptor*> DescriptorCache;
+	static FCriticalSection DescriptorCacheLock;
+
+	const UClass* GetNativeClass(const UFunction* ForFunction);
+	const UClass* GetNativeClass(const UClass* ForClass);
+	const FModuleDescriptor* FindModuleDescriptor(const UPackage* ForNativePackage);
+}
+
+const UClass* UE::Private::GetNativeClass(const UFunction* ForFunction)
+{
+	const UClass* NativeClass = GetNativeClass(ForFunction->GetOwnerClass());
+	if (NativeClass)
+	{
+		return NativeClass;
+	}
+
+	ensureMsgf(false, TEXT("Found no native base class for function - cannot validate native constraints: %s"), *ForFunction->GetPathName());
+	return nullptr;
+}
+
+const UClass* UE::Private::GetNativeClass(const UClass* ForClass)
+{
+	return FBlueprintEditorUtils::FindFirstNativeClass(const_cast<UClass*>(ForClass));
+}
+
+const FModuleDescriptor* UE::Private::FindModuleDescriptor(const UPackage* ForNativePackage)
+{
+	// The loops at the bottom are naive and may need to be optimized if we begin
+	// validating build target compatibility widely:
+	FString FunctionPackageName = ForNativePackage->GetName();
+	if (!FunctionPackageName.StartsWith(TEXT("/Script/")))
+	{
+		return nullptr;
+	}
+
+	FunctionPackageName = FunctionPackageName.RightChop(8); // Chop "/Script/" from the name
+	if (FunctionPackageName == TEXT("Engine"))
+	{ // common case, there is no ModuleDescriptor for Engine
+		return nullptr;
+	}
+	
+	const FName FunctionPackageFName(FunctionPackageName.Len(), *FunctionPackageName);
+	const auto AddToCache = [](FName Name, const FModuleDescriptor* Descriptor)
+	{
+		// must be a write lock:
+		FScopeLock WriteLock(&DescriptorCacheLock);
+		DescriptorCache.Add(Name, Descriptor);
+	};
+	const auto FindFromCache = [](FName Name) -> TOptional<const FModuleDescriptor*>
+	{
+		// could be a read lock:
+		FScopeLock FindLock(&DescriptorCacheLock);
+		const FModuleDescriptor** Result = DescriptorCache.Find(Name);
+		return Result ? *Result : TOptional<const FModuleDescriptor*>();
+	};
+	TOptional<const FModuleDescriptor*> CachedDescriptor = FindFromCache(FunctionPackageFName);
+	if(CachedDescriptor)
+	{
+		return *CachedDescriptor;
+	}
+
+	const FModuleDescriptor* OwningMD = nullptr;
+	const IProjectManager& ProjectManager = IProjectManager::Get();
+	const FProjectDescriptor* PD = ProjectManager.GetCurrentProject();
+	for (const FModuleDescriptor& MD : PD->Modules)
+	{
+		if (MD.Name == FunctionPackageFName)
+		{
+			AddToCache(FunctionPackageFName, &MD);
+			return &MD;
+		}
+	}
+
+	// This is a little expensive, hence the cache - IPluginManager could take over
+	// acceleration structure duties:
+	TArray<TSharedRef<IPlugin>> AllPlugins = IPluginManager::Get().GetDiscoveredPlugins();
+	for (const TSharedRef<IPlugin>& Plugin : AllPlugins)
+	{
+		if (Plugin->IsHidden())
+		{
+			continue;
+		}
+
+		for (const FModuleDescriptor& MD : Plugin->GetDescriptor().Modules)
+		{
+			if (MD.Name == FunctionPackageFName)
+			{
+				AddToCache(FunctionPackageFName, &MD);
+				return &MD;
+			}
+		}
+	}
+	
+	AddToCache(FunctionPackageFName, nullptr);
+	return nullptr;
+}
+
+#include "Logging/StructuredLog.h"
+#include "ProfilingDebugging/ScopedTimers.h"
+
+namespace UE::ProjectUtilities
+{
+
+FBuildTargetSet FBuildTargetSet::GetCallerTargetsUnsupportedByCallee(const UClass* Caller, const UFunction* Callee)
+{
+	using namespace UE::Private;
+	check(Caller && Callee);
+	if (Caller == Callee->GetOwnerClass())
+	{
+		return FBuildTargetSet();
+	}
+	
+	const UClass* NativeCaller = GetNativeClass(Caller);
+	const UClass* NativeCallee = GetNativeClass(Callee);
+
+	// if the native modules are the same type then we are fine:
+	const UPackage* CallerNativePackage = NativeCaller->GetPackage();
+	const UPackage* CalleeNativePackage = NativeCallee->GetPackage();
+	if (CallerNativePackage == CalleeNativePackage)
+	{
+		return FBuildTargetSet();
+	}
+	
+	FBuildTargetSet CallerSet;
+	CallerSet.BuildTargetFlags = GetSupportedTargetsForNativeClass(NativeCaller);
+	FBuildTargetSet CalleeSet;
+	CalleeSet.BuildTargetFlags = GetSupportedTargetsForNativeClass(NativeCallee);
+	
+	FBuildTargetSet Result;
+	Result.BuildTargetFlags = GetCallerTargetsUnsupportedByCalleeImpl(CallerSet.BuildTargetFlags, CalleeSet.BuildTargetFlags);
+
+	return Result;
+}
+
+FString FBuildTargetSet::LexToString() const
+{
+	return LexToStringImpl(BuildTargetFlags);
+}
+
+FString FBuildTargetSet::LexToStringImpl(EBuildTargetFlags Flags)
+{
+	FString Result;
+	const TCHAR* Seperator = TEXT("");
+	if ((Flags & EBuildTargetFlags::Server) != EBuildTargetFlags::None)
+	{
+		Result += TEXT("Server");
+		Seperator = TEXT("|");
+	}
+	if ((Flags & EBuildTargetFlags::Client) != EBuildTargetFlags::None)
+	{
+		Result += Seperator;
+		Result += TEXT("Client");
+		Seperator = TEXT("|");
+	}
+	if ((Flags & EBuildTargetFlags::Editor) != EBuildTargetFlags::None)
+	{
+		Result += Seperator;
+		Result += TEXT("Editor");
+		Seperator = TEXT("|");
+	}
+	return Result;
+}
+
+FBuildTargetSet::EBuildTargetFlags FBuildTargetSet::GetSupportedTargetsForNativeClass(const UClass* NativeBase)
+{
+	check(NativeBase);
+	const EBuildTargetFlags ALL = EBuildTargetFlags::Server
+		| EBuildTargetFlags::Client
+		| EBuildTargetFlags::Editor;
+	using namespace UE::Private;
+	const FModuleDescriptor* ModuleDescriptor = FindModuleDescriptor(NativeBase->GetPackage());
+	EBuildTargetFlags SupportedTargets = ALL;
+	if (ModuleDescriptor)
+	{
+		switch (ModuleDescriptor->Type)
+		{
+		case EHostType::Runtime:
+		case EHostType::RuntimeNoCommandlet:
+		case EHostType::RuntimeAndProgram:
+		case EHostType::CookedOnly:
+			break;
+		case EHostType::UncookedOnly:
+		case EHostType::Developer:
+		case EHostType::DeveloperTool:
+		case EHostType::Editor:
+		case EHostType::EditorNoCommandlet:
+		case EHostType::EditorAndProgram:
+		case EHostType::Program:
+			SupportedTargets = EBuildTargetFlags::Editor;
+			break;
+		case EHostType::ServerOnly:
+			// Loads on all targets except dedicated clients.
+			SupportedTargets = EBuildTargetFlags::Server
+				| EBuildTargetFlags::Editor;
+			break;
+		case EHostType::ClientOnly:
+		case EHostType::ClientOnlyNoCommandlet:
+			// Loads on all targets except dedicated servers.
+			SupportedTargets = EBuildTargetFlags::Client
+				| EBuildTargetFlags::Editor;
+			break;
+		default:
+			ensureMsgf(false, TEXT("Encountered unexpected module type: %d"), ModuleDescriptor->Type);
+			break;
+		}
+	}
+
+	// honor IsEditorOnlyObject
+	if (IsEditorOnlyObject(NativeBase))
+	{
+		SupportedTargets = EBuildTargetFlags::Editor;
+	}
+
+	// if the class itself doesn't think it needs load for server
+	if (!NativeBase->NeedsLoadForServer() ||
+		!NativeBase->GetDefaultObject()->NeedsLoadForServer())
+	{
+		SupportedTargets &= ~EBuildTargetFlags::Server;
+	}
+	// or client, then we can't support those:
+	if (!NativeBase->NeedsLoadForClient()||
+		!NativeBase->GetDefaultObject()->NeedsLoadForClient())
+	{
+		SupportedTargets &= ~EBuildTargetFlags::Client;
+	}
+	return SupportedTargets;
+}
+
+FBuildTargetSet::EBuildTargetFlags FBuildTargetSet::GetCallerTargetsUnsupportedByCalleeImpl(FBuildTargetSet::EBuildTargetFlags CallerTargets, FBuildTargetSet::EBuildTargetFlags CalleeTargets)
+{
+	using namespace UE::Private;
+	// caller targets must be a subset of valid callee targets,
+	// otherwise caller could be deployed into a context within which
+	// the callee is not available:
+	EBuildTargetFlags Differences = CallerTargets ^ CalleeTargets;
+	return (Differences & ~CalleeTargets);	
+}
+
+}
+
diff --git a/Engine/Source/Editor/KismetCompiler/Private/KismetCompilerVMBackend.cpp b/Engine/Source/Editor/KismetCompiler/Private/KismetCompilerVMBackend.cpp
--- a/Engine/Source/Editor/KismetCompiler/Private/KismetCompilerVMBackend.cpp
+++ b/Engine/Source/Editor/KismetCompiler/Private/KismetCompilerVMBackend.cpp
@@ -32,6 +32,7 @@
 #include "KismetCompilerBackend.h"
 
 #include "Misc/DefaultValueHelper.h"
+#include "ProjectUtilities/BuildTargetSet.h"
 
 #include "Kismet2/StructureEditorUtils.h"
 #include "Kismet2/KismetDebugUtilities.h"
@@ -1304,6 +1305,7 @@ public:
 		}
 
 		const bool bFinalFunction = FunctionToCall->HasAnyFunctionFlags(FUNC_Final) || Statement.bIsParentContext;
+		using namespace UE::ProjectUtilities;
 		const bool bMathCall = bFinalFunction
 			&& FunctionToCall->HasAllFunctionFlags(FUNC_Static|FUNC_Final|FUNC_Native)
 			&& !FunctionToCall->HasAnyFunctionFlags(FUNC_NetFuncFlags|FUNC_BlueprintAuthorityOnly|FUNC_BlueprintCosmetic|FUNC_NetRequest|FUNC_NetResponse)
@@ -1312,7 +1314,9 @@ public:
 			// is doing things with reflection data at runtime and will therefore benefit from
 			// the safety of a 'context' based function call (e.g. if context is invalid or
 			// 'bArrayContextFailed' we can skip the function call):
-			&& !UEdGraphSchema_K2::HasWildcardParams(FunctionToCall);
+			&& !UEdGraphSchema_K2::HasWildcardParams(FunctionToCall)
+			// Don't emit math calls if the target function may disapear on us at loadtime/runtime:
+			&& FBuildTargetSet::GetCallerTargetsUnsupportedByCallee(ClassBeingBuilt, FunctionToCall) == FBuildTargetSet();
 
 		const bool bLocalScriptFunction = 
 			!FunctionToCall->HasAnyFunctionFlags(FUNC_Native|FUNC_NetFuncFlags|FUNC_BlueprintAuthorityOnly|FUNC_BlueprintCosmetic|FUNC_NetRequest|FUNC_NetResponse);
