Avoid searching for modules when no ProjectDescriptor is available - e.g. when running engine level commandlets
#rnx

diff --git a/Engine/Source/Editor/Kismet/Private/ProjectUtilities/BuildTargetSet.cpp b/Engine/Source/Editor/Kismet/Private/ProjectUtilities/BuildTargetSet.cpp
--- a/Engine/Source/Editor/Kismet/Private/ProjectUtilities/BuildTargetSet.cpp
+++ b/Engine/Source/Editor/Kismet/Private/ProjectUtilities/BuildTargetSet.cpp
@@ -48,6 +48,13 @@ const FModuleDescriptor* UE::Private::FindModuleDescriptor(const UPackage* ForNa
 	{ // common case, there is no ModuleDescriptor for Engine
 		return nullptr;
 	}
+
+	const IProjectManager& ProjectManager = IProjectManager::Get();
+	const FProjectDescriptor* PD = ProjectManager.GetCurrentProject();
+	if(!PD)
+	{
+		return nullptr;
+	}
 	
 	const FName FunctionPackageFName(FunctionPackageName.Len(), *FunctionPackageName);
 	const auto AddToCache = [](FName Name, const FModuleDescriptor* Descriptor)
@@ -70,8 +77,7 @@ const FModuleDescriptor* UE::Private::FindModuleDescriptor(const UPackage* ForNa
 	}
 
 	const FModuleDescriptor* OwningMD = nullptr;
-	const IProjectManager& ProjectManager = IProjectManager::Get();
-	const FProjectDescriptor* PD = ProjectManager.GetCurrentProject();
+	
 	for (const FModuleDescriptor& MD : PD->Modules)
 	{
 		if (MD.Name == FunctionPackageFName)
