UE-217968 - Loading blueprints with soft actor references during PIE corrupts data

When using the legacy editor loader (ie: FLinkerLoad), we don't correctly set GPlayInEditorID to -1. By not doing this, resolving any soft paths could end up using PIE actors instead of the ones that are actually in the level (see FSoftObjectPath::ResolveObject). Affected Blueprints would have corrupted data if they were saved after this occurred.

This doesn't occur in Zen loader because it already uses FPlayInEditorLoadingScope during export creation and a few other loading states. While redundant in this case, the extra FPlayInEditorLoadingScope won't have any negative effects.

The issue has been around for a while, and only manifested in 5.4 due to some completely unrelated changes to ReinstBatch (27593863 and 32291617).

diff --git a/Engine/Source/Editor/Kismet/Private/BlueprintCompilationManager.cpp b/Engine/Source/Editor/Kismet/Private/BlueprintCompilationManager.cpp
--- a/Engine/Source/Editor/Kismet/Private/BlueprintCompilationManager.cpp
+++ b/Engine/Source/Editor/Kismet/Private/BlueprintCompilationManager.cpp
@@ -33,6 +33,7 @@
 #include "Misc/ScopedSlowTask.h"
 #include "Misc/DataValidation.h"
 #include "Misc/PackageAccessTrackingOps.h"
+#include "Misc/PlayInEditorLoadingScope.h"
 #include "ProfilingDebugging/ScopedTimers.h"
 #include "Serialization/ArchiveHasReferences.h"
 #include "Serialization/ArchiveReplaceObjectRef.h"
@@ -2274,6 +2275,10 @@ void FBlueprintCompilationManagerImpl::ReinstanceBatch(TArray<FReinstancingJob>&
 {
 	TGuardValue<bool> ReinstancingGuard(GIsReinstancing, true);
 
+	// This is only needed when using the legacy editor loader (ie: FLinkerLoad)
+	// Zen loader already uses FPlayInEditorLoadingScope
+	UE::Core::Private::FPlayInEditorLoadingScope PlayInEditorIDScope(INDEX_NONE);
+
 	const auto FilterOutOfDateClasses = [](TArray<UClass*>& ClassList)
 	{
 		// Old versions of classes can be abandoned, classes without CDOs have no instances and don't require reinstancing
