[Feature] Upgraded the detail panel of blueprint graph nodes to use visual diff and synced scrolling

diff --git a/Engine/Source/Editor/Kismet/Private/SBlueprintDiff.cpp b/Engine/Source/Editor/Kismet/Private/SBlueprintDiff.cpp
--- a/Engine/Source/Editor/Kismet/Private/SBlueprintDiff.cpp
+++ b/Engine/Source/Editor/Kismet/Private/SBlueprintDiff.cpp
@@ -58,6 +58,7 @@
 #include "Widgets/SWindow.h"
 #include "Widgets/Text/STextBlock.h"
 #include "Widgets/Views/SListView.h"
+#include "AsyncDetailViewDiff.h"
 
 class FProperty;
 
@@ -398,6 +399,18 @@ void SBlueprintDiff::Tick(const FGeometry& AllottedGeometry, const double InCurr
 	{
 		DiffControl->Tick();
 	}
+	if (PanelOld.DetailsView)
+	{
+		PanelOld.DetailsView->Tick(AllottedGeometry, InCurrentTime, InDeltaTime);
+	}
+	if (PanelNew.DetailsView)
+	{
+		PanelNew.DetailsView->Tick(AllottedGeometry, InCurrentTime, InDeltaTime);
+	}
+	if (GraphDetailDiff)
+	{
+		GraphDetailDiff->Tick();
+	}
 }
 
 void SBlueprintDiff::OnCloseAssetEditor(UObject* Asset, EAssetEditorCloseReason CloseReason)
@@ -850,31 +863,46 @@ void SBlueprintDiff::GenerateDifferencesList()
 
 	// SMyBlueprint needs to be created *before* the KismetInspector or the diffs are generated, because the KismetInspector's customizations
 	// need a reference to the SMyBlueprint widget that is controlling them...
-	const auto CreateInspector = [](TSharedPtr<SMyBlueprint> InMyBlueprint) {
+	const auto CreateInspector = [](TSharedPtr<SMyBlueprint> InMyBlueprint, const TSharedPtr<SScrollBar>& Scrollbar, EHorizontalAlignment ScrollbarAlignment) {
 		return SNew(SKismetInspector)
 			.HideNameArea(true)
 			.ViewIdentifier(FName("BlueprintInspector"))
 			.MyBlueprintWidget(InMyBlueprint)
 			.IsPropertyEditingEnabledDelegate(FIsPropertyEditingEnabled::CreateStatic([] { return false; }))
-			.ShowLocalVariables(true);
+			.ShowLocalVariables(true)
+			.ExternalScrollbar(Scrollbar)
+			.ScrollbarAlignment(ScrollbarAlignment);
 	};
 
 	TArray<UEdGraph*> GraphsOld, GraphsNew;
 	if (PanelOld.Blueprint)
 	{
 		PanelOld.GenerateMyBlueprintWidget();
-		PanelOld.DetailsView = CreateInspector(PanelOld.MyBlueprint);
+		PanelOld.DetailScrollbar = SNew(SLinkableScrollBar);
+		PanelOld.DetailsView = CreateInspector(PanelOld.MyBlueprint, PanelOld.DetailScrollbar, HAlign_Left);
 		PanelOld.MyBlueprint->SetInspector(PanelOld.DetailsView);
 		PanelOld.Blueprint->GetAllGraphs(GraphsOld);
 	}
 	if (PanelNew.Blueprint)
 	{
 		PanelNew.GenerateMyBlueprintWidget();
-		PanelNew.DetailsView = CreateInspector(PanelNew.MyBlueprint);
+		PanelNew.DetailScrollbar = SNew(SLinkableScrollBar);
+		PanelNew.DetailsView = CreateInspector(PanelNew.MyBlueprint, PanelNew.DetailScrollbar, HAlign_Right);
 		PanelNew.MyBlueprint->SetInspector(PanelNew.DetailsView);
 		PanelNew.Blueprint->GetAllGraphs(GraphsNew);
 	}
 
+	if (PanelOld.DetailsView && PanelNew.DetailsView)
+	{
+		GraphDetailDiff = MakeShared<FAsyncDetailViewDiff>(
+			PanelOld.DetailsView->GetPropertyView().ToSharedRef(),
+			PanelNew.DetailsView->GetPropertyView().ToSharedRef());
+
+		SLinkableScrollBar::LinkScrollBars(PanelOld.DetailScrollbar.ToSharedRef(), PanelNew.DetailScrollbar.ToSharedRef(),
+		TAttribute<TArray<FVector2f>>::CreateRaw(GraphDetailDiff.Get(), &FAsyncDetailViewDiff::GenerateScrollSyncRate));
+		TAttribute<TArray<FVector2f>> ScrollRate = TAttribute<TArray<FVector2f>>::CreateRaw(GraphDetailDiff.Get(), &FAsyncDetailViewDiff::GenerateScrollSyncRate);
+	}
+
 	//Add Graphs that exist in both blueprints, or in blueprint 1 only
 	for (UEdGraph* GraphOld : GraphsOld)
 	{
@@ -1091,16 +1119,14 @@ SBlueprintDiff::FDiffControl SBlueprintDiff::GenerateGraphPanel()
 			+ SSplitter::Slot()
 			.Value(.2f)
 			[
-				SNew(SSplitter)
-				.PhysicalSplitterHandleSize(10.0f)
-				+SSplitter::Slot()
-				[
-					PanelOld.GetDetailsWidget()
-				]
-				+ SSplitter::Slot()
-				[
-					PanelNew.GetDetailsWidget()
-				]
+				SNew(SDetailsSplitter)
+				+SDetailsSplitter::Slot()
+					.Value(0.5f)
+					.DetailsView(PanelOld.DetailsView ? PanelOld.DetailsView->GetPropertyView() : nullptr)
+					.DifferencesWithRightPanel(GraphDetailDiff)
+				+ SDetailsSplitter::Slot()
+					.Value(0.5f)
+					.DetailsView(PanelNew.DetailsView ? PanelNew.DetailsView->GetPropertyView() : nullptr)
 			]
 		]
 	];
diff --git a/Engine/Source/Editor/Kismet/Private/SKismetInspector.cpp b/Engine/Source/Editor/Kismet/Private/SKismetInspector.cpp
--- a/Engine/Source/Editor/Kismet/Private/SKismetInspector.cpp
+++ b/Engine/Source/Editor/Kismet/Private/SKismetInspector.cpp
@@ -403,6 +403,8 @@ void SKismetInspector::Construct(const FArguments& InArgs)
 	DetailsViewArgs.bHideSelectionTip = true;
 	DetailsViewArgs.NotifyHook = NotifyHook;
 	DetailsViewArgs.ViewIdentifier = InArgs._ViewIdentifier;
+	DetailsViewArgs.ExternalScrollbar = InArgs._ExternalScrollbar;
+	DetailsViewArgs.ScrollbarAlignment = InArgs._ScrollbarAlignment;
 	if (Kismet2.IsValid())
 	{
 		TSharedPtr<IClassViewerFilter> ImportedClassFilter = Kismet2->GetImportedClassViewerFilter();
diff --git a/Engine/Source/Editor/Kismet/Public/SBlueprintDiff.h b/Engine/Source/Editor/Kismet/Public/SBlueprintDiff.h
--- a/Engine/Source/Editor/Kismet/Public/SBlueprintDiff.h
+++ b/Engine/Source/Editor/Kismet/Public/SBlueprintDiff.h
@@ -28,6 +28,8 @@
 #include "Widgets/Views/STableViewBase.h"
 #include "Widgets/Views/STreeView.h"
 
+class FAsyncDetailViewDiff;
+class SLinkableScrollBar;
 class FBlueprintDifferenceTreeEntry;
 class FSpawnTabArgs;
 class FTabManager;
@@ -126,6 +128,8 @@ struct KISMET_API FDiffPanel
 
 	/** The widget that contains the revision info in graph mode */
 	TSharedPtr<SWidget>				OverlayGraphRevisionInfo;
+
+	TSharedPtr<SLinkableScrollBar> DetailScrollbar;
 private:
 	/** Command list for this diff panel */
 	TSharedPtr<FUICommandList> GraphEditorCommands;
@@ -262,6 +266,9 @@ protected:
 
 	/*The two panels used to show the old & new revision*/ 
 	FDiffPanel				PanelOld, PanelNew;
+
+	/* Diff info about the kismet panels in blueprint graph diff view */
+	TSharedPtr<FAsyncDetailViewDiff> GraphDetailDiff;
 	
 	/** If the two views should be locked */
 	bool	bLockViews;
diff --git a/Engine/Source/Editor/Kismet/Public/SKismetInspector.h b/Engine/Source/Editor/Kismet/Public/SKismetInspector.h
--- a/Engine/Source/Editor/Kismet/Public/SKismetInspector.h
+++ b/Engine/Source/Editor/Kismet/Public/SKismetInspector.h
@@ -31,6 +31,7 @@ class SBorder;
 class SDockTab;
 class SMyBlueprint;
 class SWidget;
+class SScrollBar;
 class UObject;
 struct FGeometry;
 struct FPropertyChangedEvent;
@@ -50,6 +51,7 @@ public:
 		, _SetNotifyHook(true)
 		, _ShowTitleArea(false)
 		, _ShowLocalVariables(false)
+		, _ScrollbarAlignment(HAlign_Right)
 		{}
 
 		SLATE_ARGUMENT(TWeakPtr<FBlueprintEditor>, Kismet2)
@@ -62,6 +64,8 @@ public:
 		SLATE_ARGUMENT( bool, SetNotifyHook)
 		SLATE_ARGUMENT( bool, ShowTitleArea)
 		SLATE_ARGUMENT( bool, ShowLocalVariables)
+		SLATE_ARGUMENT( TSharedPtr<SScrollBar>, ExternalScrollbar)
+		SLATE_ARGUMENT( EHorizontalAlignment, ScrollbarAlignment)
 	SLATE_END_ARGS()
 	
 	void Construct(const FArguments& InArgs);
diff --git a/Engine/Source/Editor/PropertyEditor/Private/SDetailsSplitter.cpp b/Engine/Source/Editor/PropertyEditor/Private/SDetailsSplitter.cpp
--- a/Engine/Source/Editor/PropertyEditor/Private/SDetailsSplitter.cpp
+++ b/Engine/Source/Editor/PropertyEditor/Private/SDetailsSplitter.cpp
@@ -534,7 +534,7 @@ int32 SDetailsSplitter::OnPaint(const FPaintArgs& Args, const FGeometry& Allotte
 	for (int32 LeftIndex = 0; LeftIndex < Panels.Num(); ++LeftIndex)
 	{
 		const FPanel& LeftPanel = Panels[LeftIndex];
-		if (!LeftPanel.DiffRight.IsBound())
+		if (!LeftPanel.DiffRight.IsSet())
 		{
 			continue;
 		}
@@ -668,7 +668,7 @@ FReply SDetailsSplitter::OnMouseMove(const FGeometry& MyGeometry, const FPointer
 	for (int32 LeftIndex = 0; LeftIndex < Panels.Num() - 1; ++LeftIndex)
 	{
 		const FPanel& LeftPanel = Panels[LeftIndex];
-		if (!LeftPanel.DiffRight.IsBound())
+		if (!LeftPanel.DiffRight.IsSet())
 		{
 			continue;
 		}
