Fix for garbaged archetypes being reinstanciated

diff --git a/Engine/Source/Editor/Kismet/Private/BlueprintCompilationManager.cpp b/Engine/Source/Editor/Kismet/Private/BlueprintCompilationManager.cpp
--- a/Engine/Source/Editor/Kismet/Private/BlueprintCompilationManager.cpp
+++ b/Engine/Source/Editor/Kismet/Private/BlueprintCompilationManager.cpp
@@ -2673,7 +2673,7 @@ void FBlueprintCompilationManagerImpl::ReinstanceBatch(TArray<FReinstancingJob>&
 			}
 
 			TArray<UObject*> ArchetypeObjects;
-			GetObjectsOfClass(OldClass, ArchetypeObjects, false);
+			GetObjectsOfClass(OldClass, ArchetypeObjects, /*bIncludeDerivedClasses*/false, RF_NoFlags, EInternalObjectFlags::Garbage);
 			
 			// filter out non-archetype instances, note that WidgetTrees and some component
 			// archetypes do not have RF_ArchetypeObject or RF_InheritableComponentTemplate so
