Fix find by references by name for event graph nodes, regressed via 30854601

[RN] Searching for references to an event function via the My Blueprints context menu no longer inadvertantly double quotes the event name

diff --git a/Engine/Source/Editor/Kismet/Private/SMyBlueprint.cpp b/Engine/Source/Editor/Kismet/Private/SMyBlueprint.cpp
--- a/Engine/Source/Editor/Kismet/Private/SMyBlueprint.cpp
+++ b/Engine/Source/Editor/Kismet/Private/SMyBlueprint.cpp
@@ -2991,6 +2991,7 @@ void SMyBlueprint::OnFindReference(bool bSearchAllBlueprints, const EGetFindRefe
 		if (!bSearchTermGenerated)
 		{
 			SearchTerm = EventAction->NodeTemplate->GetFindReferenceSearchString(EGetFindReferenceSearchStringFlags::None);
+			bUseQuotes = false;
 		}
 	}
 	else if (FEdGraphSchemaAction_K2InputAction* InputAction = SelectionAsInputAction())
