Add new viewport toolbar to the Blueprint editor (SCSEditor)

Currently has the Transforms, Select, and View Modes (with too many entries right now) menus.

diff --git a/Engine/Source/Editor/Kismet/Private/SSCSEditorViewport.cpp b/Engine/Source/Editor/Kismet/Private/SSCSEditorViewport.cpp
--- a/Engine/Source/Editor/Kismet/Private/SSCSEditorViewport.cpp
+++ b/Engine/Source/Editor/Kismet/Private/SSCSEditorViewport.cpp
@@ -38,10 +38,12 @@
 #include "SlotBase.h"
 #include "Styling/AppStyle.h"
 #include "Styling/ISlateStyle.h"
+#include "ToolMenus.h"
 #include "Types/WidgetActiveTimerDelegate.h"
 #include "UObject/NameTypes.h"
 #include "UObject/UObjectGlobals.h"
 #include "UObject/UnrealNames.h"
+#include "ViewportToolbar/UnrealEdViewportToolbar.h"
 #include "Widgets/Layout/SBorder.h"
 #include "Widgets/SBoxPanel.h"
 #include "Widgets/SCompoundWidget.h"
@@ -300,10 +302,161 @@ TSharedRef<FEditorViewportClient> SSCSEditorViewport::MakeEditorViewportClient()
 
 TSharedPtr<SWidget> SSCSEditorViewport::MakeViewportToolbar()
 {
-	return 
+	TSharedRef<SSCSEditorViewportToolBar> OldViewportToolbar =
 		SNew(SSCSEditorViewportToolBar)
-		.EditorViewport(SharedThis(this))
-		.IsEnabled(FSlateApplication::Get().GetNormalExecutionAttribute());
+			.EditorViewport(SharedThis(this))
+			.IsEnabled(FSlateApplication::Get().GetNormalExecutionAttribute())
+			.Visibility_Lambda(
+				[]()
+				{
+					return UE::UnrealEd::ShowOldViewportToolbars() ? EVisibility::Visible : EVisibility::Collapsed;
+				}
+			);
+
+	// Register the viewport toolbar if another viewport hasn't already (it's shared).
+	const FName ViewportToolbarMenuName = "SCSEditor.ViewportToolbar";
+	if (!UToolMenus::Get()->IsMenuRegistered(ViewportToolbarMenuName))
+	{
+		UToolMenu* const ViewportToolbarMenu = UToolMenus::Get()->RegisterMenu(
+			ViewportToolbarMenuName, NAME_None /* parent */, EMultiBoxType::SlimHorizontalToolBar
+		);
+
+		// Add the left-aligned part of the viewport toolbar.
+		{
+			FToolMenuSection& LeftSection = ViewportToolbarMenu->FindOrAddSection("Left");
+
+			// Add the "Transforms" sub menu.
+			{
+				FToolMenuEntry TransformsSubmenu = UE::UnrealEd::CreateViewportToolbarTransformsSection();
+				TransformsSubmenu.InsertPosition.Position = EToolMenuInsertType::First;
+				LeftSection.AddEntry(TransformsSubmenu);
+			}
+
+			// Add the "Selection" sub menu.
+			{
+				FToolMenuEntry SelectionSubmenu = UE::UnrealEd::CreateViewportToolbarSelectionSection();
+				SelectionSubmenu.InsertPosition.Position = EToolMenuInsertType::First;
+				LeftSection.AddEntry(SelectionSubmenu);
+			}
+
+			// TODO: Add this.
+			// // Add the "Snapping" sub menu.
+			// {
+			// 	FToolMenuEntry SnappingSubmenu = UE::LevelEditor::CreateViewportToolbarSnappingSubmenu();
+			// 	SnappingSubmenu.InsertPosition.Position = EToolMenuInsertType::First;
+			// 	LeftSection.AddEntry(SnappingSubmenu);
+			// }
+		}
+
+		// Add the right-aligned part of the viewport toolbar.
+		{
+			// Add the submenus of this section as EToolMenuInsertType::Last to sort them after any
+			// default-positioned submenus external code might add.
+			FToolMenuSection& RightSection = ViewportToolbarMenu->FindOrAddSection("Right");
+			RightSection.Alignment = EToolMenuSectionAlign::Last;
+
+			// TODO: Add Camera menu here.
+
+			// TODO: Filter this menu with IsViewModeSupportedDelegate (see further down in this file) and remove the
+			// "Exposure" section.
+
+			// Add the "View Modes" sub menu.
+			{
+				// Stay backward-compatible with the old viewport toolbar.
+				{
+					const FName ParentSubmenuName = "UnrealEd.ViewportToolbar.View";
+					// Create our parent menu.
+					if (!UToolMenus::Get()->IsMenuRegistered(ParentSubmenuName))
+					{
+						UToolMenus::Get()->RegisterMenu(ParentSubmenuName);
+					}
+
+					// Register our ToolMenu here first, before we create the submenu, so we can set our parent.
+					UToolMenus::Get()->RegisterMenu("SCSEditor.ViewportToolbar.ViewModes", ParentSubmenuName);
+				}
+
+				FToolMenuEntry ViewModesSubmenu = UE::UnrealEd::CreateViewportToolbarViewModesSubmenu();
+				ViewModesSubmenu.InsertPosition.Position = EToolMenuInsertType::Last;
+				RightSection.AddEntry(ViewModesSubmenu);
+			}
+
+			// TODO: Add this.
+			// // Add the "Show" submenu.
+			// {
+			// 	FToolMenuEntry ShowSubmenu = UE::UnrealEd::CreateViewportToolbarShowSubmenu();
+			// 	ShowSubmenu.InsertPosition.Position = EToolMenuInsertType::Last;
+			// 	RightSection.AddEntry(ShowSubmenu);
+			// }
+
+			// TODO: Add this.
+			// // Add the "Performance & Scalability" submenu.
+			// {
+			// 	FToolMenuEntry PerfSubmenu = UE::LevelEditor::CreateViewportToolbarPerformanceAndScalabilitySubmenu();
+			// 	PerfSubmenu.InsertPosition.Position = EToolMenuInsertType::Last;
+			// 	RightSection.AddEntry(PerfSubmenu);
+			// }
+		}
+	}
+
+	FToolMenuContext ViewportToolbarContext;
+	{
+		ViewportToolbarContext.AppendCommandList(GetCommandList());
+
+		// Add the UnrealEd viewport toolbar context.
+		{
+			UUnrealEdViewportToolbarContext* const ContextObject = NewObject<UUnrealEdViewportToolbarContext>();
+			ContextObject->Viewport = SharedThis(this);
+
+			// Hook up our toolbar's filter for supported view modes.
+			UE::UnrealEd::IsViewModeSupportedDelegate IsViewModeSupported =
+				UE::UnrealEd::IsViewModeSupportedDelegate::CreateLambda(
+					[WeakToolBar = OldViewportToolbar.ToWeakPtr()](EViewModeIndex ViewModeIndex) -> bool
+					{
+						if (TSharedPtr<SViewportToolBar> ToolBar = WeakToolBar.Pin())
+						{
+							return ToolBar->IsViewModeSupported(ViewModeIndex);
+						}
+						return true;
+					}
+				);
+
+			ViewportToolbarContext.AddObject(ContextObject);
+		}
+	}
+
+	// clang-format off
+	TSharedRef<SWidget> NewViewportToolbar =
+		SNew(SBox)
+		[
+			UToolMenus::Get()->GenerateWidget(ViewportToolbarMenuName, ViewportToolbarContext)
+		]
+		.Visibility_Lambda(
+			[]()
+			{
+				return UE::UnrealEd::ShowNewViewportToolbars() ? EVisibility::Visible : EVisibility::Collapsed;
+			}
+		);
+	// clang-format on
+
+	// clang-format off
+	return 
+		SNew(SVerticalBox)
+		.Visibility( EVisibility::SelfHitTestInvisible )
+		+SVerticalBox::Slot()
+		.AutoHeight()
+		.Padding(0, 1.0f, 0, 0)
+		.VAlign(VAlign_Top)
+		[
+			OldViewportToolbar
+		]
+		+SVerticalBox::Slot()
+		.AutoHeight()
+		.Padding(0, 1.0f, 0, 0)
+		.VAlign(VAlign_Top)
+		[
+			NewViewportToolbar
+		];
+	// clang-format on
 }
 
 void SSCSEditorViewport::PopulateViewportOverlays(TSharedRef<class SOverlay> Overlay)
