Fix potential for deadlock when terminating a global Find-in-Blueprint query in projects that may index references to StringTable assets.

diff --git a/Engine/Source/Editor/Kismet/Private/FindInBlueprintManager.cpp b/Engine/Source/Editor/Kismet/Private/FindInBlueprintManager.cpp
--- a/Engine/Source/Editor/Kismet/Private/FindInBlueprintManager.cpp
+++ b/Engine/Source/Editor/Kismet/Private/FindInBlueprintManager.cpp
@@ -228,7 +228,21 @@ void FStreamSearch::EnsureCompletion()
 		ItemsFound.Empty();
 	}
 
+	// Signal to the search thread to stop its operation.
 	Stop();
+
+	// Wait until the search thread has signaled its completion.
+	while (!IsComplete())
+	{
+		// Async tasks may have been registered to this thread (e.g. FSearchableValueInfo::GetDisplayText),
+		// so make sure we process those to unblock the stream search thread. Otherwise we can deadlock below.
+		FTaskGraphInterface::Get().ProcessThreadUntilIdle(FTaskGraphInterface::Get().GetCurrentThreadIfKnown());
+
+		// Yield time to other threads, including the search thread.
+		FPlatformProcess::Sleep(0.1f);
+	}
+
+	// Wait for the search thread to terminate.
 	Thread->WaitForCompletion();
 }
 
