Improve error message - we've heard reports of this firing but not in a coherent manner.

#rnx

diff --git a/Engine/Source/Editor/Kismet/Private/ProjectUtilities/BuildTargetSet.cpp b/Engine/Source/Editor/Kismet/Private/ProjectUtilities/BuildTargetSet.cpp
--- a/Engine/Source/Editor/Kismet/Private/ProjectUtilities/BuildTargetSet.cpp
+++ b/Engine/Source/Editor/Kismet/Private/ProjectUtilities/BuildTargetSet.cpp
@@ -216,7 +216,7 @@ FBuildTargetSet::EBuildTargetFlags FBuildTargetSet::GetSupportedTargetsForNative
 				| EBuildTargetFlags::Editor;
 			break;
 		default:
-			ensureMsgf(false, TEXT("Encountered unexpected module type: %d"), ModuleDescriptor->Type);
+			ensureMsgf(false, TEXT("Encountered unexpected module type: %d in module %s"), ModuleDescriptor->Type, *ModuleDescriptor->Name.ToString());
 			break;
 		}
 	}
