List View: Add an OnItemsRebuilt delegate and a GetGeneratedChildAt method so that it's possible to do some post-processing on the generated list

diff --git a/Engine/Source/Editor/Kismet/Private/BlueprintDetailsCustomization.cpp b/Engine/Source/Editor/Kismet/Private/BlueprintDetailsCustomization.cpp
--- a/Engine/Source/Editor/Kismet/Private/BlueprintDetailsCustomization.cpp
+++ b/Engine/Source/Editor/Kismet/Private/BlueprintDetailsCustomization.cpp
@@ -111,6 +111,7 @@
 #include "Serialization/Archive.h"
 #include "SlateOptMacros.h"
 #include "SlotBase.h"
+#include "SSearchableComboBox.h"
 #include "SSocketChooser.h"
 #include "Styling/AppStyle.h"
 #include "Styling/ISlateStyle.h"
@@ -179,6 +180,8 @@ namespace BlueprintDocumentationDetailDefs
 	static const float DetailsTitleMaxWidth = 300.f;
 	/** magic number retrieved from SGraphNodeComment::GetWrapAt() */
 	static const float DetailsTitleWrapPadding = 32.0f;
+	/** label of the option to reset the drop-down value function name */
+	static const FString EmptyDropDownOptionName = TEXT("None");
 };
 
 void FBlueprintDetails::AddEventsCategory(IDetailLayoutBuilder& DetailBuilder, FName PropertyName, UClass* PropertyClass)
@@ -1275,6 +1278,53 @@ void FBlueprintVarActionDetails::CustomizeDetails( IDetailLayoutBuilder& DetailL
 				.Font(IDetailLayoutBuilder::GetDetailFont())
 			];
 
+		CollectDropDownOptions();
+
+		TSharedPtr<SToolTip> GetOptionsTooltip = IDocumentation::Get()->CreateToolTip(
+			LOCTEXT(
+				"VariableGetOptions_Tooltip",
+				"The name of the function which will populate a list of options the user can select from for the value of this variable."
+			),
+			nullptr, DocLink, TEXT("GetOptions")
+		);
+
+		TSharedPtr<SToolTip> GetOptionsInputTooltip = IDocumentation::Get()->CreateToolTip(
+			LOCTEXT(
+				"VariableGetOptions_InputTooltip",
+				"List of functions that return an array of names or strings.\n"
+				"You can also enter the name of a global static function, in the form of 'ClassName.FunctionName'."
+			),
+			nullptr, DocLink, TEXT("GetOptions")
+		);
+
+		Category.AddCustomRow(LOCTEXT("VariableGetOptions", "Drop-down Options"), true)
+		.Visibility(TAttribute<EVisibility>(this, &FBlueprintVarActionDetails::GetDropDownOptionsVisibility))
+		.NameContent()
+		[
+			SNew(STextBlock)
+			.ToolTip(GetOptionsTooltip)
+			.Text(LOCTEXT("VariableGetOptions", "Drop-down Options"))
+			.Font(DetailFontInfo)
+		]
+		.ValueContent()
+		[
+			SNew(SSearchableComboBox)
+			.ToolTip(GetOptionsInputTooltip)
+			.OptionsSource(&DropDownFunctionOptions)
+			.OnGenerateWidget(this, &FBlueprintVarActionDetails::GenerateDropDownOptionWidget)
+			.OnSelectionChanged(this, &FBlueprintVarActionDetails::OnDropDownOptionSelectionChanged)
+			.Content()
+			[
+				SNew(SEditableTextBox)
+				.ToolTip(GetOptionsInputTooltip)
+				.SelectAllTextWhenFocused(true)
+				.RevertTextOnEscape(true)
+				.Font(DetailFontInfo)
+				.Text(this, &FBlueprintVarActionDetails::GetDropDownOptionDisplayText)
+				.OnTextCommitted(this, &FBlueprintVarActionDetails::OnDropDownOptionTextChanged)
+			]
+		];
+
 		TSharedPtr<SToolTip> PropertyFlagsTooltip = IDocumentation::Get()->CreateToolTip(LOCTEXT("DefinedPropertyFlags_Tooltip", "List of defined flags for this property"), NULL, DocLink, TEXT("PropertyFlags"));
 
 		Category.AddCustomRow(LOCTEXT("DefinedPropertyFlags", "Defined Property Flags"), true)
@@ -2986,6 +3036,169 @@ void FBlueprintVarActionDetails::OnDeprecatedChanged(ECheckBoxState InNewState)
 	}
 }
 
+EVisibility FBlueprintVarActionDetails::GetDropDownOptionsVisibility() const
+{
+	if (FProperty* VariableProperty = CachedVariableProperty.Get())
+	{
+		const bool bMatchingType = VariableProperty->IsA(FNameProperty::StaticClass())
+			|| VariableProperty->IsA(FStrProperty::StaticClass());
+
+		if (bMatchingType && IsABlueprintVariable(VariableProperty) && IsAUserVariable(VariableProperty))
+		{
+			return EVisibility::Visible;
+		}
+	}
+
+	return EVisibility::Collapsed;
+}
+
+void FBlueprintVarActionDetails::OnDropDownOptionSelectionChanged(TSharedPtr<FString> InString,
+	ESelectInfo::Type)
+{
+	SetDropDownOptionsFunctionName(*InString);
+}
+
+void FBlueprintVarActionDetails::OnDropDownOptionTextChanged(const FText& Text, ETextCommit::Type)
+{
+	SetDropDownOptionsFunctionName(Text.ToString());
+}
+
+TSharedRef<SWidget> FBlueprintVarActionDetails::GenerateDropDownOptionWidget(TSharedPtr<FString> InItem) const
+{
+	return SNew(STextBlock)
+		.Text(FText::FromString(*InItem.Get()))
+		.Font(IDetailLayoutBuilder::GetDetailFont());
+}
+
+void FBlueprintVarActionDetails::CollectDropDownOptions()
+{
+	DropDownFunctionOptions.Empty();
+
+	const FProperty* VariableProperty = CachedVariableProperty.Get();
+	if (!VariableProperty)
+	{
+		return;
+	}
+
+	for (TFieldIterator<UFunction> It(VariableProperty->GetOwner<UClass>(), EFieldIteratorFlags::IncludeSuper); It; ++It)
+	{
+		const UFunction* Func = *It;
+
+		// Is the method valid and not latent?
+		if (Func && !Func->HasMetaData(FBlueprintMetadata::MD_Latent))
+		{
+			bool bSignatureValid = false;
+			for (TFieldIterator<FProperty> PropertyId(Func); PropertyId; ++PropertyId)
+			{
+				FProperty* Property = *PropertyId;
+				if (!(Property->PropertyFlags & CPF_Parm))
+				{
+					continue;
+				}
+
+				// If we've already had a valid signature, and found one more param, then it's not valid anymore
+				if (bSignatureValid)
+				{
+					bSignatureValid = false;
+					break;
+				}
+
+				// If there is an input parameter, then it's not matching
+				if (!(Property->PropertyFlags & (CPF_ReturnParm | CPF_OutParm)))
+				{
+					break;
+				}
+
+				// Only accept array outputs
+				const FArrayProperty* ArrayProperty = CastField<FArrayProperty>(Property);
+				if (!ArrayProperty)
+				{
+					break;
+				}
+
+				// Only accept matching signatures
+				if (!ArrayProperty->Inner->SameType(VariableProperty))
+				{
+					break;
+				}
+
+				// It is the first valid function parameter!
+				bSignatureValid = true;
+			}
+
+			if (bSignatureValid)
+			{
+				DropDownFunctionOptions.Add(MakeShared<FString>(Func->GetName()));
+			}
+		}
+	}
+
+	// Sort functions by name
+	DropDownFunctionOptions.Sort([](const TSharedPtr<FString>& A, const TSharedPtr<FString>& B)
+	{
+		return A->Compare(*B) < 0;
+	});
+
+	DropDownFunctionOptions.Insert(MakeShared<FString>(BlueprintDocumentationDetailDefs::EmptyDropDownOptionName), 0);
+}
+
+FText FBlueprintVarActionDetails::GetDropDownOptionDisplayText() const
+{
+	return FText::FromString(GetDropDownOptionsFunctionName());
+}
+
+FString FBlueprintVarActionDetails::GetDropDownOptionsFunctionName() const
+{
+	const UBlueprint* PropertyBlueprint = GetPropertyOwnerBlueprint();
+	if (!PropertyBlueprint)
+	{
+		return { };
+	}
+
+	FString Value;
+	FBlueprintEditorUtils::GetBlueprintVariableMetaData(
+		PropertyBlueprint,
+		CachedVariableName,
+		GetLocalVariableScope(CachedVariableProperty.Get()),
+		FBlueprintMetadata::MD_GetOptions,
+		Value
+	);
+
+	return Value;
+}
+
+void FBlueprintVarActionDetails::SetDropDownOptionsFunctionName(const FString& InFunctionName)
+{
+	UBlueprint* PropertyBlueprint = GetPropertyOwnerBlueprint();
+	if (!PropertyBlueprint)
+	{
+		return;
+	}
+
+	if (!GetDropDownOptionsFunctionName().Equals(InFunctionName))
+	{
+		if (InFunctionName.IsEmpty() || InFunctionName == BlueprintDocumentationDetailDefs::EmptyDropDownOptionName)
+		{
+			FBlueprintEditorUtils::RemoveBlueprintVariableMetaData(
+				PropertyBlueprint,
+				CachedVariableName,
+				GetLocalVariableScope(CachedVariableProperty.Get()),
+				FBlueprintMetadata::MD_GetOptions
+			);
+		}
+		else
+		{
+			FBlueprintEditorUtils::SetBlueprintVariableMetaData(
+				PropertyBlueprint,
+				CachedVariableName,
+				GetLocalVariableScope(CachedVariableProperty.Get()),
+				FBlueprintMetadata::MD_GetOptions,
+				*InFunctionName
+			);
+		}
+	}
+}
+
 FText FBlueprintVarActionDetails::GetDeprecationMessageText() const
 {
 	FName VarName = CachedVariableName;
diff --git a/Engine/Source/Runtime/Slate/Private/Widgets/Views/STableViewBase.cpp b/Engine/Source/Runtime/Slate/Private/Widgets/Views/STableViewBase.cpp
--- a/Engine/Source/Runtime/Slate/Private/Widgets/Views/STableViewBase.cpp
+++ b/Engine/Source/Runtime/Slate/Private/Widgets/Views/STableViewBase.cpp
@@ -389,6 +389,8 @@ void STableViewBase::Tick( const FGeometry& AllottedGeometry, const double InCur
 			{
 				NotifyFinishedScrolling();
 			}
+
+			OnItemsRebuilt.ExecuteIfBound();
 		}
 	}
 }
@@ -707,6 +709,11 @@ int32 STableViewBase::GetNumGeneratedChildren() const
 		: 0;
 }
 
+TSharedPtr<SWidget> STableViewBase::GetGeneratedChildAt(const int32 Index) const
+{
+	return GetNumGeneratedChildren() > Index ? ItemsPanel->GetChildren()->GetChildAt(Index).ToSharedPtr() : nullptr;
+}
+
 TSharedPtr<SHeaderRow> STableViewBase::GetHeaderRow() const
 {
 	return HeaderRow;
diff --git a/Engine/Source/Runtime/Slate/Public/Widgets/Views/SListView.h b/Engine/Source/Runtime/Slate/Public/Widgets/Views/SListView.h
--- a/Engine/Source/Runtime/Slate/Public/Widgets/Views/SListView.h
+++ b/Engine/Source/Runtime/Slate/Public/Widgets/Views/SListView.h
@@ -93,6 +93,7 @@ public:
 		, _OnRowReleased()
 		, _MaxPinnedItems(6)
 		, _OnContextMenuOpening()
+		, _OnItemsRebuilt()
 		, _OnMouseButtonClick()
 		, _OnMouseButtonDoubleClick()
 		, _OnSelectionChanged()
@@ -144,6 +145,8 @@ public:
 
 		SLATE_EVENT( FOnContextMenuOpening, OnContextMenuOpening )
 
+		SLATE_EVENT( FSimpleDelegate, OnItemsRebuilt )
+
 		SLATE_EVENT(FOnMouseButtonClick, OnMouseButtonClick)
 
 		SLATE_EVENT( FOnMouseButtonDoubleClick, OnMouseButtonDoubleClick )
@@ -222,6 +225,7 @@ public:
 
 		this->SetItemsSource(InArgs.MakeListItemsSource(this->SharedThis(this)));
 		this->OnContextMenuOpening = InArgs._OnContextMenuOpening;
+		this->OnItemsRebuilt = InArgs._OnItemsRebuilt;
 		this->OnClick = InArgs._OnMouseButtonClick;
 		this->OnDoubleClick = InArgs._OnMouseButtonDoubleClick;
 		this->OnSelectionChanged = InArgs._OnSelectionChanged;
diff --git a/Engine/Source/Runtime/Slate/Public/Widgets/Views/STableViewBase.h b/Engine/Source/Runtime/Slate/Public/Widgets/Views/STableViewBase.h
--- a/Engine/Source/Runtime/Slate/Public/Widgets/Views/STableViewBase.h
+++ b/Engine/Source/Runtime/Slate/Public/Widgets/Views/STableViewBase.h
@@ -128,6 +128,9 @@ public:
 	/** @return The number of Widgets we currently have generated. */
 	SLATE_API int32 GetNumGeneratedChildren() const;
 
+	/** @return The generated row widget at the specified index if it exists (NULL otherwise). */
+	SLATE_API TSharedPtr<SWidget> GetGeneratedChildAt(int32 Index) const;
+
 	SLATE_API TSharedPtr<SHeaderRow> GetHeaderRow() const;
 
 	/** @return Returns true if the user is currently interactively scrolling the view by holding
@@ -487,6 +490,9 @@ protected:
 	/** Delegate to invoke when the context menu should be opening. If it is nullptr, a context menu will not be summoned. */
 	FOnContextMenuOpening OnContextMenuOpening;
 
+	/** Delegate to invoke after the children have been regenerated. */
+	FSimpleDelegate OnItemsRebuilt;
+
 	/** The selection mode that this tree/list is in. Note that it is up to the generated ITableRows to respect this setting. */
 	TAttribute<ESelectionMode::Type> SelectionMode;
 
diff --git a/Engine/Source/Runtime/Slate/Public/Widgets/Views/STileView.h b/Engine/Source/Runtime/Slate/Public/Widgets/Views/STileView.h
--- a/Engine/Source/Runtime/Slate/Public/Widgets/Views/STileView.h
+++ b/Engine/Source/Runtime/Slate/Public/Widgets/Views/STileView.h
@@ -45,6 +45,7 @@ public:
 		, _ItemWidth(128)
 		, _ItemAlignment(EListItemAlignment::EvenlyDistributed)
 		, _OnContextMenuOpening()
+		, _OnItemsRebuilt()
 		, _OnMouseButtonClick()
 		, _OnMouseButtonDoubleClick()
 		, _OnSelectionChanged()
@@ -89,6 +90,8 @@ public:
 
 		SLATE_EVENT( FOnContextMenuOpening, OnContextMenuOpening )
 
+		SLATE_EVENT( FSimpleDelegate, OnItemsRebuilt )
+
 		SLATE_EVENT( FOnMouseButtonClick, OnMouseButtonClick )
 
 		SLATE_EVENT( FOnMouseButtonDoubleClick, OnMouseButtonDoubleClick )
@@ -153,6 +156,7 @@ public:
 		
 		this->SetItemsSource(InArgs.MakeListItemsSource(this->SharedThis(this)));
 		this->OnContextMenuOpening = InArgs._OnContextMenuOpening;
+		this->OnItemsRebuilt = InArgs._OnItemsRebuilt;
 		this->OnClick = InArgs._OnMouseButtonClick;
 		this->OnDoubleClick = InArgs._OnMouseButtonDoubleClick;
 		this->OnSelectionChanged = InArgs._OnSelectionChanged;
diff --git a/Engine/Source/Runtime/Slate/Public/Widgets/Views/STreeView.h b/Engine/Source/Runtime/Slate/Public/Widgets/Views/STreeView.h
--- a/Engine/Source/Runtime/Slate/Public/Widgets/Views/STreeView.h
+++ b/Engine/Source/Runtime/Slate/Public/Widgets/Views/STreeView.h
@@ -121,6 +121,7 @@ public:
 		, _OnSetExpansionRecursive()
 		, _MaxPinnedItems(6) // Having more than the max amount of items leads to the extra items in the middle being collapsed into ellipses, and the last item is fully shown
 		, _OnContextMenuOpening()
+		, _OnItemsRebuilt()
 		, _OnMouseButtonClick()
 		, _OnMouseButtonDoubleClick()
 		, _OnSelectionChanged()
@@ -172,6 +173,8 @@ public:
 
 		SLATE_EVENT( FOnContextMenuOpening, OnContextMenuOpening )
 
+		SLATE_EVENT( FSimpleDelegate, OnItemsRebuilt )
+
 		SLATE_EVENT( FOnMouseButtonClick, OnMouseButtonClick)
 
 		SLATE_EVENT( FOnMouseButtonDoubleClick, OnMouseButtonDoubleClick )
@@ -250,6 +253,7 @@ public:
 
 		this->OnKeyDownHandler = InArgs._OnKeyDownHandler;
 		this->OnContextMenuOpening = InArgs._OnContextMenuOpening;
+		this->OnItemsRebuilt = InArgs._OnItemsRebuilt;
 		this->OnClick = InArgs._OnMouseButtonClick;
 		this->OnDoubleClick = InArgs._OnMouseButtonDoubleClick;
 		this->OnSelectionChanged = InArgs._OnSelectionChanged;
