[Backout] - CL34851857 of just this BP Details file which was incorrectly added to the CL

Original CL Desc

-----------------------------------------------------------------
List View: Add an OnItemsRebuilt delegate and a GetGeneratedChildAt method so that it's possible to do some post-processing on the generated list

diff --git a/Engine/Source/Editor/Kismet/Private/BlueprintDetailsCustomization.cpp b/Engine/Source/Editor/Kismet/Private/BlueprintDetailsCustomization.cpp
--- a/Engine/Source/Editor/Kismet/Private/BlueprintDetailsCustomization.cpp
+++ b/Engine/Source/Editor/Kismet/Private/BlueprintDetailsCustomization.cpp
@@ -111,7 +111,6 @@
 #include "Serialization/Archive.h"
 #include "SlateOptMacros.h"
 #include "SlotBase.h"
-#include "SSearchableComboBox.h"
 #include "SSocketChooser.h"
 #include "Styling/AppStyle.h"
 #include "Styling/ISlateStyle.h"
@@ -180,8 +179,6 @@ namespace BlueprintDocumentationDetailDefs
 	static const float DetailsTitleMaxWidth = 300.f;
 	/** magic number retrieved from SGraphNodeComment::GetWrapAt() */
 	static const float DetailsTitleWrapPadding = 32.0f;
-	/** label of the option to reset the drop-down value function name */
-	static const FString EmptyDropDownOptionName = TEXT("None");
 };
 
 void FBlueprintDetails::AddEventsCategory(IDetailLayoutBuilder& DetailBuilder, FName PropertyName, UClass* PropertyClass)
@@ -1278,53 +1275,6 @@ void FBlueprintVarActionDetails::CustomizeDetails( IDetailLayoutBuilder& DetailL
 				.Font(IDetailLayoutBuilder::GetDetailFont())
 			];
 
-		CollectDropDownOptions();
-
-		TSharedPtr<SToolTip> GetOptionsTooltip = IDocumentation::Get()->CreateToolTip(
-			LOCTEXT(
-				"VariableGetOptions_Tooltip",
-				"The name of the function which will populate a list of options the user can select from for the value of this variable."
-			),
-			nullptr, DocLink, TEXT("GetOptions")
-		);
-
-		TSharedPtr<SToolTip> GetOptionsInputTooltip = IDocumentation::Get()->CreateToolTip(
-			LOCTEXT(
-				"VariableGetOptions_InputTooltip",
-				"List of functions that return an array of names or strings.\n"
-				"You can also enter the name of a global static function, in the form of 'ClassName.FunctionName'."
-			),
-			nullptr, DocLink, TEXT("GetOptions")
-		);
-
-		Category.AddCustomRow(LOCTEXT("VariableGetOptions", "Drop-down Options"), true)
-		.Visibility(TAttribute<EVisibility>(this, &FBlueprintVarActionDetails::GetDropDownOptionsVisibility))
-		.NameContent()
-		[
-			SNew(STextBlock)
-			.ToolTip(GetOptionsTooltip)
-			.Text(LOCTEXT("VariableGetOptions", "Drop-down Options"))
-			.Font(DetailFontInfo)
-		]
-		.ValueContent()
-		[
-			SNew(SSearchableComboBox)
-			.ToolTip(GetOptionsInputTooltip)
-			.OptionsSource(&DropDownFunctionOptions)
-			.OnGenerateWidget(this, &FBlueprintVarActionDetails::GenerateDropDownOptionWidget)
-			.OnSelectionChanged(this, &FBlueprintVarActionDetails::OnDropDownOptionSelectionChanged)
-			.Content()
-			[
-				SNew(SEditableTextBox)
-				.ToolTip(GetOptionsInputTooltip)
-				.SelectAllTextWhenFocused(true)
-				.RevertTextOnEscape(true)
-				.Font(DetailFontInfo)
-				.Text(this, &FBlueprintVarActionDetails::GetDropDownOptionDisplayText)
-				.OnTextCommitted(this, &FBlueprintVarActionDetails::OnDropDownOptionTextChanged)
-			]
-		];
-
 		TSharedPtr<SToolTip> PropertyFlagsTooltip = IDocumentation::Get()->CreateToolTip(LOCTEXT("DefinedPropertyFlags_Tooltip", "List of defined flags for this property"), NULL, DocLink, TEXT("PropertyFlags"));
 
 		Category.AddCustomRow(LOCTEXT("DefinedPropertyFlags", "Defined Property Flags"), true)
@@ -3036,169 +2986,6 @@ void FBlueprintVarActionDetails::OnDeprecatedChanged(ECheckBoxState InNewState)
 	}
 }
 
-EVisibility FBlueprintVarActionDetails::GetDropDownOptionsVisibility() const
-{
-	if (FProperty* VariableProperty = CachedVariableProperty.Get())
-	{
-		const bool bMatchingType = VariableProperty->IsA(FNameProperty::StaticClass())
-			|| VariableProperty->IsA(FStrProperty::StaticClass());
-
-		if (bMatchingType && IsABlueprintVariable(VariableProperty) && IsAUserVariable(VariableProperty))
-		{
-			return EVisibility::Visible;
-		}
-	}
-
-	return EVisibility::Collapsed;
-}
-
-void FBlueprintVarActionDetails::OnDropDownOptionSelectionChanged(TSharedPtr<FString> InString,
-	ESelectInfo::Type)
-{
-	SetDropDownOptionsFunctionName(*InString);
-}
-
-void FBlueprintVarActionDetails::OnDropDownOptionTextChanged(const FText& Text, ETextCommit::Type)
-{
-	SetDropDownOptionsFunctionName(Text.ToString());
-}
-
-TSharedRef<SWidget> FBlueprintVarActionDetails::GenerateDropDownOptionWidget(TSharedPtr<FString> InItem) const
-{
-	return SNew(STextBlock)
-		.Text(FText::FromString(*InItem.Get()))
-		.Font(IDetailLayoutBuilder::GetDetailFont());
-}
-
-void FBlueprintVarActionDetails::CollectDropDownOptions()
-{
-	DropDownFunctionOptions.Empty();
-
-	const FProperty* VariableProperty = CachedVariableProperty.Get();
-	if (!VariableProperty)
-	{
-		return;
-	}
-
-	for (TFieldIterator<UFunction> It(VariableProperty->GetOwner<UClass>(), EFieldIteratorFlags::IncludeSuper); It; ++It)
-	{
-		const UFunction* Func = *It;
-
-		// Is the method valid and not latent?
-		if (Func && !Func->HasMetaData(FBlueprintMetadata::MD_Latent))
-		{
-			bool bSignatureValid = false;
-			for (TFieldIterator<FProperty> PropertyId(Func); PropertyId; ++PropertyId)
-			{
-				FProperty* Property = *PropertyId;
-				if (!(Property->PropertyFlags & CPF_Parm))
-				{
-					continue;
-				}
-
-				// If we've already had a valid signature, and found one more param, then it's not valid anymore
-				if (bSignatureValid)
-				{
-					bSignatureValid = false;
-					break;
-				}
-
-				// If there is an input parameter, then it's not matching
-				if (!(Property->PropertyFlags & (CPF_ReturnParm | CPF_OutParm)))
-				{
-					break;
-				}
-
-				// Only accept array outputs
-				const FArrayProperty* ArrayProperty = CastField<FArrayProperty>(Property);
-				if (!ArrayProperty)
-				{
-					break;
-				}
-
-				// Only accept matching signatures
-				if (!ArrayProperty->Inner->SameType(VariableProperty))
-				{
-					break;
-				}
-
-				// It is the first valid function parameter!
-				bSignatureValid = true;
-			}
-
-			if (bSignatureValid)
-			{
-				DropDownFunctionOptions.Add(MakeShared<FString>(Func->GetName()));
-			}
-		}
-	}
-
-	// Sort functions by name
-	DropDownFunctionOptions.Sort([](const TSharedPtr<FString>& A, const TSharedPtr<FString>& B)
-	{
-		return A->Compare(*B) < 0;
-	});
-
-	DropDownFunctionOptions.Insert(MakeShared<FString>(BlueprintDocumentationDetailDefs::EmptyDropDownOptionName), 0);
-}
-
-FText FBlueprintVarActionDetails::GetDropDownOptionDisplayText() const
-{
-	return FText::FromString(GetDropDownOptionsFunctionName());
-}
-
-FString FBlueprintVarActionDetails::GetDropDownOptionsFunctionName() const
-{
-	const UBlueprint* PropertyBlueprint = GetPropertyOwnerBlueprint();
-	if (!PropertyBlueprint)
-	{
-		return { };
-	}
-
-	FString Value;
-	FBlueprintEditorUtils::GetBlueprintVariableMetaData(
-		PropertyBlueprint,
-		CachedVariableName,
-		GetLocalVariableScope(CachedVariableProperty.Get()),
-		FBlueprintMetadata::MD_GetOptions,
-		Value
-	);
-
-	return Value;
-}
-
-void FBlueprintVarActionDetails::SetDropDownOptionsFunctionName(const FString& InFunctionName)
-{
-	UBlueprint* PropertyBlueprint = GetPropertyOwnerBlueprint();
-	if (!PropertyBlueprint)
-	{
-		return;
-	}
-
-	if (!GetDropDownOptionsFunctionName().Equals(InFunctionName))
-	{
-		if (InFunctionName.IsEmpty() || InFunctionName == BlueprintDocumentationDetailDefs::EmptyDropDownOptionName)
-		{
-			FBlueprintEditorUtils::RemoveBlueprintVariableMetaData(
-				PropertyBlueprint,
-				CachedVariableName,
-				GetLocalVariableScope(CachedVariableProperty.Get()),
-				FBlueprintMetadata::MD_GetOptions
-			);
-		}
-		else
-		{
-			FBlueprintEditorUtils::SetBlueprintVariableMetaData(
-				PropertyBlueprint,
-				CachedVariableName,
-				GetLocalVariableScope(CachedVariableProperty.Get()),
-				FBlueprintMetadata::MD_GetOptions,
-				*InFunctionName
-			);
-		}
-	}
-}
-
 FText FBlueprintVarActionDetails::GetDeprecationMessageText() const
 {
 	FName VarName = CachedVariableName;
