Blueprint Editor "My Blueprint" now properly uses a Function's metadata DisplayName to list the function in the implementation list, rather than its actual code name.
Changed all listing of functions & events in the Blueprint editor to respect the global Editor Settings' bShowFriendlyNames boolean.  The interface now looks consistent with that choice.
Added the actual function name to the Tooltip so users do not get thrown off by DisplayNames suddenly changing their anticipated function names.

diff --git a/Engine/Source/Editor/BlueprintGraph/Private/EdGraphSchema_K2.cpp b/Engine/Source/Editor/BlueprintGraph/Private/EdGraphSchema_K2.cpp
--- a/Engine/Source/Editor/BlueprintGraph/Private/EdGraphSchema_K2.cpp
+++ b/Engine/Source/Editor/BlueprintGraph/Private/EdGraphSchema_K2.cpp
@@ -5683,6 +5683,7 @@ void UEdGraphSchema_K2::GetGraphDisplayInformation(const UEdGraph& Graph, /*out*
 		if (GraphType == GT_Function && Function)
 		{
 			DisplayInfo.DisplayName = GetFriendlySignatureName(Function);
+			DisplayInfo.Tooltip = FText::FormatOrdered(LOCTEXT("Tooltip_With_RealFunctionName", "{0}\nFunction Name is {1}"), DisplayInfo.Tooltip, DisplayInfo.PlainName);
 		}
 		else
 		{
diff --git a/Engine/Source/Editor/Kismet/Private/SMyBlueprint.cpp b/Engine/Source/Editor/Kismet/Private/SMyBlueprint.cpp
--- a/Engine/Source/Editor/Kismet/Private/SMyBlueprint.cpp
+++ b/Engine/Source/Editor/Kismet/Private/SMyBlueprint.cpp
@@ -1392,7 +1392,7 @@ void SMyBlueprint::AddEventForFunctionGraph(UEdGraph* InEdGraph, int32 const Sec
 	//@TODO: Should be a bit more generic (or the AnimGraph shouldn't be stored as a FunctionGraph...)
 	const bool bIsConstructionScript = InEdGraph->GetFName() == UEdGraphSchema_K2::FN_UserConstructionScript;
 
-	TSharedPtr<FEdGraphSchemaAction_K2Graph> NewFuncAction = MakeShareable(new FEdGraphSchemaAction_K2Graph(EEdGraphSchemaAction_K2Graph::Function, ActionCategory, DisplayInfo.PlainName, DisplayInfo.Tooltip, bIsConstructionScript ? 2 : 1, SectionId));
+	TSharedPtr<FEdGraphSchemaAction_K2Graph> NewFuncAction = MakeShareable(new FEdGraphSchemaAction_K2Graph(EEdGraphSchemaAction_K2Graph::Function, ActionCategory, DisplayInfo.DisplayName, DisplayInfo.Tooltip, bIsConstructionScript ? 2 : 1, SectionId));
 	NewFuncAction->FuncName = InEdGraph->GetFName();
 	NewFuncAction->EdGraph = InEdGraph;
 
diff --git a/Engine/Source/Editor/UnrealEd/Private/ObjectTools.cpp b/Engine/Source/Editor/UnrealEd/Private/ObjectTools.cpp
--- a/Engine/Source/Editor/UnrealEd/Private/ObjectTools.cpp
+++ b/Engine/Source/Editor/UnrealEd/Private/ObjectTools.cpp
@@ -5043,10 +5043,9 @@ namespace ObjectTools
 			}
 			else
 			{
-				static const FString Namespace = TEXT("UObjectDisplayNames");
-				const FString Key = Function->GetFullGroupName(false);
-
-				ReturnDisplayName = Function->GetMetaDataText(TEXT("DisplayName"), Namespace, Key);
+				// Previous (and similar) code paths would go through FField::GetMetaDataText(DisplayName) which attempts localization
+				// However, we do not localize function names (and we've explicitly requested non-friendly names), so just show us the real name
+				ReturnDisplayName = FText::FromString(Function->GetName());
 			}
 		}
 
