[Viewport Toolbar] Create "Show" submenu for the new Blueprint editor viewport toolbar

diff --git a/Engine/Source/Editor/Kismet/Private/SSCSEditorViewport.cpp b/Engine/Source/Editor/Kismet/Private/SSCSEditorViewport.cpp
--- a/Engine/Source/Editor/Kismet/Private/SSCSEditorViewport.cpp
+++ b/Engine/Source/Editor/Kismet/Private/SSCSEditorViewport.cpp
@@ -381,13 +381,25 @@ TSharedPtr<SWidget> SSCSEditorViewport::MakeViewportToolbar()
 				RightSection.AddEntry(ViewModesSubmenu);
 			}
 
-			// TODO: Add this.
-			// // Add the "Show" submenu.
-			// {
-			// 	FToolMenuEntry ShowSubmenu = UE::UnrealEd::CreateViewportToolbarShowSubmenu();
-			// 	ShowSubmenu.InsertPosition.Position = EToolMenuInsertType::Last;
-			// 	RightSection.AddEntry(ShowSubmenu);
-			// }
+			// Add the "Show" submenu.
+			{
+				FToolMenuEntry ShowSubmenu = FToolMenuEntry::InitSubMenu(
+					"Show",
+					LOCTEXT("ShowLabel", "Show"),
+					LOCTEXT("ShowTooltip", "Show or hide elements from the viewport"),
+					FNewToolMenuDelegate::CreateLambda(
+						[](UToolMenu* Submenu) -> void
+						{
+							FToolMenuSection& UnnamedSection = Submenu->FindOrAddSection("", LOCTEXT("UnnamedLabel", ""));
+							UnnamedSection.AddMenuEntry(FBlueprintEditorCommands::Get().ShowFloor);
+							UnnamedSection.AddMenuEntry(FBlueprintEditorCommands::Get().ShowGrid);
+						}
+					)
+				);
+
+				ShowSubmenu.InsertPosition.Position = EToolMenuInsertType::Last;
+				RightSection.AddEntry(ShowSubmenu);
+			}
 
 			// Add the "Performance & Scalability" submenu.
 			{
