Recompile dependencies when they are out of date. This ensures that new or altered functions will be available, avoiding accidently referencing skeleton functions in bytecode
[RN] Dependent blueprints are automatically included in the compilation set when they are out of date, this fixes some obscure crashes during blueprint compilation

diff --git a/Engine/Source/Editor/Kismet/Private/BlueprintCompilationManager.cpp b/Engine/Source/Editor/Kismet/Private/BlueprintCompilationManager.cpp
--- a/Engine/Source/Editor/Kismet/Private/BlueprintCompilationManager.cpp
+++ b/Engine/Source/Editor/Kismet/Private/BlueprintCompilationManager.cpp
@@ -141,6 +141,8 @@ struct FBlueprintCompilationManagerImpl : public FGCObject
 	static bool IsQueuedForCompilation(UBlueprint* BP);
 	static void ConformToParentAndInterfaces(UBlueprint* BP);
 	static void RelinkSkeleton(UClass* SkeletonToRelink);
+	static void GatherOutOfDateDependenciesRecursive(TObjectPtr<UBlueprint> Gather, TSet<TObjectPtr<UBlueprint>>& OutOfDateDeps);
+	static void QueueOutOfDateDependencies( const TArray<FBPCompileRequest>& QueuedRequests, TArray<UBlueprint*>& OutBlueprintsToRecompile, TArray<FCompilerData>& CurrentlyCompilingBPs);
 
 	// Declaration of archive to fix up bytecode references of blueprints that are actively compiled:
 	class FFixupBytecodeReferences : public FArchiveUObject
@@ -611,62 +613,10 @@ void FBlueprintCompilationManagerImpl::FlushCompilationQueueImpl(bool bSuppressB
 
 		// STAGE I: Add any related blueprints that were not compiled, then add any children so that they will be relinked:
 		TArray<UBlueprint*> BlueprintsToRecompile;
-
-		// First add any dependents of macro libraries that are being compiled:
-		for (const FBPCompileRequest& CompileJob : QueuedRequests)
-		{
-			if ((CompileJob.CompileOptions & 
-				(	EBlueprintCompileOptions::RegenerateSkeletonOnly)
-				) != EBlueprintCompileOptions::None)
-			{
-				continue;
-			}
-
-			UBlueprint* BP = CompileJob.BPToCompile;
-
-			if(!BP->bHasBeenRegenerated && BP->GetLinker())
-			{
-				// we may have cached dependencies before being fully loaded:
-				BP->bCachedDependenciesUpToDate = false;
-			}
-
-			const bool bWasDependencyCacheOutOfDate = !BP->bCachedDependenciesUpToDate;
-
-			FBlueprintEditorUtils::EnsureCachedDependenciesUpToDate(BP);
-
-			if ((CompileJob.CompileOptions & 
-				(	EBlueprintCompileOptions::IsRegeneratingOnLoad)
-				) != EBlueprintCompileOptions::None)
-			{
-				continue;
-			}
-			
-			if(BP->BlueprintType == BPTYPE_MacroLibrary)
-			{
-				TArray<UBlueprint*> DependentBlueprints;
-				FBlueprintEditorUtils::GetDependentBlueprints(BP, DependentBlueprints);
-				for(UBlueprint* DependentBlueprint : DependentBlueprints)
-				{
-					if(!IsQueuedForCompilation(DependentBlueprint))
-					{
-						// The macro may have updated its dependency cache above; if so, we'll need to regenerate the dependent's set as well.
-						DependentBlueprint->bCachedDependenciesUpToDate &= !bWasDependencyCacheOutOfDate;
-
-						DependentBlueprint->bQueuedForCompilation = true;
-						CurrentlyCompilingBPs.Emplace(
-							FCompilerData(
-								DependentBlueprint, 
-								ECompilationManagerJobType::Normal, 
-								nullptr, 
-								EBlueprintCompileOptions::None,
-								false // full compile
-							)
-						);
-						BlueprintsToRecompile.Add(DependentBlueprint);
-					}
-				}
-			}
-		}
+		
+		// Make sure that we attempt to compile any functions that aren't currently
+		// in existence - this also ensures function signature changes are handled
+		QueueOutOfDateDependencies(QueuedRequests, BlueprintsToRecompile, CurrentlyCompilingBPs);
 
 		SlowTask.EnterProgressFrame();
 
@@ -3575,6 +3525,113 @@ void FBlueprintCompilationManagerImpl::RelinkSkeleton(UClass* SkeletonToRelink)
 	}
 }
 
+void FBlueprintCompilationManagerImpl::GatherOutOfDateDependenciesRecursive(TObjectPtr<UBlueprint> Gather, TSet<TObjectPtr<UBlueprint>>& OutOfDateDeps)
+{	
+	FBlueprintEditorUtils::EnsureCachedDependenciesUpToDate(Gather);
+		
+	// make sure any dirty dependencies are also compiled:
+	for (TWeakObjectPtr<UBlueprint> BPWeak : Gather->CachedDependencies)
+	{
+		if (UBlueprint* BP = BPWeak.Get())
+		{
+			if (BP->Status == BS_Dirty && BP->bQueuedForCompilation == false)
+			{
+				if(!OutOfDateDeps.Contains(BP))
+				{
+					OutOfDateDeps.Add(BP);
+					GatherOutOfDateDependenciesRecursive(BP, OutOfDateDeps);
+				}
+			}
+		}
+	}
+}
+
+void FBlueprintCompilationManagerImpl::QueueOutOfDateDependencies( const TArray<FBPCompileRequest>& QueuedRequests, TArray<UBlueprint*>& OutBlueprintsToRecompile, TArray<FCompilerData>& CurrentlyCompilingBPs)
+{
+	TArray<TObjectPtr<UBlueprint>> RootCompilationRequests;
+	// we don't care about 'skeleton only' regeneration, filter those, and operate
+	// only on 'full compilation' requests:
+	Algo::TransformIf(QueuedRequests, RootCompilationRequests,
+		[](const FBPCompileRequest& CompileRequest) -> bool
+		{
+			return
+				(CompileRequest.CompileOptions & EBlueprintCompileOptions::RegenerateSkeletonOnly)
+				== EBlueprintCompileOptions::None;
+		},
+		[](const FBPCompileRequest& CompileRequest)
+		{
+			return CompileRequest.BPToCompile;
+		}
+	);
+	
+	// Build up full compilation requests, including any 
+	// BPs that are using a macro lib that has had compilation
+	// requested (this occurs on PIE):
+	TArray<TObjectPtr<UBlueprint>> FullCompilationRequests;
+	FullCompilationRequests.Reserve(RootCompilationRequests.Num());
+	for(TObjectPtr<UBlueprint> BP : RootCompilationRequests)
+	{
+		FullCompilationRequests.Add(BP);
+
+		if(!BP->bHasBeenRegenerated && BP->GetLinker())
+		{
+			// we may have cached dependencies before being fully loaded:
+			BP->bCachedDependenciesUpToDate = false;
+		}
+
+		const bool bWasDependencyCacheOutOfDate = !BP->bCachedDependenciesUpToDate;
+
+		if(BP->BlueprintType == BPTYPE_MacroLibrary)
+		{
+			TArray<UBlueprint*> DependentBlueprints;
+			FBlueprintEditorUtils::GetDependentBlueprints(BP, DependentBlueprints);
+			for(UBlueprint* DependentBlueprint : DependentBlueprints)
+			{
+				// if the macro is out of date, transitively dirty dependencies for dependents:
+				DependentBlueprint->bCachedDependenciesUpToDate &= bWasDependencyCacheOutOfDate;
+				FullCompilationRequests.Add(DependentBlueprint);
+			}
+		}
+	}
+
+	// Gather dependencies that are also out of date - these must be recompiled
+	// in case a function signature has changed
+	TSet<TObjectPtr<UBlueprint>> DependenciesToRecompile;
+	for(TObjectPtr<UBlueprint> BP : FullCompilationRequests)
+	{
+		if(!BP->bQueuedForCompilation)
+		{
+			if(!DependenciesToRecompile.Contains(BP))
+			{
+				DependenciesToRecompile.Add(BP);
+				GatherOutOfDateDependenciesRecursive(BP, DependenciesToRecompile);
+			}
+		}
+		else
+		{
+			// root compilation request (queued by caller), just look for its
+			// out of date dependencies
+			GatherOutOfDateDependenciesRecursive(BP, DependenciesToRecompile);
+		}
+	}
+
+	for(TObjectPtr<UBlueprint> BP : DependenciesToRecompile)
+	{
+		ensure(!BP->bQueuedForCompilation);
+		BP->bQueuedForCompilation = true;
+		CurrentlyCompilingBPs.Emplace(
+			FCompilerData(
+				BP, 
+				ECompilationManagerJobType::Normal, 
+				nullptr, 
+				EBlueprintCompileOptions::None,
+				false // full compile
+			)
+		);
+		OutBlueprintsToRecompile.Add(BP);
+	}
+}
+
 // FFixupBytecodeReferences Implementation:
 FBlueprintCompilationManagerImpl::FFixupBytecodeReferences::FFixupBytecodeReferences(UObject* InObject)
 {
