Kismet: Added slate argument to SKismetInspector to specify if the section selector should be shown.

It defaults to false, matching previous behavior. This change is in preparation to UI improvements to the NDC blueprint editor.

diff --git a/Engine/Source/Editor/Kismet/Private/SKismetInspector.cpp b/Engine/Source/Editor/Kismet/Private/SKismetInspector.cpp
--- a/Engine/Source/Editor/Kismet/Private/SKismetInspector.cpp
+++ b/Engine/Source/Editor/Kismet/Private/SKismetInspector.cpp
@@ -405,6 +405,7 @@ void SKismetInspector::Construct(const FArguments& InArgs)
 	DetailsViewArgs.ViewIdentifier = InArgs._ViewIdentifier;
 	DetailsViewArgs.ExternalScrollbar = InArgs._ExternalScrollbar;
 	DetailsViewArgs.ScrollbarAlignment = InArgs._ScrollbarAlignment;
+	DetailsViewArgs.bShowSectionSelector = InArgs._ShowSectionSelector;
 	if (Kismet2.IsValid())
 	{
 		TSharedPtr<IClassViewerFilter> ImportedClassFilter = Kismet2->GetImportedClassViewerFilter();
diff --git a/Engine/Source/Editor/Kismet/Public/SKismetInspector.h b/Engine/Source/Editor/Kismet/Public/SKismetInspector.h
--- a/Engine/Source/Editor/Kismet/Public/SKismetInspector.h
+++ b/Engine/Source/Editor/Kismet/Public/SKismetInspector.h
@@ -52,6 +52,7 @@ public:
 		, _ShowTitleArea(false)
 		, _ShowLocalVariables(false)
 		, _ScrollbarAlignment(HAlign_Right)
+		, _ShowSectionSelector(false)
 		{}
 
 		SLATE_ARGUMENT(TWeakPtr<FBlueprintEditor>, Kismet2)
@@ -66,6 +67,7 @@ public:
 		SLATE_ARGUMENT( bool, ShowLocalVariables)
 		SLATE_ARGUMENT( TSharedPtr<SScrollBar>, ExternalScrollbar)
 		SLATE_ARGUMENT( EHorizontalAlignment, ScrollbarAlignment)
+		SLATE_ARGUMENT( bool, ShowSectionSelector)
 	SLATE_END_ARGS()
 	
 	void Construct(const FArguments& InArgs);
