When renaming function parameters, ensure that any getter nodes are also renamed accordingly.

Typically, OnPinRenamed would gather the affected nodes via FPinRenamedHelper and call RenameUserDefinedPin. However, FPinRenamedHelper doesn't collect getter nodes. Additionally, RenameUserDefinedPin is insufficient for getter nodes: we have to use HandleVariableRenamed to rename both the pin and the variable reference.

diff --git a/Engine/Source/Editor/Kismet/Private/BlueprintDetailsCustomization.cpp b/Engine/Source/Editor/Kismet/Private/BlueprintDetailsCustomization.cpp
--- a/Engine/Source/Editor/Kismet/Private/BlueprintDetailsCustomization.cpp
+++ b/Engine/Source/Editor/Kismet/Private/BlueprintDetailsCustomization.cpp
@@ -75,6 +75,7 @@
 #include "K2Node_MathExpression.h"
 #include "K2Node_Tunnel.h"
 #include "K2Node_Variable.h"
+#include "K2Node_VariableGet.h"
 #include "Kismet2/BlueprintEditorUtils.h"
 #include "Kismet2/ChildActorComponentEditorUtils.h"
 #include "Kismet2/ComponentEditorUtils.h"
@@ -5591,12 +5592,13 @@ bool FBaseBlueprintGraphActionDetails::OnPinRenamed(UK2Node_EditablePinBase* Tar
 {
 	// Before changing the name, verify the name
 	FText ErrorMessage;
-	if(!OnVerifyPinRename(TargetNode, OldName, NewName, ErrorMessage))
+	if (!OnVerifyPinRename(TargetNode, OldName, NewName, ErrorMessage))
 	{
 		return false;
 	}
 
 	UEdGraph* Graph = GetGraph();
+	check(Graph);
 
 	if (TargetNode)
 	{
@@ -5610,7 +5612,8 @@ bool FBaseBlueprintGraphActionDetails::OnPinRenamed(UK2Node_EditablePinBase* Tar
 			TerminalNodes.Add(EntryNode);
 		}
 
-		bool bRequiresFunctionSignatureUpdate = false;
+		bool bHasFunctionEntryNode = false;
+		bool bHasFunctionResultNode = false;
 		for (UK2Node_EditablePinBase* TerminalNode : TerminalNodes)
 		{
 			TerminalNode->Modify();
@@ -5618,7 +5621,8 @@ bool FBaseBlueprintGraphActionDetails::OnPinRenamed(UK2Node_EditablePinBase* Tar
 
 			// Since function terminator node pins map to generated function properties, we need to
 			// regenerate the referenced function so that dependent pins can be reconstructed properly.
-			bRequiresFunctionSignatureUpdate |= TerminalNode->IsA<UK2Node_FunctionTerminator>();
+			bHasFunctionEntryNode |= TerminalNode->IsA<UK2Node_FunctionEntry>();
+			bHasFunctionResultNode |= TerminalNode->IsA<UK2Node_FunctionResult>();
 		}
 
 		UBlueprint* TargetBlueprint = GetBlueprintObj();
@@ -5659,8 +5663,25 @@ bool FBaseBlueprintGraphActionDetails::OnPinRenamed(UK2Node_EditablePinBase* Tar
 			}
 		}
 
-		// If necessary, regenerate the skeleton class to update function properties.
-		if (bRequiresFunctionSignatureUpdate)
+		// A function signature change requires recompilation to update the underlying property chain.
+		// However, if we changed the function inputs at all, then we need to update any getter nodes that referenced the old name.
+		if (bHasFunctionEntryNode)
+		{
+			TArray<UK2Node_VariableGet*> GetterNodes;
+			Graph->GetNodesOfClass<UK2Node_VariableGet>(GetterNodes);
+
+			for (UK2Node_VariableGet* GetterNode : GetterNodes)
+			{
+				check(GetterNode);
+				if (GetterNode->ReferencesVariable(OldName, nullptr))
+				{
+					GetterNode->HandleVariableRenamed(TargetBlueprint, TargetBlueprint->GeneratedClass, Graph, OldName, NewFName);
+				}
+			}
+
+			FBlueprintEditorUtils::MarkBlueprintAsStructurallyModified(TargetBlueprint);
+		}
+		else if (bHasFunctionResultNode)
 		{
 			FBlueprintEditorUtils::MarkBlueprintAsStructurallyModified(TargetBlueprint);
 		}
