Add JsonObjectGraphFunctionLibrary, a set of edit time functions for stringifying objects and persisting those strings to temporary files. Includes two new console commands snapshotblueprints <label> and snapshotblueprintclasses <label> which will write snapshots of all loaded blueprints to a folder as specified by the label. Available to python in unreal.JsonObjectGraphFunctionLibrary, as all blueprint callable functions are.

[RN] Added JsonObjectGraphFunctionLibrary, a set of edit time functions for stringifying objects and persisting those strings to temporary files. Includes two new console commands snapshotblueprints <label> and snapshotblueprintclasses <label> which will write snapshots of all loaded blueprints to a folder as specified by the label. Available to python in unreal.JsonObjectGraphFunctionLibrary, as all blueprint callable functions are.

diff --git a/Engine/Source/Editor/Kismet/Kismet.Build.cs b/Engine/Source/Editor/Kismet/Kismet.Build.cs
--- a/Engine/Source/Editor/Kismet/Kismet.Build.cs
+++ b/Engine/Source/Editor/Kismet/Kismet.Build.cs
@@ -51,6 +51,7 @@ public class Kismet : ModuleRules
 				"JsonUtilities",
 				"DesktopPlatform",
 				"HotReload",
+				"JsonObjectGraph",
 				"UMGEditor",
 				"UMG", // for SBlueprintDiff
 				"WorkspaceMenuStructure",
diff --git a/Engine/Source/Editor/Kismet/Private/Blueprints/JsonObjectGraphFunctionLibrary.cpp b/Engine/Source/Editor/Kismet/Private/Blueprints/JsonObjectGraphFunctionLibrary.cpp
new file mode 100644
--- /dev/null
+++ b/Engine/Source/Editor/Kismet/Private/Blueprints/JsonObjectGraphFunctionLibrary.cpp
@@ -0,0 +1,142 @@
+// Copyright Epic Games, Inc. All Rights Reserved.
+
+#include "JsonObjectGraphFunctionLibrary.h"
+
+#include "HAL/FileManager.h"
+#include "Engine/Blueprint.h"
+#include "UnrealEngine.h"
+#include "UObject/Package.h"
+#include "UObject/UObjectIterator.h"
+
+namespace UE::Private 
+{
+
+static const TCHAR* SnapshotBlueprintsHelp = TEXT("Usage: snapshotblueprints label - label is a name for the folder where snapshots are saved");
+
+static void SnapshotBlueprints(const TArray<FString>& Args)
+{
+	if (Args.Num() == 0)
+	{
+		UE_LOG(LogEngine, Display, TEXT("%s"), SnapshotBlueprintsHelp);
+		return;
+	}
+
+	const FString& Label = Args[0];
+	FString ScratchFilenameWritten;
+	for (TObjectIterator<UBlueprint> BPIt; BPIt; ++BPIt)
+	{
+		const UBlueprint* BP = *BPIt;
+		UJsonObjectGraphFunctionLibrary::WritePackageToTempFile(BP, Label, FJsonStringifyOptions(), ScratchFilenameWritten);
+	}
+}
+
+FAutoConsoleCommand SnapshotBlueprintsCommand(
+	TEXT("snapshotblueprints"),
+	*FString::Format(TEXT("Write out a snapshot to the Saved directory of currently loaded blueprints.\n{0}"), { SnapshotBlueprintsHelp }),
+	FConsoleCommandWithArgsDelegate::CreateStatic(&SnapshotBlueprints),
+	ECVF_Default
+);
+
+static const TCHAR* SnapshotBlueprintClassesHelp = TEXT("Usage: snapshotblueprintclasses label - label is a name for the folder where snapshots are saved");
+
+static void SnapshotBlueprintClasses(const TArray<FString>& Args)
+{
+	if (Args.Num() == 0)
+	{
+		UE_LOG(LogEngine, Display, TEXT("%s"), SnapshotBlueprintClassesHelp);
+		return;
+	}
+
+	const FString& Label = Args[0];
+	FString ScratchFilenameWritten;
+	for (TObjectIterator<UBlueprint> BPIt; BPIt; ++BPIt)
+	{
+		const UBlueprint* BP = *BPIt;
+		// skip editor only blueprint classes, they may have unstable
+		// data because they are never cooked:
+		if (!BP->GeneratedClass || IsEditorOnlyObject(BP->GeneratedClass))
+		{
+			continue;
+		}
+
+		UJsonObjectGraphFunctionLibrary::WriteBlueprintClassToTempFile(BP, Label, FJsonStringifyOptions(), ScratchFilenameWritten);
+	}
+}
+
+FAutoConsoleCommand SnapshotBlueprintClassesCommand(
+	TEXT("snapshotblueprintclasses"),
+	*FString::Format(TEXT("Write out a snapshot to the Saved directory of currently loaded blueprint classes - the principle outputs of blueprint compilation.\n{0}"), { SnapshotBlueprintClassesHelp }),
+	FConsoleCommandWithArgsDelegate::CreateStatic(&SnapshotBlueprintClasses),
+	ECVF_Default
+);
+
+FString GetIntermediateAssetName(const UPackage* Package, const TCHAR* Prefix)
+{
+	return
+		FPaths::ProjectSavedDir() +
+		FString(TEXT("Temp/")) + Prefix +
+		Package->GetPathName() +
+		FString(TEXT("_snap.json"));
+}
+}
+
+void UJsonObjectGraphFunctionLibrary::Stringify(const TArray<UObject*>& Objects, FJsonStringifyOptions Options, FString& ResultString)
+{
+	FUtf8String Result = UE::JsonObjectGraph::Stringify(Objects, Options);
+	ResultString = FString(Result);
+}
+
+void UJsonObjectGraphFunctionLibrary::WritePackageToTempFile(const UObject* Object, const FString& Label, FJsonStringifyOptions Options, FString& OutFilename)
+{
+	OutFilename = FString();
+	if (Object == nullptr)
+	{
+		return;
+	}
+
+	const UPackage* Package = Object->GetPackage();
+	if (Package == GetTransientPackage())
+	{
+		UE_LOG(LogBlueprint, Warning, TEXT("Attempted to snapshot the transient package"));
+		return;
+	}
+
+	FUtf8String Result = UE::JsonObjectGraph::Stringify({Package}, Options);
+	FString Filename = UE::Private::GetIntermediateAssetName(Package, *Label);
+	if (ensure(!Result.IsEmpty()))
+	{
+		TUniquePtr<FArchive> FileArchive(IFileManager::Get().CreateDebugFileWriter(*Filename));
+		FileArchive->Serialize((void*)Result.GetCharArray().GetData(), Result.GetCharArray().Num() - 1);
+		OutFilename = MoveTemp(Filename);
+	}
+}
+
+void UJsonObjectGraphFunctionLibrary::WriteBlueprintClassToTempFile(const UBlueprint* BP, const FString& Label, FJsonStringifyOptions Options, FString& OutFilename)
+{
+	OutFilename = FString();
+	if (BP == nullptr)
+	{
+		return;
+	}
+
+	UClass* BPGC = BP->GeneratedClass;
+	if (BP->BlueprintType == BPTYPE_MacroLibrary ||
+		!BPGC)
+	{
+		return;
+	}
+
+	if (!BPGC->ClassDefaultObject)
+	{
+		UE_LOG(LogBlueprint, Warning, TEXT("Attempted to serialize class with no CDO: %s"), *BPGC->GetPathName());
+		return;
+	}
+
+	// When writing a class we should always exclude editor only data:
+	Options.Flags |= EJsonStringifyFlags::FilterEditorOnlyData;
+	FUtf8String Result = UE::JsonObjectGraph::Stringify({ BPGC, BPGC->ClassDefaultObject }, Options);
+	FString Filename = UE::Private::GetIntermediateAssetName(BP->GetPackage(), *Label);
+	TUniquePtr<FArchive> FileArchive(IFileManager::Get().CreateDebugFileWriter(*Filename));
+	FileArchive->Serialize((void*)Result.GetCharArray().GetData(), Result.GetCharArray().Num() - 1);
+	OutFilename = MoveTemp(Filename);
+}
diff --git a/Engine/Source/Editor/Kismet/Private/Blueprints/JsonObjectGraphFunctionLibrary.h b/Engine/Source/Editor/Kismet/Private/Blueprints/JsonObjectGraphFunctionLibrary.h
new file mode 100644
--- /dev/null
+++ b/Engine/Source/Editor/Kismet/Private/Blueprints/JsonObjectGraphFunctionLibrary.h
@@ -0,0 +1,59 @@
+// Copyright Epic Games, Inc. All Rights Reserved.
+
+#pragma once
+
+#include "Kismet/BlueprintFunctionLibrary.h"
+#include "JsonObjectGraph/Stringify.h"
+
+#include "JsonObjectGraphFunctionLibrary.generated.h"
+
+UCLASS()
+class UJsonObjectGraphFunctionLibrary : public UBlueprintFunctionLibrary
+{
+	GENERATED_BODY()
+
+public:
+	/**
+	 * Writes the provided objects to a string output, using the JsonObjectGraph format. Reachable
+	 * nested objects will be included automatically. Objects not within a root should be included in 
+	 * RootObjects if they want to be deeply represented in the result string
+	 * 
+	 * Examples of invocation from python:
+	 *  Print an object:
+	 *	print( unreal.JsonObjectGraphFunctionLibrary.stringify([object], tuple()) )
+	 *  Print a list objects:
+	 *	print( unreal.JsonObjectGraphFunctionLibrary.stringify(objects, tuple()) )
+	 *  Print an object's entire package:
+	 *	print( unreal.JsonObjectGraphFunctionLibrary.stringify([unreal.EditorAssetLibrary.get_package_for_object(object)], tuple()) )
+	 * 
+	 * @param	RootObjects		The objects to write at the root level
+	 * @param	Options			Options controlling the written format
+	 * @param	ResultString	The objects stringified
+	 */
+	UFUNCTION(BlueprintCallable, Category = "Json")
+	static void Stringify(const TArray<UObject*>& RootObjects, FJsonStringifyOptions Options, FString& ResultString);
+
+	/** 
+	 * Writes all objects in the provided object's package to a temporary file
+	 * using the JsonObjectGraph format.
+	 * 
+	 * @param	Object			The object whose package will be written to the file
+	 * @param	Label			A label to disambiguate the temporary file
+	 * @param	Options			Options controlling the written format
+	 * @param	OutFilename		The filename written, empty if no file written
+	 */
+	UFUNCTION(BlueprintCallable, Category = "Json")
+	static void WritePackageToTempFile(const UObject* Object, const FString& Label, FJsonStringifyOptions Options, FString& OutFilename);
+	
+	/** 
+	 * Writes only the provided blueprint's Class and CDO to a temporary file
+	 * using the JsonObjectGraph format. Always excludes editor only data.
+	 * 
+	 * @param	BP				The blueprint to write to a file
+	 * @param	Label			A label to disambiguate the temporary file
+	 * @param	OutFilename		The filename written, empty if no file written
+	 */
+	UFUNCTION(BlueprintCallable, Category = "Json")
+	static void WriteBlueprintClassToTempFile(const UBlueprint* BP, const FString& Label, FJsonStringifyOptions Options, FString& OutFilename);
+};
+
