Remove redundant remapping of CDO subobjects, PatchNewCDOIntoLinker should invoke PatchCDOSubobjectsIntoExport if necessary

#rnx

diff --git a/Engine/Source/Editor/Kismet/Private/BlueprintCompilationManager.cpp b/Engine/Source/Editor/Kismet/Private/BlueprintCompilationManager.cpp
--- a/Engine/Source/Editor/Kismet/Private/BlueprintCompilationManager.cpp
+++ b/Engine/Source/Editor/Kismet/Private/BlueprintCompilationManager.cpp
@@ -2583,7 +2583,6 @@ void FBlueprintCompilationManagerImpl::ReinstanceBatch(TArray<FReinstancingJob>&
 					if(OldCDOIndex != INDEX_NONE)
 					{
 						FBlueprintEditorUtils::PatchNewCDOIntoLinker(CurrentBP->GeneratedClass->ClassDefaultObject, CurrentLinker, OldCDOIndex, InLoadContext);
-						FBlueprintEditorUtils::PatchCDOSubobjectsIntoExport(OldCDO, CurrentBP->GeneratedClass->ClassDefaultObject);
 					}
 				}
 			}
