Fix config variable tooltip suggesting usage of SKEL_ classes, which should never be used unless a preview of the authoritative class is needed

[RN] The tooltip associated with the config variable checkbox in the blueprint editor no longer erroniously mentions SKEL_ classes

diff --git a/Engine/Source/Editor/Kismet/Private/BlueprintDetailsCustomization.cpp b/Engine/Source/Editor/Kismet/Private/BlueprintDetailsCustomization.cpp
--- a/Engine/Source/Editor/Kismet/Private/BlueprintDetailsCustomization.cpp
+++ b/Engine/Source/Editor/Kismet/Private/BlueprintDetailsCustomization.cpp
@@ -654,8 +654,9 @@ void FBlueprintVarActionDetails::CustomizeDetails( IDetailLayoutBuilder& DetailL
 		FFormatNamedArguments ConfigTooltipArgs;
 		if (UClass* OwnerClass = VariableProperty->GetOwnerClass())
 		{
-			ConfigTooltipArgs.Add(TEXT("ConfigName"), FText::FromName(OwnerClass->ClassConfigName));
-			ConfigTooltipArgs.Add(TEXT("ConfigSection"), FText::FromString(OwnerClass->GetPathName()));
+			const UClass* RealClass = OwnerClass->GetAuthoritativeClass();
+			ConfigTooltipArgs.Add(TEXT("ConfigName"), FText::FromName(RealClass->ClassConfigName));
+			ConfigTooltipArgs.Add(TEXT("ConfigSection"), FText::FromString(RealClass->GetPathName()));
 		}
 		LocalizedTooltip = FText::Format(LOCTEXT("VariableExposeToConfig_Tooltip", "Should this variable read its default value from a config file if it is present?\r\n\r\nThis is used for customizing variable default values and behavior between different projects and configurations.\r\n\r\nConfig file [{ConfigName}]\r\nConfig section [{ConfigSection}]"), ConfigTooltipArgs);
 	}
