Make GIsReinstancing atomic as we understand it may be used across threads. Worst case is slightly misleading error messages, but the nature of reinstancing is that we will not be loading data or otherwise instantiating objects while updating class layouts.

#rnx

diff --git a/Engine/Source/Editor/Kismet/Private/BlueprintCompilationManager.cpp b/Engine/Source/Editor/Kismet/Private/BlueprintCompilationManager.cpp
--- a/Engine/Source/Editor/Kismet/Private/BlueprintCompilationManager.cpp
+++ b/Engine/Source/Editor/Kismet/Private/BlueprintCompilationManager.cpp
@@ -1728,7 +1728,7 @@ void FBlueprintCompilationManagerImpl::FlushCompilationQueueImpl(bool bSuppressB
 				if(!bSuppressBroadcastCompiled)
 				{
 					// Some logic (e.g. UObject::ProcessInternal) uses this flag to suppress warnings:
-					TGuardValue<bool> ReinstancingGuard(GIsReinstancing, true);
+					TGuardValue<std::atomic<bool>, bool> ReinstancingGuard(GIsReinstancing, true);
 					CompilerData.BP->BroadcastCompiled();
 				}
 
@@ -1915,7 +1915,7 @@ void FBlueprintCompilationManagerImpl::FlushReinstancingQueueImpl(bool bFindAndR
 	{
 		FScopedDurationTimer ReinstTimer(TimeReinstancing);
 		
-		TGuardValue<bool> ReinstancingGuard(GIsReinstancing, true);
+		TGuardValue<std::atomic<bool>, bool> ReinstancingGuard(GIsReinstancing, true);
 		
 		TMap<UClass*, UClass*> ClassesToReinstanceOwned = ObjectPtrDecay(MoveTemp(ClassesToReinstance));
 		ClassesToReinstance = {};
@@ -2168,7 +2168,7 @@ void FBlueprintCompilationManagerImpl::ReparentHierarchies(const TMap<UClass*, U
 	{
 		OldClassToNewClassDerivedTypes.Add(ReinstancingJob.OldToNew);
 	}
-	TGuardValue<bool> ReinstancingGuard(GIsReinstancing, true);
+	TGuardValue<std::atomic<bool>, bool> ReinstancingGuard(GIsReinstancing, true);
 	FReplaceInstancesOfClassParameters BatchOptions;
 	BatchOptions.bArchetypesAreUpToDate = true;
 	BatchOptions.bReplaceReferencesToOldClasses = bReplaceReferencesToOldClasses;
@@ -2223,7 +2223,7 @@ void FBlueprintCompilationManagerImpl::BuildDSOMap(UObject* OldObject, UObject*
 
 void FBlueprintCompilationManagerImpl::ReinstanceBatch(TArray<FReinstancingJob>& Reinstancers, TMap< UClass*, UClass* >& InOutOldToNewClassMap, FUObjectSerializeContext* InLoadContext, TMap<UClass*, TMap<UObject*, UObject*>>* OldToNewTemplates /* = nullptr*/)
 {
-	TGuardValue<bool> ReinstancingGuard(GIsReinstancing, true);
+	TGuardValue<std::atomic<bool>, bool> ReinstancingGuard(GIsReinstancing, true);
 
 	// This is only needed when using the legacy editor loader (ie: FLinkerLoad)
 	// Zen loader already uses FPlayInEditorLoadingScope
diff --git a/Engine/Source/Editor/UnrealEd/Private/Kismet2/KismetReinstanceUtilities.cpp b/Engine/Source/Editor/UnrealEd/Private/Kismet2/KismetReinstanceUtilities.cpp
--- a/Engine/Source/Editor/UnrealEd/Private/Kismet2/KismetReinstanceUtilities.cpp
+++ b/Engine/Source/Editor/UnrealEd/Private/Kismet2/KismetReinstanceUtilities.cpp
@@ -991,7 +991,7 @@ void FBlueprintCompileReinstancer::ReinstanceObjects(bool bForceAlwaysReinstance
 
 				if (!IsReinstancingSkeleton())
 				{
-					TGuardValue<bool> ReinstancingGuard(GIsReinstancing, true);
+					TGuardValue<std::atomic<bool>, bool> ReinstancingGuard(GIsReinstancing, true);
 
 					TArray<TSharedPtr<FReinstanceFinalizer>> Finalizers;
 
diff --git a/Engine/Source/Runtime/Core/Private/Misc/CoreGlobals.cpp b/Engine/Source/Runtime/Core/Private/Misc/CoreGlobals.cpp
--- a/Engine/Source/Runtime/Core/Private/Misc/CoreGlobals.cpp
+++ b/Engine/Source/Runtime/Core/Private/Misc/CoreGlobals.cpp
@@ -161,7 +161,7 @@ bool GIsGCingAfterBlueprintCompile = false;
 bool GIsReconstructingBlueprintInstances = false;
 
 /** True if actors and objects are being re-instanced. */
-bool GIsReinstancing = false;
+std::atomic<bool> GIsReinstancing = false;
 
 /** Settings for when using UE as a library */
 FUELibraryOverrideSettings GUELibraryOverrideSettings;
diff --git a/Engine/Source/Runtime/Core/Public/CoreGlobals.h b/Engine/Source/Runtime/Core/Public/CoreGlobals.h
--- a/Engine/Source/Runtime/Core/Public/CoreGlobals.h
+++ b/Engine/Source/Runtime/Core/Public/CoreGlobals.h
@@ -166,7 +166,7 @@ extern CORE_API bool GIsGCingAfterBlueprintCompile;
 extern CORE_API bool GIsReconstructingBlueprintInstances;
 
 /** True if actors and objects are being re-instanced. */
-extern CORE_API bool GIsReinstancing;
+extern CORE_API std::atomic<bool> GIsReinstancing;
 
 /** Helper function to flush resource streaming. */
 extern CORE_API void(*GFlushStreamingFunc)(void);
