Add support for tagging reference parameters as const in blueprints

[RN] Reference parameters in blueprint functions can now be tagged as const, preventing mutation within the function's graph

diff --git a/Engine/Source/Editor/Kismet/Private/BlueprintDetailsCustomization.cpp b/Engine/Source/Editor/Kismet/Private/BlueprintDetailsCustomization.cpp
--- a/Engine/Source/Editor/Kismet/Private/BlueprintDetailsCustomization.cpp
+++ b/Engine/Source/Editor/Kismet/Private/BlueprintDetailsCustomization.cpp
@@ -3800,6 +3800,22 @@ void FBlueprintGraphArgumentLayout::GenerateChildContent( IDetailChildrenBuilder
 					.OnCheckStateChanged(this, &FBlueprintGraphArgumentLayout::OnRefCheckStateChanged)
 					.IsEnabled(!ShouldPassByRefBeReadOnly())
 				];
+
+			ChildrenBuilder.AddCustomRow(LOCTEXT("FunctionArgDetailsConst", "Const"))
+				.NameContent()
+				[
+					SNew(STextBlock)
+					.Text(LOCTEXT("FunctionArgDetailsConst", "Const"))
+					.ToolTipText(LOCTEXT("FunctionArgDetailsConstTooltip", "When passing by reference a parameter can be specified as const (not writable)"))
+					.Font(IDetailLayoutBuilder::GetDetailFont())
+				]
+			.ValueContent()
+				[
+					SNew(SCheckBox)
+					.IsChecked(this, &FBlueprintGraphArgumentLayout::IsConstChecked)
+					.OnCheckStateChanged(this, &FBlueprintGraphArgumentLayout::OnConstCheckStateChanged)
+					.IsEnabled(this, &FBlueprintGraphArgumentLayout::CanChangeConst)
+				];
 		}
 	}
 		
@@ -3950,16 +3966,51 @@ ECheckBoxState FBlueprintGraphArgumentLayout::IsRefChecked() const
 	return (PinType.bIsReference || PinType.IsArray())  ? ECheckBoxState::Checked : ECheckBoxState::Unchecked;
 }
 
+ECheckBoxState FBlueprintGraphArgumentLayout::IsConstChecked() const
+{
+	return OnGetPinInfo().bIsConst ? ECheckBoxState::Checked : ECheckBoxState::Unchecked;
+}
+
 void FBlueprintGraphArgumentLayout::OnRefCheckStateChanged(ECheckBoxState InState)
 {
 	const FScopedTransaction Transaction(LOCTEXT("ChangeByRef", "Change Pass By Reference"));
 
+	const bool bIsReference = (InState == ECheckBoxState::Checked);
 	FEdGraphPinType PinType = OnGetPinInfo();
-	PinType.bIsReference = (InState == ECheckBoxState::Checked) ? true : false;
-	
+	PinType.bIsReference = bIsReference;
+	if(!bIsReference)
+	{
+		// constness is not meaningful for non reference types,
+		// reset constness when refness is toggled off:
+		PinType.bIsConst = ShouldBeForceConst();
+	}
+
+	PinInfoChanged(PinType);
+}
+
+void FBlueprintGraphArgumentLayout::OnConstCheckStateChanged(ECheckBoxState InState)
+{
+	FEdGraphPinType PinType = OnGetPinInfo();
+
+	PinType.bIsConst = ShouldBeForceConst() || (InState == ECheckBoxState::Checked);
+
 	PinInfoChanged(PinType);
 }
 
+bool FBlueprintGraphArgumentLayout::ShouldBeForceConst() const
+{
+	// Const-ness is not meaningful unless we're passing by reference, including all arrays 
+	// which are implicitly passed by reference. If we have a reference pin type (implicit or
+	// otherwise) we want to honor ShouldUseConstRefParams:
+	FEdGraphPinType PinType = OnGetPinInfo();
+	return (PinType.bIsReference || PinType.IsArray()) && TargetNode && TargetNode->ShouldUseConstRefParams();
+}
+
+bool FBlueprintGraphArgumentLayout::CanChangeConst() const
+{
+	return !ShouldBeForceConst() && IsRefChecked() == ECheckBoxState::Checked;
+}
+
 void FBlueprintGraphArgumentLayout::PinInfoChanged(const FEdGraphPinType& PinType)
 {
 	if (ParamItemPtr.IsValid() && FBlueprintEditorUtils::IsPinTypeValid(PinType))
diff --git a/Engine/Source/Editor/Kismet/Private/BlueprintDetailsCustomization.h b/Engine/Source/Editor/Kismet/Private/BlueprintDetailsCustomization.h
--- a/Engine/Source/Editor/Kismet/Private/BlueprintDetailsCustomization.h
+++ b/Engine/Source/Editor/Kismet/Private/BlueprintDetailsCustomization.h
@@ -561,10 +561,15 @@ private:
 
 	/** Returns whether the "Pass-by-Reference" checkbox is checked or not */
 	ECheckBoxState IsRefChecked() const;
+	ECheckBoxState IsConstChecked() const;
 
 	/** Handles toggling the "Pass-by-Reference" checkbox */
 	void OnRefCheckStateChanged(ECheckBoxState InState);
 
+	void OnConstCheckStateChanged(ECheckBoxState InState);
+	bool ShouldBeForceConst() const;
+	bool CanChangeConst() const;
+
 private:
 	/** The parent graph action details customization */
 	TWeakPtr<class FBaseBlueprintGraphActionDetails> GraphActionDetailsPtr;
