Fix category and comments on the new JsonObjectGraph related functions, mostly emphasizing experimental nature

#rnx

diff --git a/Engine/Source/Editor/Kismet/Private/Blueprints/JsonObjectGraphFunctionLibrary.h b/Engine/Source/Editor/Kismet/Private/Blueprints/JsonObjectGraphFunctionLibrary.h
--- a/Engine/Source/Editor/Kismet/Private/Blueprints/JsonObjectGraphFunctionLibrary.h
+++ b/Engine/Source/Editor/Kismet/Private/Blueprints/JsonObjectGraphFunctionLibrary.h
@@ -14,6 +14,8 @@ class UJsonObjectGraphFunctionLibrary : public UBlueprintFunctionLibrary
 
 public:
 	/**
+	 * ! EXPERIMENTAL ! 
+	 * 
 	 * Writes the provided objects to a string output, using the JsonObjectGraph format. Reachable
 	 * nested objects will be included automatically. Objects not within a root should be included in 
 	 * RootObjects if they want to be deeply represented in the result string
@@ -30,10 +32,12 @@ public:
 	 * @param	Options			Options controlling the written format
 	 * @param	ResultString	The objects stringified
 	 */
-	UFUNCTION(BlueprintCallable, Category = "Json")
+	UFUNCTION(BlueprintCallable, Category = "Experimental|Json")
 	static void Stringify(const TArray<UObject*>& RootObjects, FJsonStringifyOptions Options, FString& ResultString);
 
 	/** 
+	 * ! EXPERIMENTAL ! 
+	 * 
 	 * Writes all objects in the provided object's package to a temporary file
 	 * using the JsonObjectGraph format.
 	 * 
@@ -42,10 +46,12 @@ public:
 	 * @param	Options			Options controlling the written format
 	 * @param	OutFilename		The filename written, empty if no file written
 	 */
-	UFUNCTION(BlueprintCallable, Category = "Json")
+	UFUNCTION(BlueprintCallable, Category = "Experimental|Json")
 	static void WritePackageToTempFile(const UObject* Object, const FString& Label, FJsonStringifyOptions Options, FString& OutFilename);
 	
 	/** 
+	 * ! EXPERIMENTAL ! 
+	 * 
 	 * Writes only the provided blueprint's Class and CDO to a temporary file
 	 * using the JsonObjectGraph format. Always excludes editor only data.
 	 * 
@@ -53,7 +59,7 @@ public:
 	 * @param	Label			A label to disambiguate the temporary file
 	 * @param	OutFilename		The filename written, empty if no file written
 	 */
-	UFUNCTION(BlueprintCallable, Category = "Json")
+	UFUNCTION(BlueprintCallable, Category = "Experimental|Json")
 	static void WriteBlueprintClassToTempFile(const UBlueprint* BP, const FString& Label, FJsonStringifyOptions Options, FString& OutFilename);
 };
 
diff --git a/Engine/Source/Runtime/Experimental/JsonObjectGraph/Public/JsonObjectGraph/Stringify.h b/Engine/Source/Runtime/Experimental/JsonObjectGraph/Public/JsonObjectGraph/Stringify.h
--- a/Engine/Source/Runtime/Experimental/JsonObjectGraph/Public/JsonObjectGraph/Stringify.h
+++ b/Engine/Source/Runtime/Experimental/JsonObjectGraph/Public/JsonObjectGraph/Stringify.h
@@ -44,7 +44,7 @@ namespace UE::JsonObjectGraph
 {
 	/** 
 	 * ! EXPERIMENTAL ! contents of return string will change. Currently this is used
-	 * as a debugging faculty.
+	 * as a debugging facility.
 	 * 
 	 * @return: A string containing single JSON object with serialized representations of the 
 	 * provided Objects in that single object's __RootObjects field.
