Fix references being cleared upon archetype reinstantiation
UDN case number 12529527
Regression introduced in cl 28083808 and was probably a cut&paste error, the bClearReference was false previously and was turned on in that cl.

[RN] Fix Archetypes self reference being reset bug during reinstantiation code.

diff --git a/Engine/Source/Editor/Kismet/Private/BlueprintCompilationManager.cpp b/Engine/Source/Editor/Kismet/Private/BlueprintCompilationManager.cpp
--- a/Engine/Source/Editor/Kismet/Private/BlueprintCompilationManager.cpp
+++ b/Engine/Source/Editor/Kismet/Private/BlueprintCompilationManager.cpp
@@ -2819,7 +2819,7 @@ void FBlueprintCompilationManagerImpl::ReinstanceBatch(TArray<FReinstancingJob>&
 					}
 					for (const auto& Pair : OrderedListOfObjectToCopy)
 					{
-						FBlueprintCompileReinstancer::CopyPropertiesForUnrelatedObjects(Pair.Key, Pair.Value, /*bClearExternalReferences*/true, bUseDeltaSerialization, /*bOnlyHandleDirectSubObjects*/true, &OldToNewInstanceMap, &InOutOldToNewClassMap);
+						FBlueprintCompileReinstancer::CopyPropertiesForUnrelatedObjects(Pair.Key, Pair.Value, /*bClearExternalReferences*/false, bUseDeltaSerialization, /*bOnlyHandleDirectSubObjects*/true, &OldToNewInstanceMap, &InOutOldToNewClassMap);
 					}
 
 					if (OldToNewTemplates && !OldToNewInstanceMap.IsEmpty())
