Don't add root component if the blueprint doesn't support a construction script but currently has one

[RN] The blueprint compiler no longer attempts to add a scene root component to blueprints that do not support a construction script but currently have one

diff --git a/Engine/Source/Editor/Kismet/Private/BlueprintCompilationManager.cpp b/Engine/Source/Editor/Kismet/Private/BlueprintCompilationManager.cpp
--- a/Engine/Source/Editor/Kismet/Private/BlueprintCompilationManager.cpp
+++ b/Engine/Source/Editor/Kismet/Private/BlueprintCompilationManager.cpp
@@ -676,7 +676,7 @@ void FBlueprintCompilationManagerImpl::FlushCompilationQueueImpl(bool bSuppressB
 				if( BPGC->SimpleConstructionScript &&
 					BPGC->SimpleConstructionScript->GetSceneRootComponentTemplate(true) == nullptr)
 				{
-					bDefaultComponentMustBeAdded = true;
+					bDefaultComponentMustBeAdded = FBlueprintEditorUtils::SupportsConstructionScript(QueuedBP);
 				}
 
 				bHasPendingUberGraphFrame = BPGC->UberGraphFramePointerProperty || BPGC->UberGraphFunction;
