Fix a potential race condition during Blueprint asset search data indexing.

diff --git a/Engine/Source/Editor/Kismet/Private/FindInBlueprintManager.cpp b/Engine/Source/Editor/Kismet/Private/FindInBlueprintManager.cpp
--- a/Engine/Source/Editor/Kismet/Private/FindInBlueprintManager.cpp
+++ b/Engine/Source/Editor/Kismet/Private/FindInBlueprintManager.cpp
@@ -1549,6 +1549,8 @@ public:
 		, TickCacheIndex(0)
 		, AsyncTaskBatchIndex(0)
 		, bIsGatheringSearchMetadata(false)
+		, bHasAssetsPendingGatherWork(false)
+		, bHasAssetsPendingAsyncIndexing(false)
 		, bIsStarted(false)
 		, bIsCancelled(false)
 	{
@@ -1670,7 +1672,7 @@ public:
 
 		const int32 AsyncTaskBatchSize = CacheParams.AsyncTaskBatchSize;
 		const int32 StartIndex = AsyncTaskBatchIndex * AsyncTaskBatchSize;
-		return StartIndex < UncachedAssets.Num() || !AssetsPendingGatherQueue.IsEmpty() || !AssetsPendingAsyncIndexing.IsEmpty() || bIsGatheringSearchMetadata;
+		return StartIndex < UncachedAssets.Num() || bHasAssetsPendingGatherWork || bHasAssetsPendingAsyncIndexing || bIsGatheringSearchMetadata;
 	}
 
 	virtual bool ShouldFullyIndexAssets() const override
@@ -1720,11 +1722,14 @@ public:
 				OutAssetPaths.Add(AssetPath);
 				++Count;
 			}
+
+			bHasAssetsPendingAsyncIndexing = !AssetsPendingAsyncIndexing.IsEmpty();
 		}
 	}
 
 	virtual void AddAssetPathToGatherQueue(const FSoftObjectPath& InAssetPath) override
 	{
+		bHasAssetsPendingGatherWork = true;
 		AssetsPendingGatherQueue.Enqueue(InAssetPath);
 	}
 
@@ -1811,6 +1816,9 @@ public:
 
 							if (bEnqueueForAsyncIndexing)
 							{
+								// Signal that we have a new pending asset in the async indexing queue.
+								bHasAssetsPendingAsyncIndexing = true;
+
 								// Enqueue this asset path to restart async indexing with the updated search metadata.
 								AssetsPendingAsyncIndexing.Enqueue(AssetPath);
 							}
@@ -1821,6 +1829,9 @@ public:
 							}
 						}
 
+						// Signal whether or not the queue has been emptied.
+						bHasAssetsPendingGatherWork = !AssetsPendingGatherQueue.IsEmpty();
+
 						// Indicate that gather work is no longer in progress.
 						bIsGatheringSearchMetadata = false;
 					}
@@ -2041,6 +2052,12 @@ private:
 	/** TRUE if we're busy gathering search metadata from a loaded object on the main thread */
 	TAtomic<bool> bIsGatheringSearchMetadata;
 
+	/** TRUE if we have assets pending in the gather queue. It is not thread safe to check the queue state directly from a producer thread. */
+	TAtomic<bool> bHasAssetsPendingGatherWork;
+
+	/** TRUE if we have assets pending in the indexing queue. Adds consistency with the gather queue, but not strictly needed for thread safety. */
+	TAtomic<bool> bHasAssetsPendingAsyncIndexing;
+
 	/** TRUE if the caching process is started */
 	bool bIsStarted;
 
