Temp roll back change to BP undo subobject fixup logic while I investigate a potential regression.

#rnx

diff --git a/Engine/Source/Editor/Kismet/Private/BlueprintEditorModule.cpp b/Engine/Source/Editor/Kismet/Private/BlueprintEditorModule.cpp
--- a/Engine/Source/Editor/Kismet/Private/BlueprintEditorModule.cpp
+++ b/Engine/Source/Editor/Kismet/Private/BlueprintEditorModule.cpp
@@ -22,6 +22,7 @@
 #include "HAL/PlatformCrt.h"
 #include "IMessageLogListing.h"
 #include "ISettingsModule.h"
+#include "InstancedReferenceSubobjectHelper.h"
 #include "InstancedStaticMeshSCSEditorCustomization.h"
 #include "Internationalization/Internationalization.h"
 #include "Kismet2/BlueprintEditorUtils.h"
@@ -96,8 +97,9 @@ void FixSubObjectReferencesPostUndoRedo(UObject* InObject)
 	GetObjectsWithOuter(InObject, SubObjects, bIncludeNestedSubObjects);
 
 	// Post undo/redo, these may have the in-correct Outer but are incorrectly referenced by the CDO's UProperties
-	TArray<UObject*> PropertySubObjectReferences;
-	FReferenceFinder(PropertySubObjectReferences).FindReferences(InObject);
+PRAGMA_DISABLE_DEPRECATION_WARNINGS
+	TSet<FInstancedSubObjRef> PropertySubObjectReferences;
+	FFindInstancedReferenceSubobjectHelper::GetInstancedSubObjects(InObject, PropertySubObjectReferences);
 
 	TMap<UObject*, UObject*> OldToNewInstanceMap;
 	for (UObject* PropertySubObject : PropertySubObjectReferences)
@@ -132,7 +134,7 @@ void FixSubObjectReferencesPostUndoRedo(UObject* InObject)
 			OldToNewInstanceMap.Add(PropertySubObject, NewSubObject);
 		}
 	}
-
+PRAGMA_ENABLE_DEPRECATION_WARNINGS
 	FArchiveReplaceObjectRef<UObject> Replacer(InObject, OldToNewInstanceMap);
 }
 
