Minor refactor of legacy Blueprint editor post undo/redo fixup logic to avoid use of FFindInstancedReferenceSubobjectHelper (deprecated).

#rnx

diff --git a/Engine/Source/Editor/Kismet/Private/BlueprintEditorModule.cpp b/Engine/Source/Editor/Kismet/Private/BlueprintEditorModule.cpp
--- a/Engine/Source/Editor/Kismet/Private/BlueprintEditorModule.cpp
+++ b/Engine/Source/Editor/Kismet/Private/BlueprintEditorModule.cpp
@@ -22,7 +22,6 @@
 #include "HAL/PlatformCrt.h"
 #include "IMessageLogListing.h"
 #include "ISettingsModule.h"
-#include "InstancedReferenceSubobjectHelper.h"
 #include "InstancedStaticMeshSCSEditorCustomization.h"
 #include "Internationalization/Internationalization.h"
 #include "Kismet2/BlueprintEditorUtils.h"
@@ -97,13 +96,25 @@ void FixSubObjectReferencesPostUndoRedo(UObject* InObject)
 	GetObjectsWithOuter(InObject, SubObjects, bIncludeNestedSubObjects);
 
 	// Post undo/redo, these may have the in-correct Outer but are incorrectly referenced by the CDO's UProperties
-PRAGMA_DISABLE_DEPRECATION_WARNINGS
-	TSet<FInstancedSubObjRef> PropertySubObjectReferences;
-	FFindInstancedReferenceSubobjectHelper::GetInstancedSubObjects(InObject, PropertySubObjectReferences);
-
+	// @todo - The original use case (UE-18137) no longer seems to rely on this code path for fixup, can we remove it?
 	TMap<UObject*, UObject*> OldToNewInstanceMap;
-	for (UObject* PropertySubObject : PropertySubObjectReferences)
+	for (TPropertyValueIterator<FObjectProperty> It(InObject->GetClass(), InObject); It; ++It)
 	{
+		// This is a legacy path that was intentionally excluding non-"persistent" instanced properties. We'll keep the
+		// same logic in place to avoid unintentional side effects and to avoid introducing any new regressive edge cases.
+		const FObjectProperty* Property = It.Key();
+		if (!Property->HasAllPropertyFlags(CPF_InstancedReference | CPF_PersistentInstance))
+		{
+			continue;
+		}
+
+		// Ignore values that are NULL (e.g. child class overrides that exclude the property from instancing).
+		UObject* PropertySubObject = Property->GetObjectPropertyValue(It.Value());
+		if (!PropertySubObject)
+		{
+			continue;
+		}
+
 		bool bFoundMatchingSubObject = false;
 		for (UObject* SubObject : SubObjects)
 		{
@@ -134,7 +145,7 @@ PRAGMA_DISABLE_DEPRECATION_WARNINGS
 			OldToNewInstanceMap.Add(PropertySubObject, NewSubObject);
 		}
 	}
-PRAGMA_ENABLE_DEPRECATION_WARNINGS
+
 	FArchiveReplaceObjectRef<UObject> Replacer(InObject, OldToNewInstanceMap);
 }
 
