Fix crash when compiling a macro blueprint that does not yet have a class. Properly order blueprints that don't have a class so they compile after any other blueprints in the batch,
also avoid finding dependencies when a blueprint doesn't have a class

[RN] Fixed crash when creating a macro blueprint based on a blueprint that is not up to date

diff --git a/Engine/Source/Editor/Kismet/Private/BlueprintCompilationManager.cpp b/Engine/Source/Editor/Kismet/Private/BlueprintCompilationManager.cpp
--- a/Engine/Source/Editor/Kismet/Private/BlueprintCompilationManager.cpp
+++ b/Engine/Source/Editor/Kismet/Private/BlueprintCompilationManager.cpp
@@ -864,6 +864,23 @@ void FBlueprintCompilationManagerImpl::FlushCompilationQueueImpl(bool bSuppressB
 				return false;
 			}
 
+			// BPs that do not yet have classes should be compiled last,
+			// so that we can guarantee their parent class is compiled 
+			// before them. We know they do not yet have child types
+			// because they have no class themselves:
+			if(A.GeneratedClass == nullptr && B.GeneratedClass == nullptr)
+			{
+				return A.GetFName().LexicalLess(B.GetFName());
+			}
+			else if(A.GeneratedClass == nullptr)
+			{
+				return false;
+			}
+			else if(B.GeneratedClass == nullptr)
+			{
+				return true;
+			}
+
 			return FBlueprintCompileReinstancer::ReinstancerOrderingFunction(A.GeneratedClass, B.GeneratedClass);
 		};
 		CurrentlyCompilingBPs.Sort( HierarchyDepthSortFn );
@@ -3550,11 +3567,14 @@ void FBlueprintCompilationManagerImpl::QueueOutOfDateDependencies( const TArray<
 {
 	TArray<TObjectPtr<UBlueprint>> RootCompilationRequests;
 	// we don't care about 'skeleton only' regeneration, filter those, and operate
-	// only on 'full compilation' requests:
+	// only on 'full compilation' requests - we also don't need to gather
+	// dependencies when a BP does not yet have a generated class, as it has
+	// nothing anyone could be dependent upon:
 	Algo::TransformIf(QueuedRequests, RootCompilationRequests,
 		[](const FBPCompileRequest& CompileRequest) -> bool
 		{
 			return
+				CompileRequest.BPToCompile->GeneratedClass != nullptr &&
 				(CompileRequest.CompileOptions & EBlueprintCompileOptions::RegenerateSkeletonOnly)
 				== EBlueprintCompileOptions::None;
 		},
