Basic logic to recook data only BPs only when their parent class(es) observably change
#rnx

diff --git a/Engine/Source/Editor/Kismet/Internal/Blueprints/BlueprintDependencies.h b/Engine/Source/Editor/Kismet/Internal/Blueprints/BlueprintDependencies.h
new file mode 100644
--- /dev/null
+++ b/Engine/Source/Editor/Kismet/Internal/Blueprints/BlueprintDependencies.h
@@ -0,0 +1,36 @@
+// Copyright Epic Games, Inc. All Rights Reserved.
+
+#pragma once
+
+#include "Cooker/CookDependency.h"
+
+#include "BlueprintDependencies.generated.h"
+
+USTRUCT()
+struct FInheritanceDependency
+{
+	GENERATED_BODY()
+	
+	UPROPERTY()
+	FName PackageName;
+
+	UPROPERTY()
+	FName NativeClassName;
+
+	UPROPERTY()
+	FString Hash;
+};
+
+USTRUCT()
+struct FBlueprintDependencies
+{
+	GENERATED_BODY()
+
+	UPROPERTY()
+	TArray<FInheritanceDependency> InheritanceChain;
+};
+
+namespace UE::Private::BlueprintDependencies
+{
+KISMET_API UE::Cook::FCookDependency RecordCookDependencies(const UBlueprint* BP);
+}
diff --git a/Engine/Source/Editor/Kismet/Private/Blueprints/BlueprintDependencies.cpp b/Engine/Source/Editor/Kismet/Private/Blueprints/BlueprintDependencies.cpp
new file mode 100644
--- /dev/null
+++ b/Engine/Source/Editor/Kismet/Private/Blueprints/BlueprintDependencies.cpp
@@ -0,0 +1,260 @@
+// Copyright Epic Games, Inc. All Rights Reserved.
+
+#include "Blueprints/BlueprintDependencies.h"
+
+#include "Engine/Blueprint.h"
+#include "HAL/FileManager.h"
+#include "Misc/ScopeExit.h"
+#include "Serialization/CompactBinary.h"
+#include "Serialization/CompactBinaryWriter.h"
+#include "JsonObjectGraph/Stringify.h"
+
+/*
+	BLUEPRINT DEPENDENCIES DESIGN
+
+	This file is meant to capture external dependencies of a blueprint so that it can
+	be reliably recooked when its inputs change. It currently has only been tested
+	against data only blueprints, which represente default value changes made against
+	some parent class. This means that any time the parent class(es) change the blueprint
+	must be recooked. Because there is not an obvious artifact to identify for native
+	classes any native classes are serialized to json for the purpose of providing identity.
+
+	This captures, among other things, the class default values and custom versions used
+	in that native class(es) serialization routines. Non-native classes are checked for
+	identity using a blake3 hash of their source artifact. Both of these approaches have
+	drawbacks, and as these drawbacks are addressed I will look towards unifying native
+	or blueprint class dependencies.
+
+	We also want to extend this design to blueprints with bytecode, functions, extensions,
+	etc.
+*/
+
+namespace UE::Private
+{
+constexpr int32 BlueprintCookDependenciesVersion = 1;
+
+static bool IsNativeObject( const UObject* Obj);
+static FString HashClass(const UClass* Class);
+static FString HashNativeClass( const UClass* NativeClass);
+static FString HashPackage( const UPackage* Package );
+static FString HashFile(const FString& Filename);
+static bool ReadBytesFromFile(const TCHAR* Filename, TArray64<uint8>& OutBytes);
+
+static void ValidateBPCookDependenciesImpl(FCbFieldViewIterator Args, UE::Cook::FCookDependencyContext& Context);
+static void GenerateBlueprintDependencies(const UBlueprint* ForBP, FBlueprintDependencies& OutDependencies);
+static bool LoadBPCookDependenciesImpl(FCbFieldViewIterator Args, UE::Cook::FCookDependencyContext& Context, FBlueprintDependencies& OutDependencies);
+static void SaveBPCookDepenenciesImpl(FCbWriter& Writer, const FBlueprintDependencies& Dependencies);
+}
+
+static bool UE::Private::IsNativeObject( const UObject* Obj)
+{
+	return Obj->GetPackage()->HasAnyPackageFlags(PKG_CompiledIn);
+}
+
+static FString UE::Private::HashClass(const UClass* Class)
+{
+	// for native classes we have no artifact to hash so hash the class itself,
+	// for non-native classes just hash the artifact (.uasset, aka package):
+	if(Class->HasAnyClassFlags(CLASS_Native))
+	{
+		return HashNativeClass(Class);
+	}
+	else
+	{
+		return HashPackage(Class->GetPackage());
+	}
+}
+
+static FString UE::Private::HashNativeClass( const UClass* NativeClass)
+{
+	// We're using UE::JsonObjectGraph because it is robust compared to ExportText and 
+	// the other core level facilities. Exporting a class to text using ExportText has
+	// always produced an empty object, which is not meaningful. JsonObjectGraph::Stringify
+	// will provide identity for the UClass:
+	FUtf8String StringifiedClass = UE::JsonObjectGraph::Stringify(
+		{NativeClass, NativeClass->ClassDefaultObject}, 
+		FJsonStringifyOptions(EJsonStringifyFlags::FilterEditorOnlyData));
+	FBlake3 Hash;
+	Hash.Update(StringifiedClass.GetCharArray().GetData(), StringifiedClass.GetCharArray().Num());
+	return LexToString(Hash.Finalize());
+}
+
+static FString UE::Private::HashPackage(const UPackage* Package)
+{
+	const FName ParentPackageName = Package->GetFName();
+	// we don't want to have to load the asset so hash its file rather than the class itself:
+	FString Result;
+	const bool bSuceeded = FPackageName::TryConvertLongPackageNameToFilename(
+		ParentPackageName.ToString(), Result, FPackageName::GetAssetPackageExtension());
+	if(bSuceeded)
+	{
+		return UE::Private::HashFile(Result);
+	}
+	return FString();
+}
+
+static FString UE::Private::HashFile(const FString& Filename)
+{
+	TArray64<uint8> FileData;
+	if(!ReadBytesFromFile(*Filename, FileData))
+	{
+		return FString();
+	}
+
+	FBlake3 Hash;
+	Hash.Update(FileData.GetData(), FileData.Num());
+
+	return LexToString(Hash.Finalize());
+}
+
+static bool UE::Private::ReadBytesFromFile(const TCHAR* Filename, TArray64<uint8>& OutBytes)
+{
+	OutBytes.Reset();
+	TUniquePtr<FArchive> FileArchive(IFileManager::Get().CreateFileReader(Filename));
+	if(FileArchive == TUniquePtr<FArchive>())
+	{
+		return false;
+	}
+	
+	const int64 TotalSize = FileArchive->TotalSize();
+	if(TotalSize <= 0)
+	{
+		return false;
+	}
+	
+	OutBytes.AddUninitialized(TotalSize);
+	FileArchive->Serialize(OutBytes.GetData(), TotalSize);
+	return true;
+}
+
+// BEGIN PAIR VALIDATE/GENERATE
+void UE::Private::ValidateBPCookDependenciesImpl(FCbFieldViewIterator Args, UE::Cook::FCookDependencyContext& Context)
+{
+	// MUST MATCH UE::Private::GenerateBlueprintDependencies
+	FBlueprintDependencies Dependencies;
+	if(!LoadBPCookDependenciesImpl(Args, Context, Dependencies))
+	{
+		return;
+	}
+
+	// confirm dependencies do/do not match the artifact - must correspond to 
+	for(const FInheritanceDependency& InheritanceDependency : Dependencies.InheritanceChain)
+	{
+		// if the dependency is native hash the native class and confirm it matches:
+		const FName NativeClassName = InheritanceDependency.NativeClassName;
+		const FName PackageName = InheritanceDependency.PackageName;
+		if( NativeClassName != FName())
+		{
+			const UPackage* Package = FindObjectFast<UPackage>(nullptr, PackageName);
+			if(!Package)
+			{
+				Context.LogInvalidated(
+					FString::Printf(TEXT("Native Package Missing: %s"), *PackageName.ToString()));
+				return;
+			}
+			const UClass* Class = FindObjectFast<UClass>(
+				const_cast<UPackage*>(Package), NativeClassName);
+			if(!Class)
+			{
+				Context.LogInvalidated(
+					FString::Printf(TEXT("Native Class Missing: %s"), *NativeClassName.ToString()));
+				return;
+			}
+			if(InheritanceDependency.Hash != UE::Private::HashClass(Class))
+			{
+				Context.LogInvalidated(FString::Printf(TEXT("Native Class Changed: %s"), *Class->GetPathName()));
+				return;
+			}
+		}
+		// else if the dependency is from a package hash that package and confirm it matches:
+		else
+		{
+			FString PackageFilename;
+			const bool bSuceeded = FPackageName::TryConvertLongPackageNameToFilename(
+				PackageName.ToString(), PackageFilename,  FPackageName::GetAssetPackageExtension());
+			check(bSuceeded);
+			if(InheritanceDependency.Hash != UE::Private::HashFile(PackageFilename))
+			{
+				Context.LogInvalidated(FString::Printf(TEXT("Package Changed: %s"), *PackageFilename));
+				return;
+			}
+		}
+	}
+}
+UE_COOK_DEPENDENCY_FUNCTION(ValidateBPCookDependenciesImpl, UE::Private::ValidateBPCookDependenciesImpl);
+
+void UE::Private::GenerateBlueprintDependencies(const UBlueprint* ForBP, FBlueprintDependencies& OutDependencies)
+{
+	// MUST MATCH UE::Private::ValidateBPCookDependenciesImpl
+	const UClass* Hierarchy = ForBP->GeneratedClass;
+	while(Hierarchy)
+	{
+		FInheritanceDependency Dependency = {
+			.PackageName = Hierarchy->GetPackage()->GetFName(),
+			// class native branch sets NativeClassName, which controls class hashing 
+			// behavior here and above in ValidateBPCookDependenciesImpl
+			.NativeClassName = Hierarchy->HasAnyClassFlags(CLASS_Native) ? Hierarchy->GetFName() : FName(),
+			.Hash = UE::Private::HashClass(Hierarchy)
+		};
+		OutDependencies.InheritanceChain.Emplace(
+			MoveTemp(Dependency)
+		);
+		Hierarchy = Hierarchy->GetSuperClass();
+	}
+}
+// END PAIR VALIDATE/GENERATE
+
+// BEGIN PAIR LOAD/SAVE
+bool UE::Private::LoadBPCookDependenciesImpl(FCbFieldViewIterator Args, UE::Cook::FCookDependencyContext& Context, FBlueprintDependencies& OutDependencies)
+{
+	// MUST MATCH UE::Private::SaveBPCookDepenenciesImpl:
+	FCbFieldViewIterator ArgField(Args);
+	const int32 ArgsVersion = (ArgField++).AsInt32(); // Writer << BlueprintCookDependenciesVersion;
+	if (ArgsVersion != BlueprintCookDependenciesVersion)
+	{
+		Context.LogInvalidated(TEXT("Blueprint Cook Dependency Version Changed"));
+		return false;
+	}
+	
+	FUtf8StringView DependencyData = (ArgField++).AsString(); // Writer << DependenciesSerialized;
+	FString AsFString(DependencyData);
+	FStringOutputDevice Errors;
+	FBlueprintDependencies::StaticStruct()->ImportText(
+		*AsFString, &OutDependencies, nullptr, 0, &Errors, []() {return FString(); });
+	if (!Errors.IsEmpty())
+	{
+		Context.LogError(FString::Printf(TEXT("Could not load Blueprint dependencies: %s"), *Errors));
+		return false;
+	}
+
+	return true;
+}
+
+void UE::Private::SaveBPCookDepenenciesImpl(FCbWriter& Writer, const FBlueprintDependencies& Dependencies)
+{
+	// MUST MATCH UE::Private::LoadBPCookDependenciesImpl
+	FString DependenciesSerialized;
+	FBlueprintDependencies::StaticStruct()->ExportText(
+		DependenciesSerialized, &Dependencies, nullptr, nullptr, 0, nullptr);
+
+	Writer << BlueprintCookDependenciesVersion;
+	Writer << DependenciesSerialized;
+}
+// END PAIR LOAD/SAVE
+
+UE::Cook::FCookDependency UE::Private::BlueprintDependencies::RecordCookDependencies(const UBlueprint* BP)
+{
+	if(!BP->ParentClass)
+	{
+		return UE::Cook::FCookDependency();
+	}
+
+	FBlueprintDependencies Deps;
+	UE::Private::GenerateBlueprintDependencies(BP, Deps);
+
+	FCbWriter Writer;
+	UE::Private::SaveBPCookDepenenciesImpl(Writer, Deps);
+
+	return UE::Cook::FCookDependency::Function(
+		UE_COOK_DEPENDENCY_FUNCTION_CALL(ValidateBPCookDependenciesImpl), Writer.Save());
+}
diff --git a/Engine/Source/Runtime/Engine/Classes/Engine/Blueprint.h b/Engine/Source/Runtime/Engine/Classes/Engine/Blueprint.h
--- a/Engine/Source/Runtime/Engine/Classes/Engine/Blueprint.h
+++ b/Engine/Source/Runtime/Engine/Classes/Engine/Blueprint.h
@@ -898,6 +898,9 @@ private:
 	/** Sets the current object being debugged */
 	ENGINE_API void DebuggingWorldRegistrationHelper(UObject* ObjectProvidingWorld, UObject* ValueToRegister);
 	
+#if WITH_EDITORONLY_DATA
+	void PreSaveCook(FObjectPreSaveContext ObjectSaveContext);
+#endif
 public:
 
 	/** @return the current object being debugged, which can be nullptr */
diff --git a/Engine/Source/Runtime/Engine/Private/Blueprint.cpp b/Engine/Source/Runtime/Engine/Private/Blueprint.cpp
--- a/Engine/Source/Runtime/Engine/Private/Blueprint.cpp
+++ b/Engine/Source/Runtime/Engine/Private/Blueprint.cpp
@@ -17,6 +17,7 @@
 
 #if WITH_EDITOR
 #include "BlueprintCompilationManager.h"
+#include "Blueprints/BlueprintDependencies.h"
 #include "Engine/SCS_Node.h"
 #include "Kismet2/BlueprintEditorUtils.h"
 #include "Kismet2/ComponentEditorUtils.h"
@@ -398,6 +399,13 @@ void UBlueprint::PreSave(FObjectPreSaveContext ObjectSaveContext)
 		// Cache the BP for use (immediate, since we're about to save)
 		FFindInBlueprintSearchManager::Get().AddOrUpdateBlueprintSearchMetadata(this, Flags, OverrideVersion);
 	}
+	
+#if WITH_EDITORONLY_DATA
+	if (ObjectSaveContext.IsCooking())
+	{
+		PreSaveCook(ObjectSaveContext);
+	}
+#endif
 }
 
 void UBlueprint::GetPreloadDependencies(TArray<UObject*>& OutDeps)
@@ -864,6 +872,35 @@ void UBlueprint::DebuggingWorldRegistrationHelper(UObject* ObjectProvidingWorld,
 	}
 }
 
+#if WITH_EDITORONLY_DATA
+void UBlueprint::PreSaveCook(FObjectPreSaveContext ObjectSaveContext)
+{
+	using namespace UE::Cook;
+	const UClass* const CurrentGeneratedClass = GeneratedClass;
+	if (!CurrentGeneratedClass)
+	{
+		return; // macro blueprints, blueprints in invalid state
+	}
+	const UClass* const SuperClass = CurrentGeneratedClass->GetSuperClass();
+	if(!SuperClass)
+	{
+		return; // does this ever happen? unlikely
+	}
+
+	// We don't yet support gathering all kinds of dependencies (e.g. called functions, 
+	// read/written properties, delegate signatures, structs, API) but DOBP
+	// only depend on their inherited values so it should be relatively easy to 
+	// record dependencies for them:
+	if (!FBlueprintEditorUtils::IsDataOnlyBlueprint(this))
+	{
+		return;
+	}
+
+	ObjectSaveContext.AddCookBuildDependency(
+		UE::Private::BlueprintDependencies::RecordCookDependencies(this));
+}
+#endif
+
 UClass* UBlueprint::GetBlueprintClass() const
 {
 	return UBlueprintGeneratedClass::StaticClass();
