Explicitly record complicated or invalid bps as not cacheable, to distinguish them from trivial assets w/ no dependencies

#rnx

diff --git a/Engine/Source/Editor/Kismet/Private/Blueprints/BlueprintDependencies.cpp b/Engine/Source/Editor/Kismet/Private/Blueprints/BlueprintDependencies.cpp
--- a/Engine/Source/Editor/Kismet/Private/Blueprints/BlueprintDependencies.cpp
+++ b/Engine/Source/Editor/Kismet/Private/Blueprints/BlueprintDependencies.cpp
@@ -4,6 +4,7 @@
 
 #include "Engine/Blueprint.h"
 #include "HAL/FileManager.h"
+#include "Kismet2/BlueprintEditorUtils.h"
 #include "Misc/ScopeExit.h"
 #include "Serialization/CompactBinary.h"
 #include "Serialization/CompactBinaryWriter.h"
@@ -40,6 +41,7 @@ static FString HashPackage( const UPackage* Package );
 static FString HashFile(const FString& Filename);
 static bool ReadBytesFromFile(const TCHAR* Filename, TArray64<uint8>& OutBytes);
 
+static void NotCacheable(FCbFieldViewIterator Args, UE::Cook::FCookDependencyContext& Context);
 static void ValidateBPCookDependenciesImpl(FCbFieldViewIterator Args, UE::Cook::FCookDependencyContext& Context);
 static void GenerateBlueprintDependencies(const UBlueprint* ForBP, FBlueprintDependencies& OutDependencies);
 static bool LoadBPCookDependenciesImpl(FCbFieldViewIterator Args, UE::Cook::FCookDependencyContext& Context, FBlueprintDependencies& OutDependencies);
@@ -127,6 +129,13 @@ static bool UE::Private::ReadBytesFromFile(const TCHAR* Filename, TArray64<uint8
 	return true;
 }
 
+void UE::Private::NotCacheable(FCbFieldViewIterator Args, UE::Cook::FCookDependencyContext& Context)
+{	
+	Context.LogInvalidated(
+		TEXT("Package is not yet cacheable"));
+}
+UE_COOK_DEPENDENCY_FUNCTION(NotCacheable, UE::Private::NotCacheable);
+
 // BEGIN PAIR VALIDATE/GENERATE
 void UE::Private::ValidateBPCookDependenciesImpl(FCbFieldViewIterator Args, UE::Cook::FCookDependencyContext& Context)
 {
@@ -244,9 +253,14 @@ void UE::Private::SaveBPCookDepenenciesImpl(FCbWriter& Writer, const FBlueprintD
 
 UE::Cook::FCookDependency UE::Private::BlueprintDependencies::RecordCookDependencies(const UBlueprint* BP)
 {
-	if(!BP->ParentClass)
+	if(!BP->ParentClass || 
+		!BP->GeneratedClass ||
+		!BP->GeneratedClass->IsChildOf(UObject::StaticClass()) ||
+		!FBlueprintEditorUtils::IsDataOnlyBlueprint(BP) )
 	{
-		return UE::Cook::FCookDependency();
+		// These BPs cannot yet be cached:
+		return UE::Cook::FCookDependency::Function(
+			UE_COOK_DEPENDENCY_FUNCTION_CALL(NotCacheable), FCbFieldIterator());
 	}
 
 	FBlueprintDependencies Deps;
diff --git a/Engine/Source/Runtime/Engine/Private/Blueprint.cpp b/Engine/Source/Runtime/Engine/Private/Blueprint.cpp
--- a/Engine/Source/Runtime/Engine/Private/Blueprint.cpp
+++ b/Engine/Source/Runtime/Engine/Private/Blueprint.cpp
@@ -876,26 +876,6 @@ void UBlueprint::DebuggingWorldRegistrationHelper(UObject* ObjectProvidingWorld,
 void UBlueprint::PreSaveCook(FObjectPreSaveContext ObjectSaveContext)
 {
 	using namespace UE::Cook;
-	const UClass* const CurrentGeneratedClass = GeneratedClass;
-	if (!CurrentGeneratedClass)
-	{
-		return; // macro blueprints, blueprints in invalid state
-	}
-	const UClass* const SuperClass = CurrentGeneratedClass->GetSuperClass();
-	if(!SuperClass)
-	{
-		return; // does this ever happen? unlikely
-	}
-
-	// We don't yet support gathering all kinds of dependencies (e.g. called functions, 
-	// read/written properties, delegate signatures, structs, API) but DOBP
-	// only depend on their inherited values so it should be relatively easy to 
-	// record dependencies for them:
-	if (!FBlueprintEditorUtils::IsDataOnlyBlueprint(this))
-	{
-		return;
-	}
-
 	ObjectSaveContext.AddCookBuildDependency(
 		UE::Private::BlueprintDependencies::RecordCookDependencies(this));
 }
