Expose method to get the ModuleDescriptor from a UPackage asset.

diff --git a/Engine/Source/Editor/Kismet/Internal/ProjectUtilities/BuildTargetSet.h b/Engine/Source/Editor/Kismet/Internal/ProjectUtilities/BuildTargetSet.h
--- a/Engine/Source/Editor/Kismet/Internal/ProjectUtilities/BuildTargetSet.h
+++ b/Engine/Source/Editor/Kismet/Internal/ProjectUtilities/BuildTargetSet.h
@@ -5,8 +5,10 @@
 #include "Containers/UnrealString.h"
 #include "Misc/EnumClassFlags.h"
 
+struct FModuleDescriptor;
 class UClass;
 class UFunction;
+class UPackage;
 
 namespace UE::ProjectUtilities
 {
@@ -52,4 +54,10 @@ namespace UE::ProjectUtilities
 		EBuildTargetFlags BuildTargetFlags = EBuildTargetFlags::None;
 	};
 	ENUM_CLASS_FLAGS(FBuildTargetSet::EBuildTargetFlags);
+
+	/**
+	 * Given a native UPackage attempts to find the module descriptor.
+	 * Returns nullptr if it can not be found.
+	 */
+	KISMET_API const FModuleDescriptor* FindModuleDescriptor(const UPackage* ForNativePackage);
 } 
diff --git a/Engine/Source/Editor/Kismet/Private/ProjectUtilities/BuildTargetSet.cpp b/Engine/Source/Editor/Kismet/Private/ProjectUtilities/BuildTargetSet.cpp
--- a/Engine/Source/Editor/Kismet/Private/ProjectUtilities/BuildTargetSet.cpp
+++ b/Engine/Source/Editor/Kismet/Private/ProjectUtilities/BuildTargetSet.cpp
@@ -1,19 +1,18 @@
 // Copyright Epic Games, Inc. All Rights Reserved.
 
 #include "ProjectUtilities/BuildTargetSet.h"
+
 #include "Interfaces/IProjectManager.h"
 #include "Interfaces/IPluginManager.h"
-#include "ProjectDescriptor.h"
 #include "Kismet2/BlueprintEditorUtils.h"
+#include "Logging/StructuredLog.h"
+#include "ProfilingDebugging/ScopedTimers.h"
+#include "ProjectDescriptor.h"
 
 namespace UE::Private
 {
-	static TMap<FName, const FModuleDescriptor*> DescriptorCache;
-	static FCriticalSection DescriptorCacheLock;
-
 	const UClass* GetNativeClass(const UFunction* ForFunction);
 	const UClass* GetNativeClass(const UClass* ForClass);
-	const FModuleDescriptor* FindModuleDescriptor(const UPackage* ForNativePackage);
 }
 
 const UClass* UE::Private::GetNativeClass(const UFunction* ForFunction)
@@ -33,8 +32,19 @@ const UClass* UE::Private::GetNativeClass(const UClass* ForClass)
 	return FBlueprintEditorUtils::FindFirstNativeClass(const_cast<UClass*>(ForClass));
 }
 
-const FModuleDescriptor* UE::Private::FindModuleDescriptor(const UPackage* ForNativePackage)
+
+namespace UE::ProjectUtilities
+{
+	static TMap<FName, const FModuleDescriptor*> DescriptorCache;
+	static FCriticalSection DescriptorCacheLock;
+
+const FModuleDescriptor* FindModuleDescriptor(const UPackage* ForNativePackage)
 {
+	if (!ForNativePackage)
+	{
+		return nullptr;
+	}
+
 	// The loops at the bottom are naive and may need to be optimized if we begin
 	// validating build target compatibility widely:
 	FString FunctionPackageName = ForNativePackage->GetName();
@@ -51,33 +61,33 @@ const FModuleDescriptor* UE::Private::FindModuleDescriptor(const UPackage* ForNa
 
 	const IProjectManager& ProjectManager = IProjectManager::Get();
 	const FProjectDescriptor* PD = ProjectManager.GetCurrentProject();
-	if(!PD)
+	if (!PD)
 	{
 		return nullptr;
 	}
-	
+
 	const FName FunctionPackageFName(FunctionPackageName.Len(), *FunctionPackageName);
 	const auto AddToCache = [](FName Name, const FModuleDescriptor* Descriptor)
-	{
-		// must be a write lock:
-		FScopeLock WriteLock(&DescriptorCacheLock);
-		DescriptorCache.Add(Name, Descriptor);
-	};
+		{
+			// must be a write lock:
+			FScopeLock WriteLock(&DescriptorCacheLock);
+			DescriptorCache.Add(Name, Descriptor);
+		};
 	const auto FindFromCache = [](FName Name) -> TOptional<const FModuleDescriptor*>
-	{
-		// could be a read lock:
-		FScopeLock FindLock(&DescriptorCacheLock);
-		const FModuleDescriptor** Result = DescriptorCache.Find(Name);
-		return Result ? *Result : TOptional<const FModuleDescriptor*>();
-	};
+		{
+			// could be a read lock:
+			FScopeLock FindLock(&DescriptorCacheLock);
+			const FModuleDescriptor** Result = DescriptorCache.Find(Name);
+			return Result ? *Result : TOptional<const FModuleDescriptor*>();
+		};
 	TOptional<const FModuleDescriptor*> CachedDescriptor = FindFromCache(FunctionPackageFName);
-	if(CachedDescriptor)
+	if (CachedDescriptor)
 	{
 		return *CachedDescriptor;
 	}
 
 	const FModuleDescriptor* OwningMD = nullptr;
-	
+
 	for (const FModuleDescriptor& MD : PD->Modules)
 	{
 		if (MD.Name == FunctionPackageFName)
@@ -106,17 +116,11 @@ const FModuleDescriptor* UE::Private::FindModuleDescriptor(const UPackage* ForNa
 			}
 		}
 	}
-	
+
 	AddToCache(FunctionPackageFName, nullptr);
 	return nullptr;
 }
 
-#include "Logging/StructuredLog.h"
-#include "ProfilingDebugging/ScopedTimers.h"
-
-namespace UE::ProjectUtilities
-{
-
 FBuildTargetSet FBuildTargetSet::GetCallerTargetsUnsupportedByCallee(const UClass* Caller, const UFunction* Callee)
 {
 	using namespace UE::Private;
@@ -184,7 +188,7 @@ FBuildTargetSet::EBuildTargetFlags FBuildTargetSet::GetSupportedTargetsForNative
 		| EBuildTargetFlags::Client
 		| EBuildTargetFlags::Editor;
 	using namespace UE::Private;
-	const FModuleDescriptor* ModuleDescriptor = FindModuleDescriptor(NativeBase->GetPackage());
+	const FModuleDescriptor* ModuleDescriptor = UE::ProjectUtilities::FindModuleDescriptor(NativeBase->GetPackage());
 	EBuildTargetFlags SupportedTargets = ALL;
 	if (ModuleDescriptor)
 	{
