Avoid queuing refreshes when loads are in progress

[RN] Fixed performance problem when booting the editor with an open blueprint editor when the palette tab is active

diff --git a/Engine/Source/Editor/Kismet/Private/SBlueprintSubPalette.cpp b/Engine/Source/Editor/Kismet/Private/SBlueprintSubPalette.cpp
--- a/Engine/Source/Editor/Kismet/Private/SBlueprintSubPalette.cpp
+++ b/Engine/Source/Editor/Kismet/Private/SBlueprintSubPalette.cpp
@@ -2,6 +2,7 @@
 
 #include "SBlueprintSubPalette.h"
 
+#include "AssetRegistry/AssetRegistryModule.h"
 #include "BPDelegateDragDropAction.h"
 #include "BPFunctionDragDropAction.h"
 #include "BPVariableDragDropAction.h"
@@ -305,14 +306,6 @@ void SBlueprintSubPalette::Construct(FArguments const& InArgs, TWeakPtr<FBluepri
 	ActionDatabase.OnEntryUpdated().AddSP(this, &SBlueprintSubPalette::OnDatabaseActionsUpdated);
 }
 
-//------------------------------------------------------------------------------
-EActiveTimerReturnType SBlueprintSubPalette::TriggerRefreshActionsList(double InCurrentTime, float InDeltaTime)
-{
-	RefreshActionsList(true);
-	bIsActiveTimerRegistered = false;
-	return EActiveTimerReturnType::Stop;
-}
-
 //------------------------------------------------------------------------------
 UBlueprint* SBlueprintSubPalette::GetBlueprint() const
 {
@@ -436,11 +429,7 @@ void SBlueprintSubPalette::GenerateContextMenuEntries(FMenuBuilder& MenuBuilder)
 //------------------------------------------------------------------------------
 void SBlueprintSubPalette::RequestRefreshActionsList()
 {
-	if (!bIsActiveTimerRegistered)
-	{
-		bIsActiveTimerRegistered = true;
-		RegisterActiveTimer(0.f, FWidgetActiveTimerDelegate::CreateSP(this, &SBlueprintSubPalette::TriggerRefreshActionsList));
-	}
+	bIsActiveTimerRegistered = true;
 }
 
 //------------------------------------------------------------------------------
@@ -468,6 +457,21 @@ void SBlueprintSubPalette::OnDatabaseActionsRemoved(UObject* ActionsKey)
 	}
 }
 
+void SBlueprintSubPalette::Tick(const FGeometry& AllottedGeometry, const double InCurrentTime, const float InDeltaTime)
+{
+	SGraphPalette::Tick(AllottedGeometry, InCurrentTime, InDeltaTime);
+
+	if(bIsActiveTimerRegistered)
+	{
+		static FAssetRegistryModule& AssetRegistryModule = FModuleManager::LoadModuleChecked<FAssetRegistryModule>(TEXT("AssetRegistry"));
+		if (!AssetRegistryModule.Get().IsLoadingAssets() && !AssetRegistryModule.Get().IsGathering() && !IsAsyncLoading())
+		{
+			bIsActiveTimerRegistered = false;
+			RefreshActionsList(true);
+		}
+	}
+}
+
 /*******************************************************************************
 * Private SBlueprintSubPalette Methods
 *******************************************************************************/
diff --git a/Engine/Source/Editor/Kismet/Private/SBlueprintSubPalette.h b/Engine/Source/Editor/Kismet/Private/SBlueprintSubPalette.h
--- a/Engine/Source/Editor/Kismet/Private/SBlueprintSubPalette.h
+++ b/Engine/Source/Editor/Kismet/Private/SBlueprintSubPalette.h
@@ -111,6 +111,8 @@ protected:
 	void OnDatabaseActionsUpdated(UObject* ActionsKey);
 	void OnDatabaseActionsRemoved(UObject* ActionsKey);
 
+	virtual void Tick(const FGeometry& AllottedGeometry, const double InCurrentTime, const float InDeltaTime) override;
+
 	/** Pointer back to the blueprint editor that owns us */
 	TWeakPtr<FBlueprintEditor> BlueprintEditorPtr;
 
@@ -118,9 +120,6 @@ private:
 	/** Immediately calls RefreshActionsList(), doesn't not defer until Tick() like RequestRefreshActionsList() does. */
 	void ForceRefreshActionList();
 
-	/** One-off active timer to trigger a refresh of the action list */
-	EActiveTimerReturnType TriggerRefreshActionsList(double InCurrentTime, float InDeltaTime);
-
 	/**
 	 * Constructs a slate widget for the right-click context menu in this 
 	 * palette. While this isn't virtual, sub-classes can override GenerateContextMenuEntries()
