Refresh the pins and clear the DefaultValuePinWidget in the details menu when the graph changes

[RN] Fixed crash when editing a blueprint math expression node

diff --git a/Engine/Source/Editor/Kismet/Private/BlueprintDetailsCustomization.cpp b/Engine/Source/Editor/Kismet/Private/BlueprintDetailsCustomization.cpp
--- a/Engine/Source/Editor/Kismet/Private/BlueprintDetailsCustomization.cpp
+++ b/Engine/Source/Editor/Kismet/Private/BlueprintDetailsCustomization.cpp
@@ -3633,6 +3633,14 @@ static bool ShouldAllowWildcard(UK2Node_EditablePinBase* TargetNode)
 	return false;
 }
 
+FBlueprintGraphArgumentLayout::~FBlueprintGraphArgumentLayout()
+{
+	if (IsValid(TargetNode) && TargetNode->HasValidBlueprint() && GraphChangedHandler != FDelegateHandle())
+	{
+		TargetNode->GetGraph()->RemoveOnGraphChangedHandler(GraphChangedHandler);
+	}
+}
+
 void FBlueprintGraphArgumentLayout::GenerateHeaderRowContent( FDetailWidgetRow& NodeRow )
 {
 	const UEdGraphSchema_K2* K2Schema = GetDefault<UEdGraphSchema_K2>();
@@ -3718,6 +3726,7 @@ void FBlueprintGraphArgumentLayout::GenerateHeaderRowContent( FDetailWidgetRow&
 
 void FBlueprintGraphArgumentLayout::GenerateChildContent( IDetailChildrenBuilder& ChildrenBuilder )
 {
+	bool bMadePinWidgets = false;
 	if (bHasDefaultValue)
 	{
 		UEdGraphPin* FoundPin = GetPin();
@@ -3732,6 +3741,7 @@ void FBlueprintGraphArgumentLayout::GenerateChildContent( IDetailChildrenBuilder
 
 			if (!FoundPin->PinType.bIsReference && !bTypeWithNoDefaults)
 			{
+				bMadePinWidgets = true;
 				DefaultValuePinWidget = FNodeFactory::CreatePinWidget(FoundPin);
 				DefaultValuePinWidget->SetOnlyShowDefaultValue(true);
 				TSharedRef<SWidget> DefaultValueWidget = DefaultValuePinWidget->GetDefaultValueWidget();
@@ -3819,7 +3829,15 @@ void FBlueprintGraphArgumentLayout::GenerateChildContent( IDetailChildrenBuilder
 		}
 	}
 		
-	
+	if(bMadePinWidgets)
+	{
+		// add a listener to the owning graph so we can clean up our pin widgets:
+		if (TargetNode && TargetNode->HasValidBlueprint())
+		{
+			GraphChangedHandler = TargetNode->GetGraph()->AddOnGraphChangedHandler(
+				FOnGraphChanged::FDelegate::CreateSP(this, &FBlueprintGraphArgumentLayout::OnArgumentGraphChanged));
+		}
+	}
 }
 
 void FBlueprintGraphArgumentLayout::OnRemoveClicked()
@@ -3844,6 +3862,16 @@ void FBlueprintGraphArgumentLayout::OnRemoveClicked()
 	}
 }
 
+void FBlueprintGraphArgumentLayout::OnArgumentGraphChanged(const FEdGraphEditAction& Action)
+{
+	TSharedPtr<FBaseBlueprintGraphActionDetails> GraphActionDetails = GraphActionDetailsPtr.Pin();
+	if (GraphActionDetails.IsValid())
+	{
+		GraphActionDetails->OnPostEditorRefresh();
+		DefaultValuePinWidget.Reset();
+	}
+}
+
 bool FBlueprintGraphArgumentLayout::ShouldPinBeReadOnly(bool bIsEditingPinType/* = false*/) const
 {
 	if (TargetNode && ParamItemPtr.IsValid())
diff --git a/Engine/Source/Editor/Kismet/Private/BlueprintDetailsCustomization.h b/Engine/Source/Editor/Kismet/Private/BlueprintDetailsCustomization.h
--- a/Engine/Source/Editor/Kismet/Private/BlueprintDetailsCustomization.h
+++ b/Engine/Source/Editor/Kismet/Private/BlueprintDetailsCustomization.h
@@ -526,6 +526,7 @@ public:
 		, TargetNode(InTargetNode)
 		, bHasDefaultValue(bInHasDefaultValue)
 		, ArgumentName(InArgName) {}
+	virtual ~FBlueprintGraphArgumentLayout();
 
 private:
 	/** IDetailCustomNodeBuilder Interface*/
@@ -547,6 +548,8 @@ private:
 	/** Callbacks for all the functionality for modifying arguments */
 	void OnRemoveClicked();
 
+	void OnArgumentGraphChanged(const FEdGraphEditAction& Action);
+
 	FText OnGetArgNameText() const;
 	FText OnGetArgToolTipText() const;
 	void OnArgNameChange(const FText& InNewText);
@@ -591,6 +594,9 @@ private:
 
 	/** The SGraphPin widget created to show/edit default value */
 	TSharedPtr<SGraphPin> DefaultValuePinWidget;
+
+	/** Delegate handle for node changes */
+	FDelegateHandle GraphChangedHandler;
 };
 
 // local variables in the details view only support read-only mode at the moment.
