Don't map subobject archetypes until they've been reinstanced during a batched Blueprint compile on load.

diff --git a/Engine/Source/Editor/Kismet/Private/BlueprintCompilationManager.cpp b/Engine/Source/Editor/Kismet/Private/BlueprintCompilationManager.cpp
--- a/Engine/Source/Editor/Kismet/Private/BlueprintCompilationManager.cpp
+++ b/Engine/Source/Editor/Kismet/Private/BlueprintCompilationManager.cpp
@@ -136,8 +136,8 @@ struct FBlueprintCompilationManagerImpl : public FGCObject
 	void VerifyNoQueuedRequests(const TArray<FCompilerData>& CurrentlyCompilingBPs);
 
 	static void ReparentHierarchies(const TMap<UClass*, UClass*>& OldClassToNewClass, EReparentClassOptions Options);
-	static void BuildDSOMap(UObject* OldObject, UObject* NewObject, TMap<UObject*, UObject*>& OutOldToNewDSO);
-	static void ReinstanceBatch(TArray<FReinstancingJob>& Reinstancers, TMap< UClass*, UClass* >& InOutOldToNewClassMap, FUObjectSerializeContext* InLoadContext, TMap<UClass*, TMap<UObject*, UObject*>>* OldToNewTemplates = nullptr);
+	static void BuildDSOMap(UObject* OldObject, UObject* NewObject, const TMap<UClass*, UClass*>& InOldToNewClassMap, TMap<UObject*, UObject*>& OutOldToNewDSO);
+	static void ReinstanceBatch(TArray<FReinstancingJob>& Reinstancers, TMap<UClass*, UClass*>& InOutOldToNewClassMap, FUObjectSerializeContext* InLoadContext, TMap<UClass*, TMap<UObject*, UObject*>>* OldToNewTemplates = nullptr);
 	static UClass* FastGenerateSkeletonClass(UBlueprint* BP, FKismetCompilerContext& CompilerContext, bool bIsSkeletonOnly, TArray<FSkeletonFixupData>& OutSkeletonFixupData);
 	static bool IsQueuedForCompilation(UBlueprint* BP);
 	static void ConformToParentAndInterfaces(UBlueprint* BP);
@@ -2223,7 +2223,7 @@ void FBlueprintCompilationManagerImpl::ReparentHierarchies(const TMap<UClass*, U
 }
 
 
-void FBlueprintCompilationManagerImpl::BuildDSOMap(UObject* OldObject, UObject* NewObject, TMap<UObject*, UObject*>& OutOldToNewDSO)
+void FBlueprintCompilationManagerImpl::BuildDSOMap(UObject* OldObject, UObject* NewObject, const TMap<UClass*, UClass*>& InOldToNewClassMap, TMap<UObject*, UObject*>& OutOldToNewDSO)
 {
 	// IsDefaultSubObject() unfortunately cannot be relied upon for archetypes, so we explicitly search for the flag:
 	TArray<UObject*> OldSubobjects;
@@ -2238,10 +2238,18 @@ void FBlueprintCompilationManagerImpl::BuildDSOMap(UObject* OldObject, UObject*
 	for(UObject* OldSubobject : OldSubobjects)
 	{
 		UObject* NewSubobject = NewObject ? StaticFindObjectFast(UObject::StaticClass(), NewObject, OldSubobject->GetFName()) : nullptr;
+
+		// Don't map an archetype instance if we haven't yet reinstanced it (this can happen during batched compile/reinstance operations).
+		// Archetype instances (along with their own default subobjects) are added to the instance map elsewhere, when they are reinstanced.
+		if (NewSubobject && InOldToNewClassMap.Contains(NewSubobject->GetClass()))
+		{
+			continue;
+		}
+
 		// It may seem aggressive to ensure here, but we should always have a new version of the DSO. If that's
 		// not the case then some testing will need to be done on client of OutOldToNewDSO
 		OutOldToNewDSO.Add(OldSubobject, NewSubobject);
-		BuildDSOMap(OldSubobject, NewSubobject, OutOldToNewDSO);
+		BuildDSOMap(OldSubobject, NewSubobject, InOldToNewClassMap, OutOldToNewDSO);
 	}
 }
 
@@ -2656,7 +2664,7 @@ void FBlueprintCompilationManagerImpl::ReinstanceBatch(TArray<FReinstancingJob>&
 			{
 				OldArchetypeToNewArchetype.Add(OldClass->ClassDefaultObject, NewClass->ClassDefaultObject);
 				// also map old *default* subobjects to new default subobjects:
-				BuildDSOMap(OldClass->ClassDefaultObject, NewClass->ClassDefaultObject, OldArchetypeToNewArchetype);
+				BuildDSOMap(OldClass->ClassDefaultObject, NewClass->ClassDefaultObject, InOutOldToNewClassMap, OldArchetypeToNewArchetype);
 			}
 
 			TArray<UObject*> ArchetypeObjects;
@@ -2833,7 +2841,7 @@ void FBlueprintCompilationManagerImpl::ReinstanceBatch(TArray<FReinstancingJob>&
 				}
 
 				// also map old *default* subobjects to new default subobjects:
-				BuildDSOMap(OldInstance, NewArchetype, OldArchetypeToNewArchetype);
+				BuildDSOMap(OldInstance, NewArchetype, InOutOldToNewClassMap, OldArchetypeToNewArchetype);
 
 				ArchetypeReferencers.Add(NewArchetype);
 
