Avoid gathering dependencies of blueprints that have no CDO, as they will never have been instantiated. This specifically avoids a crash when creating a new blueprint when a parent class is in the BS_Dirty state.
#rnx

diff --git a/Engine/Source/Editor/Kismet/Private/BlueprintCompilationManager.cpp b/Engine/Source/Editor/Kismet/Private/BlueprintCompilationManager.cpp
--- a/Engine/Source/Editor/Kismet/Private/BlueprintCompilationManager.cpp
+++ b/Engine/Source/Editor/Kismet/Private/BlueprintCompilationManager.cpp
@@ -3590,13 +3590,14 @@ void FBlueprintCompilationManagerImpl::QueueOutOfDateDependencies( const TArray<
 	TArray<TObjectPtr<UBlueprint>> RootCompilationRequests;
 	// we don't care about 'skeleton only' regeneration, filter those, and operate
 	// only on 'full compilation' requests - we also don't need to gather
-	// dependencies when a BP does not yet have a generated class, as it has
+	// dependencies when a BP does not yet have a generated class or CDO, as it has
 	// nothing anyone could be dependent upon:
 	Algo::TransformIf(QueuedRequests, RootCompilationRequests,
 		[](const FBPCompileRequest& CompileRequest) -> bool
 		{
 			return
 				CompileRequest.BPToCompile->GeneratedClass != nullptr &&
+				CompileRequest.BPToCompile->GeneratedClass->ClassDefaultObject != nullptr &&
 				(CompileRequest.CompileOptions & EBlueprintCompileOptions::RegenerateSkeletonOnly)
 				== EBlueprintCompileOptions::None;
 		},
