Fix potential for a new race condition during async Blueprint indexing (regression).

Change summary:
- Removed an atomic state flag from the consumer since it was not technically required and introduced a race condition between where I set the flag on the producer side and where I dequeue from it on the consumer side.
- Added a CVar to allow for looping random rebuilds of the index cache as a debugging tool.

diff --git a/Engine/Source/Editor/Kismet/Private/FindInBlueprintManager.cpp b/Engine/Source/Editor/Kismet/Private/FindInBlueprintManager.cpp
--- a/Engine/Source/Editor/Kismet/Private/FindInBlueprintManager.cpp
+++ b/Engine/Source/Editor/Kismet/Private/FindInBlueprintManager.cpp
@@ -4,6 +4,7 @@
 #include "Misc/CoreMisc.h"
 #include "Misc/MessageDialog.h"
 #include "HAL/FileManager.h"
+#include "HAL/IConsoleManager.h"
 #include "Misc/Paths.h"
 #include "Async/ParallelFor.h"
 #include "Async/TaskGraphInterfaces.h"
@@ -111,6 +112,14 @@ const FString FFiBMD::FiBSearchableExplicitMD = TEXT("BlueprintSearchableExplici
 const FString FFiBMD::FiBSearchableHiddenExplicitMD = TEXT("BlueprintSearchableHiddenExplicit");
 const FString FFiBMD::FiBSearchableFormatVersionMD = TEXT("BlueprintSearchableFormatVersion");
 
+/** Attempts to initiate a rebuild of a random subset of the index cache each tick. */
+static TAutoConsoleVariable<bool> CVarRandomSearchIndexCacheRebuildEachTick(
+	TEXT("bp.RandomSearchIndexCacheRebuildEachTick"),
+	false,
+	TEXT("Attempts to initiate a random index cache update each tick (for debugging purposes)"),
+	ECVF_Default
+);
+
 /* Return the outer of the specified object that is a direct child of a package */
 inline UObject* GetAssetObject(UObject* InObject)
 {
@@ -1550,7 +1559,6 @@ public:
 		, AsyncTaskBatchIndex(0)
 		, bIsGatheringSearchMetadata(false)
 		, bHasAssetsPendingGatherWork(false)
-		, bHasAssetsPendingAsyncIndexing(false)
 		, bIsStarted(false)
 		, bIsCancelled(false)
 	{
@@ -1672,7 +1680,7 @@ public:
 
 		const int32 AsyncTaskBatchSize = CacheParams.AsyncTaskBatchSize;
 		const int32 StartIndex = AsyncTaskBatchIndex * AsyncTaskBatchSize;
-		return StartIndex < UncachedAssets.Num() || bHasAssetsPendingGatherWork || bHasAssetsPendingAsyncIndexing || bIsGatheringSearchMetadata;
+		return StartIndex < UncachedAssets.Num() || !AssetsPendingAsyncIndexing.IsEmpty() || bHasAssetsPendingGatherWork || bIsGatheringSearchMetadata;
 	}
 
 	virtual bool ShouldFullyIndexAssets() const override
@@ -1722,8 +1730,6 @@ public:
 				OutAssetPaths.Add(AssetPath);
 				++Count;
 			}
-
-			bHasAssetsPendingAsyncIndexing = !AssetsPendingAsyncIndexing.IsEmpty();
 		}
 	}
 
@@ -1816,9 +1822,6 @@ public:
 
 							if (bEnqueueForAsyncIndexing)
 							{
-								// Signal that we have a new pending asset in the async indexing queue.
-								bHasAssetsPendingAsyncIndexing = true;
-
 								// Enqueue this asset path to restart async indexing with the updated search metadata.
 								AssetsPendingAsyncIndexing.Enqueue(AssetPath);
 							}
@@ -2055,9 +2058,6 @@ private:
 	/** TRUE if we have assets pending in the gather queue. It is not thread safe to check the queue state directly from a producer thread. */
 	TAtomic<bool> bHasAssetsPendingGatherWork;
 
-	/** TRUE if we have assets pending in the indexing queue. Adds consistency with the gather queue, but not strictly needed for thread safety. */
-	TAtomic<bool> bHasAssetsPendingAsyncIndexing;
-
 	/** TRUE if the caching process is started */
 	bool bIsStarted;
 
@@ -3178,6 +3178,45 @@ void FFindInBlueprintSearchManager::DumpCache(FArchive& Ar)
 	}
 }
 
+void FFindInBlueprintSearchManager::DoRandomIndexCacheRebuild()
+{
+	if (AssetRegistryModule && bHasFirstSearchOccurred && !IsCacheInProgress())
+	{
+		FARFilter ARFilter;
+		ARFilter.bRecursiveClasses = true;
+		for (FTopLevelAssetPath ClassPathName : FBlueprintAssetHandler::Get().GetRegisteredClassNames())
+		{
+			ARFilter.ClassPaths.Add(ClassPathName);
+		}
+
+		TArray<FAssetData> AssetList;
+		AssetRegistryModule->Get().GetAssets(ARFilter, AssetList);
+		if (AssetList.IsEmpty())
+		{
+			return;
+		}
+
+		const int32 MinSize = FMath::Min(1000, AssetList.Num());
+		const int32 MaxSize = FMath::Min(10000, AssetList.Num());
+		const int32 SubsetSize = MinSize + (FMath::Rand() % (MaxSize - MinSize + 1));
+		const int32 StartingIndex = FMath::Rand() % (AssetList.Num() - SubsetSize + 1);
+		for (int32 i = StartingIndex; i < StartingIndex + SubsetSize; ++i)
+		{
+			RemoveBlueprintByPath(AssetList[i].ToSoftObjectPath());
+
+			if (AssetList[i].IsAssetLoaded())
+			{
+				UBlueprint* Blueprint = CastChecked<UBlueprint>(AssetList[i].GetAsset());
+				AddOrUpdateBlueprintSearchMetadata(Blueprint);
+			}
+			else
+			{
+				AddUnloadedBlueprintSearchMetadata(AssetList[i]);
+			}
+		}
+	}
+}
+
 FText FFindInBlueprintSearchManager::ConvertHexStringToFText(FString InHexString)
 {
 	TArray<uint8> SerializedData;
@@ -3894,6 +3933,10 @@ void FFindInBlueprintSearchManager::Tick(float DeltaTime)
 		CachingOptions.OpType = EFiBCacheOpType::CachePendingAssets;
 		CacheAllAssets(nullptr, CachingOptions);
 	}
+	else if (CVarRandomSearchIndexCacheRebuildEachTick.GetValueOnGameThread())
+	{
+		DoRandomIndexCacheRebuild();
+	}
 }
 
 bool FFindInBlueprintSearchManager::IsTickable() const
@@ -3902,7 +3945,7 @@ bool FFindInBlueprintSearchManager::IsTickable() const
 	const bool bNeedsFirstIndex = bHasFirstSearchOccurred && AssetsToIndexOnFirstSearch.Num() > 0;
 
 	// Tick only if we have an active caching operation or if a search has occured before we're ready or if we have pending assets and an open FiB context or an active async search query
-	return IsCacheInProgress() || bNeedsFirstIndex || (bHasPendingAssets && (IsGlobalFindResultsOpen() || IsAsyncSearchQueryInProgress()));
+	return IsCacheInProgress() || bNeedsFirstIndex || (bHasPendingAssets && (IsGlobalFindResultsOpen() || IsAsyncSearchQueryInProgress())) || CVarRandomSearchIndexCacheRebuildEachTick.GetValueOnGameThread();
 }
 
 TStatId FFindInBlueprintSearchManager::GetStatId() const
diff --git a/Engine/Source/Editor/Kismet/Public/FindInBlueprintManager.h b/Engine/Source/Editor/Kismet/Public/FindInBlueprintManager.h
--- a/Engine/Source/Editor/Kismet/Public/FindInBlueprintManager.h
+++ b/Engine/Source/Editor/Kismet/Public/FindInBlueprintManager.h
@@ -706,6 +706,9 @@ public:
 	/** Dumps the full index cache to the given stream (for debugging purposes) */
 	void DumpCache(FArchive& Ar);
 
+	/** Randomly initiates a partial rebuild of the index cache (for debugging purposes) */
+	void DoRandomIndexCacheRebuild();
+
 public:
 	/** Converts a string of hex characters, previously converted by ConvertFTextToHexString, to an FText. */
 	static FText ConvertHexStringToFText(FString InHexString);
