Avoid ensure when a details panel is refreshed by a customization, this in turn requires fixing ~FDetailsDiff clearing property listener.. 'mistake' introduced by 25043188, but facility was vestigile post [REDACTED_PHONE])

diff --git a/Engine/Plugins/Runtime/StateTree/Source/StateTreeEditorModule/Private/SStateTreeDiff.cpp b/Engine/Plugins/Runtime/StateTree/Source/StateTreeEditorModule/Private/SStateTreeDiff.cpp
--- a/Engine/Plugins/Runtime/StateTree/Source/StateTreeEditorModule/Private/SStateTreeDiff.cpp
+++ b/Engine/Plugins/Runtime/StateTree/Source/StateTreeEditorModule/Private/SStateTreeDiff.cpp
@@ -341,12 +341,12 @@ void SDiffWidget::HandleStateDiffEntryFocused(const FSingleDiffEntry& StateDiff)
 		const TSharedPtr<SBox> RightWidget = SNew(SBox);
 		if (OldState)
 		{
-			const FDetailsDiff DetailsDiff(OldState, FDetailsDiff::FOnDisplayedPropertiesChanged(), true);
+			const FDetailsDiff DetailsDiff(OldState, true);
 			LeftWidget->SetContent(DetailsDiff.DetailsWidget());
 		}
 		if (NewState)
 		{
-			const FDetailsDiff DetailsDiff(NewState, FDetailsDiff::FOnDisplayedPropertiesChanged(), false);
+			const FDetailsDiff DetailsDiff(NewState, false);
 			RightWidget->SetContent(DetailsDiff.DetailsWidget());
 		}
 		// And displayed in a way that resembles the details diff view
diff --git a/Engine/Source/Developer/Merge/Private/SMergeDetailsView.cpp b/Engine/Source/Developer/Merge/Private/SMergeDetailsView.cpp
--- a/Engine/Source/Developer/Merge/Private/SMergeDetailsView.cpp
+++ b/Engine/Source/Developer/Merge/Private/SMergeDetailsView.cpp
@@ -22,19 +22,19 @@ void SMergeDetailsView::Construct(const FArguments InArgs
 	// EMergeParticipant::Remote
 	{
 		DetailsViews.Add(
-			FDetailsDiff(RemoteCDO, FDetailsDiff::FOnDisplayedPropertiesChanged() )
+			FDetailsDiff(RemoteCDO )
 		);
 	}
 	// EMergeParticipant::Base
 	{
 		DetailsViews.Add(
-			FDetailsDiff(BaseCDO, FDetailsDiff::FOnDisplayedPropertiesChanged())
+			FDetailsDiff(BaseCDO)
 		);
 	}
 	// EMergeParticipant::Local
 	{
 		DetailsViews.Add(
-			FDetailsDiff(LocalCDO, FDetailsDiff::FOnDisplayedPropertiesChanged())
+			FDetailsDiff(LocalCDO)
 		);
 	}
 	
diff --git a/Engine/Source/Editor/Kismet/Private/DetailsDiff.cpp b/Engine/Source/Editor/Kismet/Private/DetailsDiff.cpp
--- a/Engine/Source/Editor/Kismet/Private/DetailsDiff.cpp
+++ b/Engine/Source/Editor/Kismet/Private/DetailsDiff.cpp
@@ -29,16 +29,21 @@ class UObject;
 
 #define LOCTEXT_NAMESPACE "DetailsDif"
 
-FDetailsDiff::FDetailsDiff(const UObject* InObject, FOnDisplayedPropertiesChanged InOnDisplayedPropertiesChanged, bool bScrollbarOnLeft )
-	: OnDisplayedPropertiesChanged( InOnDisplayedPropertiesChanged )
-	, DisplayedObject( InObject )
+FDetailsDiff::FDetailsDiff(const UObject* InObject, bool bScrollbarOnLeft )
+	: DisplayedObject( InObject )
 {
 	ScrollBar = SNew(SLinkableScrollBar);
 	DetailsView = CreateDetailsView(InObject, ScrollBar, bScrollbarOnLeft);
-	DetailsView->SetOnDisplayedPropertiesChanged( ::FOnDisplayedPropertiesChanged::CreateRaw(this, &FDetailsDiff::HandlePropertiesChanged) );
 	DisplayedProperties = DetailsView->GetPropertiesInOrderDisplayed();
 }
 
+PRAGMA_DISABLE_DEPRECATION_WARNINGS
+FDetailsDiff::FDetailsDiff( const UObject* InObject, FDetailsDiff::FOnDisplayedPropertiesChanged InOnDisplayedPropertiesChanged, bool bScrollbarOnLeft )
+	: FDetailsDiff(InObject, bScrollbarOnLeft)
+{
+}
+PRAGMA_ENABLE_DEPRECATION_WARNINGS
+
 class DOBPDetailsCustomization : public IDetailCustomization
 {
 public:
@@ -185,11 +190,6 @@ TSharedRef<IDetailsView> FDetailsDiff::CreateDetailsView(const UObject* InObject
 	return DetailsView;
 }
 
-FDetailsDiff::~FDetailsDiff()
-{
-	DetailsView->SetOnDisplayedPropertiesChanged(::FOnDisplayedPropertiesChanged());
-}
-
 void FDetailsDiff::HighlightProperty(const FPropertySoftPath& PropertyName)
 {
 	// resolve the property soft path:
@@ -202,15 +202,6 @@ TSharedRef< IDetailsView > FDetailsDiff::DetailsWidget() const
 	return DetailsView.ToSharedRef();
 }
 
-void FDetailsDiff::HandlePropertiesChanged()
-{
-	if( OnDisplayedPropertiesChanged.IsBound() )
-	{
-		DisplayedProperties = DetailsView->GetPropertiesInOrderDisplayed();
-		OnDisplayedPropertiesChanged.Execute();
-	}
-}
-
 TArray<FPropertySoftPath> FDetailsDiff::GetDisplayedProperties() const
 {
 	TArray<FPropertySoftPath> Ret;
diff --git a/Engine/Source/Editor/Kismet/Private/DiffControl.cpp b/Engine/Source/Editor/Kismet/Private/DiffControl.cpp
--- a/Engine/Source/Editor/Kismet/Private/DiffControl.cpp
+++ b/Engine/Source/Editor/Kismet/Private/DiffControl.cpp
@@ -417,9 +417,12 @@ void FDetailsDiffControl::EnableComments(TWeakPtr<STreeView<TSharedPtr<FBlueprin
 
 TSharedRef<IDetailsView> FDetailsDiffControl::InsertObject(const UObject* Object, bool bScrollbarOnLeft, int32 Index)
 {
-	FDetailsDiff DetailsDiff(Object, FDetailsDiff::FOnDisplayedPropertiesChanged(), bScrollbarOnLeft);
+	FDetailsDiff DetailsDiff(Object, bScrollbarOnLeft);
 	const TSharedRef<IDetailsView> DetailsView = DetailsDiff.DetailsWidget();
 	DetailsView->UpdatePropertyAllowList(PropertyAllowList);
+	// hook details view so we can recalculate diffs if/when a customization
+	// triggers a refresh:
+	DetailsView->SetOnDisplayedPropertiesChanged( ::FOnDisplayedPropertiesChanged::CreateRaw(this, &FDetailsDiffControl::HandlePropertiesChanged) );
 
 	if (Index == INDEX_NONE)
 	{
@@ -507,6 +510,28 @@ TArray<FVector2f> FDetailsDiffControl::GetLinkedScrollRate(TSharedRef<IDetailsVi
 	return PropertyTreeDifferences[LeftObject].Right->GenerateScrollSyncRate();
 }
 
+void FDetailsDiffControl::HandlePropertiesChanged()
+{
+	// Set up interaction with left panel
+	if(ObjectDisplayOrder.Num() != 2)
+	{
+		return;
+	}
+
+	const UObject* Object = ObjectDisplayOrder[0];
+	const UObject* OtherObject = ObjectDisplayOrder[1];
+
+	if(!Object || !OtherObject)
+	{
+		return; // should have any FAsyncDetailViewDiff to refresh
+	}
+
+	const TSharedRef<IDetailsView> DetailsView = GetDetailsWidget(Object);
+	const TSharedRef<IDetailsView> OtherDetailsView = GetDetailsWidget(OtherObject);
+	PropertyTreeDifferences[OtherObject].Left = MakeShared<FAsyncDetailViewDiff>(DetailsView, OtherDetailsView);
+	PropertyTreeDifferences[Object].Right = PropertyTreeDifferences[OtherObject].Left;
+}
+
 
 /////////////////////////////////////////////////////////////////////////////
 /// FCDODiffControl
diff --git a/Engine/Source/Editor/Kismet/Private/DiffControl.h b/Engine/Source/Editor/Kismet/Private/DiffControl.h
--- a/Engine/Source/Editor/Kismet/Private/DiffControl.h
+++ b/Engine/Source/Editor/Kismet/Private/DiffControl.h
@@ -124,6 +124,9 @@ protected:
 
 	// helper function that analyzes two details views and determines the rate they should scroll relative to one another to be in sync
 	TArray<FVector2f> GetLinkedScrollRate(TSharedRef<IDetailsView> LeftDetailsView, TSharedRef<IDetailsView> RightDetailsView) const;
+private:
+	void HandlePropertiesChanged();
+protected:
 
 	FOnDiffEntryFocused SelectionCallback;
 	
@@ -140,7 +143,6 @@ protected:
 	const bool bPopulateOutTreeEntries;
 	
 	TSet<FPropertyPath> PropertyAllowList;
-
 };
 
 /** Override for CDO special case */
@@ -173,8 +175,8 @@ struct FBlueprintTypeDiffControl : public TSharedFromThis<FBlueprintTypeDiffCont
 
 		FSubObjectDiff(const FDiffSingleResult& InSourceResult, const UObject* OldObject, const UObject* NewObject)
 			: SourceResult(InSourceResult)
-			, OldDetails(OldObject, FDetailsDiff::FOnDisplayedPropertiesChanged())
-			, NewDetails(NewObject, FDetailsDiff::FOnDisplayedPropertiesChanged())
+			, OldDetails(OldObject)
+			, NewDetails(NewObject)
 		{}
 	};
 
diff --git a/Engine/Source/Editor/Kismet/Public/DetailsDiff.h b/Engine/Source/Editor/Kismet/Public/DetailsDiff.h
--- a/Engine/Source/Editor/Kismet/Public/DetailsDiff.h
+++ b/Engine/Source/Editor/Kismet/Public/DetailsDiff.h
@@ -18,10 +18,15 @@ class UObject;
 class KISMET_API FDetailsDiff
 {
 public:
-	DECLARE_DELEGATE( FOnDisplayedPropertiesChanged );
+	FDetailsDiff( const UObject* InObject, bool bScrollbarOnLeft = false );	
 
+	UE_DEPRECATED(5.6, "FOnDisplayedPropertiesChanged removed, interact with IDetailsView directly")
+	DECLARE_DELEGATE( FOnDisplayedPropertiesChanged );
+	
+	PRAGMA_DISABLE_DEPRECATION_WARNINGS
+	UE_DEPRECATED(5.6, "InOnDisplayedPropertiesChanged, interact with IDetailsView directly")
 	FDetailsDiff( const UObject* InObject, FOnDisplayedPropertiesChanged InOnDisplayedPropertiesChanged, bool bScrollbarOnLeft = false );
-	~FDetailsDiff();
+	PRAGMA_ENABLE_DEPRECATION_WARNINGS
 
 	/** Attempt to highlight the property with the given path, may not always succeed */
 	void HighlightProperty( const FPropertySoftPath& PropertyName );
@@ -44,10 +49,6 @@ public:
 	static TSharedRef<IDetailsView> CreateDetailsView(const UObject* InObject, TSharedPtr<SScrollBar> ExternalScrollbar = nullptr, bool bScrollbarOnLeft = false);
 
 private:
-	void HandlePropertiesChanged();
-
-	FOnDisplayedPropertiesChanged OnDisplayedPropertiesChanged;
-
 	TArray< FPropertyPath > DisplayedProperties;
 	const UObject* DisplayedObject;
 
