[BlueprintSupport]

Adds ability for manually flushing the BlueprintCompilationQueue similar to manual flushing of the reinstance queue

#rnx

diff --git a/Engine/Source/Editor/Kismet/Private/BlueprintCompilationManager.cpp b/Engine/Source/Editor/Kismet/Private/BlueprintCompilationManager.cpp
--- a/Engine/Source/Editor/Kismet/Private/BlueprintCompilationManager.cpp
+++ b/Engine/Source/Editor/Kismet/Private/BlueprintCompilationManager.cpp
@@ -186,6 +186,7 @@ struct FBlueprintCompilationManagerImpl : public FGCObject
 
 // free function that we use to cross a module boundary (from CoreUObject to here)
 void FlushReinstancingQueueImplWrapper();
+void FlushCompilationQueueImplWrapper(FUObjectSerializeContext* InLoadContext);
 void MoveSkelCDOAside(UClass* Class, TMap<UClass*, UClass*>& OldToNewMap);
 void ReparentHierarchiesWrapper(const TMap<UClass*, UClass*>& OldToNewMap)
 {
@@ -195,13 +196,15 @@ void ReparentHierarchiesWrapper(const TMap<UClass*, UClass*>& OldToNewMap)
 FBlueprintCompilationManagerImpl::FBlueprintCompilationManagerImpl()
 {
 	FBlueprintSupport::SetFlushReinstancingQueueFPtr(&FlushReinstancingQueueImplWrapper);
+	FBlueprintSupport::SetFlushCompilationQueueFPtr(&FlushCompilationQueueImplWrapper);
 	FBlueprintSupport::SetClassReparentingFPtr(&ReparentHierarchiesWrapper);
 	bGeneratedClassLayoutReady = true;
 }
 
 FBlueprintCompilationManagerImpl::~FBlueprintCompilationManagerImpl() 
 { 
-	FBlueprintSupport::SetFlushReinstancingQueueFPtr(nullptr); 
+	FBlueprintSupport::SetFlushReinstancingQueueFPtr(nullptr);
+	FBlueprintSupport::SetFlushCompilationQueueFPtr(nullptr);
 	FBlueprintSupport::SetClassReparentingFPtr(nullptr);
 }
 
@@ -3822,6 +3825,11 @@ void FlushReinstancingQueueImplWrapper()
 	BPCMImpl->FlushReinstancingQueueImpl();
 }
 
+void FlushCompilationQueueImplWrapper(FUObjectSerializeContext* InLoadContext)
+{
+	FBlueprintCompilationManager::FlushCompilationQueue(InLoadContext);
+}
+
 // Recursive function to move CDOs aside to immutable versions of classes
 // so that CDOs can be safely GC'd. Recursion is necessary to find REINST_ classes
 // that are still parented to a valid SKEL (e.g. from MarkBlueprintAsStructurallyModified)
diff --git a/Engine/Source/Runtime/CoreUObject/Private/Blueprint/BlueprintSupport.cpp b/Engine/Source/Runtime/CoreUObject/Private/Blueprint/BlueprintSupport.cpp
--- a/Engine/Source/Runtime/CoreUObject/Private/Blueprint/BlueprintSupport.cpp
+++ b/Engine/Source/Runtime/CoreUObject/Private/Blueprint/BlueprintSupport.cpp
@@ -132,6 +132,7 @@ bool FBlueprintSupport::IsDeferredCDOInitializationDisabled()
 }
 
 static FFlushReinstancingQueueFPtr FlushReinstancingQueueFPtr = nullptr;
+static FFlushCompilationQueueFPtr FlushCompilationQueueFPtr = nullptr;
 static FClassReparentingFPtr ClassReparentingFPtr = nullptr;
 
 void FBlueprintSupport::FlushReinstancingQueue()
@@ -142,6 +143,14 @@ void FBlueprintSupport::FlushReinstancingQueue()
 	}
 }
 
+void FBlueprintSupport::FlushCompilationQueue(FUObjectSerializeContext* InLoadContext)
+{
+	if (FlushCompilationQueueFPtr)
+	{
+		(*FlushCompilationQueueFPtr)(InLoadContext);
+	}
+}
+
 void FBlueprintSupport::ReparentHierarchies(const TMap<UClass*, UClass*>& OldClassToNewClass)
 {
 	if (ClassReparentingFPtr)
@@ -155,6 +164,11 @@ void FBlueprintSupport::SetFlushReinstancingQueueFPtr(FFlushReinstancingQueueFPt
 	FlushReinstancingQueueFPtr = Ptr;
 }
 
+void FBlueprintSupport::SetFlushCompilationQueueFPtr(FFlushCompilationQueueFPtr Ptr)
+{
+	FlushCompilationQueueFPtr = Ptr;
+}
+
 void FBlueprintSupport::SetClassReparentingFPtr(FClassReparentingFPtr Ptr)
 {
 	ClassReparentingFPtr = Ptr;
diff --git a/Engine/Source/Runtime/CoreUObject/Public/Blueprint/BlueprintSupport.h b/Engine/Source/Runtime/CoreUObject/Public/Blueprint/BlueprintSupport.h
--- a/Engine/Source/Runtime/CoreUObject/Public/Blueprint/BlueprintSupport.h
+++ b/Engine/Source/Runtime/CoreUObject/Public/Blueprint/BlueprintSupport.h
@@ -73,6 +73,7 @@ struct FBlueprintWarningDeclaration
 };
 
 typedef void (*FFlushReinstancingQueueFPtr)();
+typedef void (*FFlushCompilationQueueFPtr)(FUObjectSerializeContext*);
 typedef void (*FClassReparentingFPtr)(const TMap<UClass*, UClass*>&);
 
 /** 
@@ -101,7 +102,10 @@ struct FBlueprintSupport
 
 	/** Checks for any old instances and reinstances them: */
 	static void FlushReinstancingQueue();
+	/** Flushes all queued compilation requests */
+	static void FlushCompilationQueue(FUObjectSerializeContext* InLoadContext = nullptr);
 	COREUOBJECT_API static void SetFlushReinstancingQueueFPtr(FFlushReinstancingQueueFPtr Ptr);	
+	COREUOBJECT_API static void SetFlushCompilationQueueFPtr(FFlushCompilationQueueFPtr Ptr);
 
 	COREUOBJECT_API static void ReparentHierarchies(const TMap<UClass*, UClass*>& OldClassToNewClass);
 	COREUOBJECT_API static void SetClassReparentingFPtr(FClassReparentingFPtr Ptr);
