Fix recently introduced abuse of CreateRaw, this can cause crashes if a details panel leaks, which they apparently are

#rnx

diff --git a/Engine/Source/Editor/Kismet/Private/DiffControl.cpp b/Engine/Source/Editor/Kismet/Private/DiffControl.cpp
--- a/Engine/Source/Editor/Kismet/Private/DiffControl.cpp
+++ b/Engine/Source/Editor/Kismet/Private/DiffControl.cpp
@@ -422,7 +422,7 @@ TSharedRef<IDetailsView> FDetailsDiffControl::InsertObject(const UObject* Object
 	DetailsView->UpdatePropertyAllowList(PropertyAllowList);
 	// hook details view so we can recalculate diffs if/when a customization
 	// triggers a refresh:
-	DetailsView->SetOnDisplayedPropertiesChanged( ::FOnDisplayedPropertiesChanged::CreateRaw(this, &FDetailsDiffControl::HandlePropertiesChanged) );
+	DetailsView->SetOnDisplayedPropertiesChanged( ::FOnDisplayedPropertiesChanged::CreateSP(AsShared(), &FDetailsDiffControl::HandlePropertiesChanged) );
 
 	if (Index == INDEX_NONE)
 	{
