Reliably clear SetOnDisplayedPropertiesChanged when using hte diff tool

#rnx

diff --git a/Engine/Source/Editor/Kismet/Private/DiffControl.cpp b/Engine/Source/Editor/Kismet/Private/DiffControl.cpp
--- a/Engine/Source/Editor/Kismet/Private/DiffControl.cpp
+++ b/Engine/Source/Editor/Kismet/Private/DiffControl.cpp
@@ -310,6 +310,14 @@ FDetailsDiffControl::FDetailsDiffControl(const UObject* InOldObject, const UObje
 	}
 }
 
+FDetailsDiffControl::~FDetailsDiffControl()
+{
+	for(const TPair<const UObject*, FDetailsDiff>& Pair : DetailsDiffs)
+	{
+		Pair.Value.DetailsWidget()->SetOnDisplayedPropertiesChanged(FOnDisplayedPropertiesChanged());
+	}
+}
+
 void FDetailsDiffControl::Tick()
 {
 	for (auto& [Object, PropertyTreeDiffs] : PropertyTreeDifferences)
@@ -422,7 +430,7 @@ TSharedRef<IDetailsView> FDetailsDiffControl::InsertObject(const UObject* Object
 	DetailsView->UpdatePropertyAllowList(PropertyAllowList);
 	// hook details view so we can recalculate diffs if/when a customization
 	// triggers a refresh:
-	DetailsView->SetOnDisplayedPropertiesChanged( ::FOnDisplayedPropertiesChanged::CreateSP(AsShared(), &FDetailsDiffControl::HandlePropertiesChanged) );
+	DetailsView->SetOnDisplayedPropertiesChanged( ::FOnDisplayedPropertiesChanged::CreateRaw(this, &FDetailsDiffControl::HandlePropertiesChanged) );
 
 	if (Index == INDEX_NONE)
 	{
diff --git a/Engine/Source/Editor/Kismet/Private/DiffControl.h b/Engine/Source/Editor/Kismet/Private/DiffControl.h
--- a/Engine/Source/Editor/Kismet/Private/DiffControl.h
+++ b/Engine/Source/Editor/Kismet/Private/DiffControl.h
@@ -98,6 +98,7 @@ class KISMET_API FDetailsDiffControl : public TSharedFromThis<FDetailsDiffContro
 {
 public:
 	FDetailsDiffControl(const UObject* InOldObject, const UObject* InNewObject, FOnDiffEntryFocused InSelectionCallback, bool bPopulateOutTreeEntries);
+	~FDetailsDiffControl();
 
 	virtual void Tick() override;
 	
@@ -126,6 +127,11 @@ protected:
 	TArray<FVector2f> GetLinkedScrollRate(TSharedRef<IDetailsView> LeftDetailsView, TSharedRef<IDetailsView> RightDetailsView) const;
 private:
 	void HandlePropertiesChanged();
+
+	FDetailsDiffControl& operator=(const FDetailsDiffControl&) = delete;
+	FDetailsDiffControl( const FDetailsDiffControl& ) = delete;
+	FDetailsDiffControl& operator=(FDetailsDiffControl&&) = delete;
+	FDetailsDiffControl( FDetailsDiffControl&& ) = delete;
 protected:
 
 	FOnDiffEntryFocused SelectionCallback;
