[BugFix] When diffing a blueprint with multiple nested graphs such as an animation blueprint, very different graphs can have the same FName. This causes them to diff with one another even though there's a better match. To mitigate this, graphs are now matched by their outer chain

diff --git a/Engine/Source/Editor/Kismet/Private/SBlueprintDiff.cpp b/Engine/Source/Editor/Kismet/Private/SBlueprintDiff.cpp
--- a/Engine/Source/Editor/Kismet/Private/SBlueprintDiff.cpp
+++ b/Engine/Source/Editor/Kismet/Private/SBlueprintDiff.cpp
@@ -903,34 +903,31 @@ void SBlueprintDiff::GenerateDifferencesList()
 		TAttribute<TArray<FVector2f>> ScrollRate = TAttribute<TArray<FVector2f>>::CreateRaw(GraphDetailDiff.Get(), &FAsyncDetailViewDiff::GenerateScrollSyncRate);
 	}
 
-	//Add Graphs that exist in both blueprints, or in blueprint 1 only
-	for (UEdGraph* GraphOld : GraphsOld)
+	// match the graphs in the old blueprint with the graphs in the new blueprint so they can be diffed
+	struct FGraphPair
 	{
-		UEdGraph* GraphNew = nullptr;
-		for (UEdGraph*& TestGraph : GraphsNew)
-		{
-			if (TestGraph && GraphOld->GetName() == TestGraph->GetName())
-			{
-				GraphNew = TestGraph;
-
-				// Null reference inside array
-				TestGraph = nullptr;
-				break;
-			}
-		}
-		// Do not worry about graphs that are contained in MathExpression nodes, they are recreated each compile
-		if (IsGraphDiffNeeded(GraphOld))
-		{
-			CreateGraphEntry(GraphOld,GraphNew);
-		}
+		TObjectPtr<UEdGraph> GraphOld;
+		TObjectPtr<UEdGraph> GraphNew;
+	};
+	TMap<FString, FGraphPair> PathToGraphs;
+	for (UEdGraph* Graph : GraphsNew)
+	{
+		FString Path;
+		Graph->GetPathName(PanelNew.Blueprint, Path);
+		PathToGraphs.FindOrAdd(Path).GraphOld = Graph;
+	}
+	for (UEdGraph* Graph : GraphsOld)
+	{
+		FString Path;
+		Graph->GetPathName(PanelOld.Blueprint, Path);
+		PathToGraphs.FindOrAdd(Path).GraphNew = Graph;
 	}
 
-	//Add graphs that only exist in 2nd(new) blueprint
-	for (UEdGraph* GraphNew : GraphsNew)
+	for (const auto& [Path, GraphPair] : PathToGraphs)
 	{
-		if (GraphNew != nullptr && IsGraphDiffNeeded(GraphNew))
+		if (GraphPair.GraphOld ? IsGraphDiffNeeded(GraphPair.GraphOld) : IsGraphDiffNeeded(GraphPair.GraphNew))
 		{
-			CreateGraphEntry(nullptr, GraphNew);
+			CreateGraphEntry(GraphPair.GraphOld,GraphPair.GraphNew);
 		}
 	}
 
