Iteration on BP deps to support non DOBP. Include enough information in blueprint dependencies to detect changes that would impact bytecode generation

#rnx

diff --git a/Engine/Source/Editor/Kismet/Internal/Blueprints/BlueprintDependencies.h b/Engine/Source/Editor/Kismet/Internal/Blueprints/BlueprintDependencies.h
--- a/Engine/Source/Editor/Kismet/Internal/Blueprints/BlueprintDependencies.h
+++ b/Engine/Source/Editor/Kismet/Internal/Blueprints/BlueprintDependencies.h
@@ -6,16 +6,27 @@
 
 #include "BlueprintDependencies.generated.h"
 
+UENUM()
+enum class EBPDependencyType
+{
+	Asset,
+	Struct,
+	Class
+};
+
 USTRUCT()
-struct FInheritanceDependency
+struct FBlueprintDependency
 {
 	GENERATED_BODY()
 	
+	UPROPERTY()
+	EBPDependencyType DependencyType;
+
 	UPROPERTY()
 	FName PackageName;
 
 	UPROPERTY()
-	FName NativeClassName;
+	FName NativeObjectName;
 
 	UPROPERTY()
 	FString Hash;
@@ -27,7 +38,7 @@ struct FBlueprintDependencies
 	GENERATED_BODY()
 
 	UPROPERTY()
-	TArray<FInheritanceDependency> InheritanceChain;
+	TArray<FBlueprintDependency> BlueprintDependencies;
 };
 
 namespace UE::Private::BlueprintDependencies
diff --git a/Engine/Source/Editor/Kismet/Private/Blueprints/BlueprintDependencies.cpp b/Engine/Source/Editor/Kismet/Private/Blueprints/BlueprintDependencies.cpp
--- a/Engine/Source/Editor/Kismet/Private/Blueprints/BlueprintDependencies.cpp
+++ b/Engine/Source/Editor/Kismet/Private/Blueprints/BlueprintDependencies.cpp
@@ -4,67 +4,300 @@
 
 #include "Engine/Blueprint.h"
 #include "HAL/FileManager.h"
+#include "JsonObjectConverter.h"
+#include "JsonObjectGraph/Stringify.h"
 #include "Kismet2/BlueprintEditorUtils.h"
 #include "Misc/ScopeExit.h"
 #include "Serialization/CompactBinary.h"
 #include "Serialization/CompactBinaryWriter.h"
-#include "JsonObjectGraph/Stringify.h"
+#include "UObject/StructOnScope.h"
 
 /*
 	BLUEPRINT DEPENDENCIES DESIGN
 
 	This file is meant to capture external dependencies of a blueprint so that it can
-	be reliably recooked when its inputs change. It currently has only been tested
-	against data only blueprints, which represente default value changes made against
-	some parent class. This means that any time the parent class(es) change the blueprint
-	must be recooked. Because there is not an obvious artifact to identify for native
-	classes any native classes are serialized to json for the purpose of providing identity.
-
-	This captures, among other things, the class default values and custom versions used
-	in that native class(es) serialization routines. Non-native classes are checked for
-	identity using a blake3 hash of their source artifact. Both of these approaches have
-	drawbacks, and as these drawbacks are addressed I will look towards unifying native
-	or blueprint class dependencies.
-
-	We also want to extend this design to blueprints with bytecode, functions, extensions,
-	etc.
+	be reliably recooked when its (cook time) inputs change. Dependencies of the blueprint
+	are categorized as (1) native classes, (2) native structures, or (3) assets. Native classes 
+	and native structs are canonized as strings for the purpose of comparison. Asset dependencies 
+	have their file contents hashed. There is ample opportunity for optimization here, but
+	there is even more ample opportunity for false positives and negatives. The current emphasis
+	is on simplicity, durability and correctness, at the possible expense of raw throughput.
 */
 
 namespace UE::Private
 {
-constexpr int32 BlueprintCookDependenciesVersion = 1;
+constexpr int32 BlueprintCookDependenciesVersion = 2;
 
-static bool IsNativeObject( const UObject* Obj);
-static FString HashClass(const UClass* Class);
+// gathering helpers:
+static TArray<const UObject*> GetAllImportedObjects(TConstArrayView<const UObject*> Roots);
+static TArray<const UObject*> GetAllImportedObjects(const UBlueprint* ForBP);
+// helpers for handling transitiveness of dependencies (e.g. if i use a struct, i also depend on its member structs):
+static void AddStructDependencyImpl(const UScriptStruct* Struct, TSet<const UScriptStruct*>& ReferencedStructs);
+static void AddClassDependencyImpl(const UClass* Class, TSet<const UClass*>& ReferencedClasses);
+// returns all the dependencies of a bp, for recording/generating a snapshot of dependencies:
+static void GetAllDependencies(
+	const UBlueprint* ForBP, 
+	TArray<const UPackage*>& AssetDependencies, 
+	TArray<const UClass*>& NativeClassDependencies,
+	TArray<const UScriptStruct*>& NativeStructDependencies);
+// Native struct and class canonization:
+static FString HashNativeStruct(const UScriptStruct* Struct);
 static FString HashNativeClass( const UClass* NativeClass);
+// Asset canonization:
 static FString HashPackage( const UPackage* Package );
 static FString HashFile(const FString& Filename);
+// Filesystem helper:
 static bool ReadBytesFromFile(const TCHAR* Filename, TArray64<uint8>& OutBytes);
-
+// Fail caching helper:
 static void NotCacheable(FCbFieldViewIterator Args, UE::Cook::FCookDependencyContext& Context);
+// Main implementation, validate and generate are pairs, as are load and save:
 static void ValidateBPCookDependenciesImpl(FCbFieldViewIterator Args, UE::Cook::FCookDependencyContext& Context);
 static void GenerateBlueprintDependencies(const UBlueprint* ForBP, FBlueprintDependencies& OutDependencies);
 static bool LoadBPCookDependenciesImpl(FCbFieldViewIterator Args, UE::Cook::FCookDependencyContext& Context, FBlueprintDependencies& OutDependencies);
 static void SaveBPCookDepenenciesImpl(FCbWriter& Writer, const FBlueprintDependencies& Dependencies);
 }
 
-static bool UE::Private::IsNativeObject( const UObject* Obj)
+static TArray<const UObject*> UE::Private::GetAllImportedObjects(TConstArrayView<const UObject*> Roots)
 {
-	return Obj->GetPackage()->HasAnyPackageFlags(PKG_CompiledIn);
+	// This function just gathers all of the UObjects the Root objects (or their inners)
+	// depend upon. Inners are not currently tested for reachability. Tautological self reference
+	// (e.g. I depend on myself or I depend on my inner) are not reported - only external objects:
+	struct FExternalReferenceFinder : public FArchiveUObject
+	{
+		FExternalReferenceFinder(const UObject* Obj, TSet<const UObject*>& InReferences, const TSet<const UObject*>& InRoots)
+			: FArchiveUObject()
+			, SearchRoots(InRoots)
+			, References(InReferences)
+		{
+			// Copying FPackageHarvester:
+			SetIsPersistent(true);
+			SetIsSaving(true);
+			SetUseUnversionedPropertySerialization(true);
+			ArIsObjectReferenceCollector = true;
+			ArShouldSkipBulkData = true;
+			if (Obj->HasAnyFlags(RF_ClassDefaultObject))
+			{
+				Obj->GetClass()->SerializeDefaultObject((UObject*)Obj, *this);
+			}
+			else
+			{
+				((UObject*)Obj)->Serialize(*this);
+			}
+		}
+
+	private:
+		virtual FArchive& operator<<(UObject*& ObjRef) override
+		{
+			if (ObjRef != nullptr &&
+				// uht objects are tagged as RF_Transient but they 'are native' 
+				// so still referencable:
+				(!ObjRef->HasAnyFlags(RF_Transient) || ObjRef->IsNative()) &&
+				!ObjRef->IsA<UPackage>() &&
+				!IsInRoots(ObjRef))
+			{
+				// Set to null any pointer to an external asset
+				References.Add(ObjRef);
+			}
+
+			return *this;
+		}
+
+		bool IsInRoots(const UObject* Obj)
+		{
+			const UObject* Iter = Obj;
+			while(Iter)
+			{
+				if(SearchRoots.Contains(Iter))
+				{
+					return true;
+				}
+				Iter = Iter->GetOuter();
+			}
+			return false;
+		}
+
+		const TSet<const UObject*>& SearchRoots;
+		TSet<const UObject*>& References;
+	};
+
+	TSet<const UObject*> RootSet(Roots);
+	TSet<const UObject*> Result;
+	TArray<const UObject*> Objects;
+	for(const UObject* Obj : Roots)
+	{
+		GetObjectsWithOuter(Obj, (TArray<UObject*>&)Objects, true, RF_Transient, EInternalObjectFlags::Garbage);
+		Objects.Add(Obj);
+
+		// include any class's super struct chain for reference finding,
+		// I'm doing this specifically to gather inherited structure 
+		// dependencies..
+		if(const UClass* AsClass = Cast<UClass>(Obj))
+		{
+			const UClass* Iter = AsClass->GetSuperClass();
+			while(Iter)
+			{
+				Objects.Add(Iter);
+				Iter = Iter->GetSuperClass();
+			}
+		}
+	}
+
+	for(const UObject* Obj: Objects)
+	{
+		FExternalReferenceFinder References(Obj, Result, RootSet);
+	}
+	
+	return Result.Array();
+}
+
+static TArray<const UObject*> UE::Private::GetAllImportedObjects( const UBlueprint* ForBP)
+{
+	return GetAllImportedObjects({ForBP, ForBP->GeneratedClass, ForBP->GeneratedClass->ClassDefaultObject});
+}
+
+static void UE::Private::AddStructDependencyImpl(const UScriptStruct* Struct, TSet<const UScriptStruct*>& ReferencedStructs)
+{
+	bool bWasAlreadyPresent = false;
+	ReferencedStructs.Add(Struct, &bWasAlreadyPresent);
+	if(bWasAlreadyPresent)
+	{
+		return;
+	}
+	
+	// add aggregate structs, and their aggregate structs - and their super structs
+	for (TFieldIterator<FStructProperty> It(Struct, EFieldIteratorFlags::ExcludeSuper); It; ++It)
+	{
+		AddStructDependencyImpl(It->Struct, ReferencedStructs);
+	}
+
+	const UScriptStruct* Super = Cast<UScriptStruct>(Struct->GetSuperStruct()); // can't be castchecked only because we may have a null super
+	if(Super)
+	{
+		AddStructDependencyImpl(Super, ReferencedStructs);
+	}
+}
+
+static void UE::Private::AddClassDependencyImpl(const UClass* Class, TSet<const UClass*>& ReferencedClasses)
+{
+	const UClass* Iter = Class;
+	while (Iter)
+	{
+		bool bWasAlreadyPresent = false;
+		ReferencedClasses.Add(Iter, &bWasAlreadyPresent);
+		if (bWasAlreadyPresent)
+		{
+			return;
+		}
+
+		Iter = Iter->GetSuperClass();
+	}
 }
 
-static FString UE::Private::HashClass(const UClass* Class)
+static void UE::Private::GetAllDependencies(
+	const UBlueprint* ForBP,
+	TArray<const UPackage*>& AssetDependencies, 
+	TArray<const UClass*>& NativeClassDependencies,
+	TArray<const UScriptStruct*>& NativeStructDependencies)
 {
-	// for native classes we have no artifact to hash so hash the class itself,
-	// for non-native classes just hash the artifact (.uasset, aka package):
-	if(Class->HasAnyClassFlags(CLASS_Native))
+	TRACE_CPUPROFILER_EVENT_SCOPE(BlueprintDependencies::GetAllDependencies);
+
+	// gather all references in the BP, its class, and the CDO, along with any subobjects:
+	TArray<const UObject*> ImportedObjects = GetAllImportedObjects(ForBP);
+
+	// honor transitive dependencies - e.g. inherited classes, inherited structs, aggregate structs
+	// for now all class and struct references are fully transitive:
+	TSet<const UScriptStruct*> ReferencedStructs;
+	TSet<const UClass*> ReferencedClasses;
+	TSet<const UPackage*> ReferencedAssets;
+	for(const UObject* Object : ImportedObjects)
 	{
-		return HashNativeClass(Class);
+		const UPackage* Package = Object->GetPackage();
+		if(const UClass* AsClass = Cast<UClass>(Object))
+		{
+			AddClassDependencyImpl(AsClass, ReferencedClasses);
+		}
+		else if(const UScriptStruct* AsStruct = Cast<UScriptStruct>(Object))
+		{
+			AddStructDependencyImpl(AsStruct, ReferencedStructs);
+		}
+		else if (!Package->HasAnyPackageFlags(PKG_CompiledIn))
+		{
+			ReferencedAssets.Add(Package);
+		}
 	}
-	else
+
+	// process classes and structs, categorize them as native or asset:
+	for(const UClass* Class : ReferencedClasses)
 	{
-		return HashPackage(Class->GetPackage());
+		if (Class->HasAnyClassFlags(CLASS_Native))
+		{
+			NativeClassDependencies.Add(Class);
+		}
+		else
+		{
+			ReferencedAssets.Add(Class->GetPackage());
+		}
 	}
+
+	for(const UScriptStruct* Struct : ReferencedStructs)
+	{
+		if(Struct->StructFlags & STRUCT_Native)
+		{
+			NativeStructDependencies.Add(Struct);
+		}
+		else
+		{
+			ReferencedAssets.Add(Struct->GetPackage());
+		}
+	}
+	
+	AssetDependencies = ReferencedAssets.Array();
+
+	// give the arrays stable order - noisey arrays will be obnoxious for
+	// memoization/distribution:
+	const auto SortTopLevelObject = [](const UObject& ObjA, const UObject& ObjB)
+	{
+		ensure(ObjA.GetOuter()->IsA<UPackage>());
+		ensure(ObjB.GetOuter()->IsA<UPackage>());
+		// sort by packagename/objectname:
+		const FName PackageNameA = ObjA.GetOuter()->GetFName();
+		const FName PackageNameB = ObjB.GetOuter()->GetFName();
+		if(PackageNameA == PackageNameB)
+		{
+			return ObjA.GetFName().LexicalLess(ObjB.GetFName());
+		}
+		else
+		{
+			return PackageNameA.LexicalLess(PackageNameB);
+		}
+	};
+	NativeStructDependencies.Sort(SortTopLevelObject);
+	NativeClassDependencies.Sort(SortTopLevelObject);
+	AssetDependencies.Sort([](const UPackage& A, const UPackage& B)
+	{
+		return A.GetFName().LexicalLess(B.GetFName());
+	});
+}
+
+static FString UE::Private::HashNativeStruct(const UScriptStruct* Struct)
+{
+	check(Struct->StructFlags & STRUCT_Native);
+	FUtf8String StringifiedStruct = UE::JsonObjectGraph::Stringify(
+		{Struct}, 
+		FJsonStringifyOptions(EJsonStringifyFlags::FilterEditorOnlyData));
+
+	// JsonObjectGraph doesn't support const UStruct*/void* pairs,
+	// so use another json routine to get identity for default values
+	// shortcomings here are that we aren't getting much info about 
+	// native serialize overrides (inc. custom version bumps)
+	FStructOnScope Defaults(Struct);
+	FString DefaultValues;
+	FJsonObjectConverter::UStructToJsonObjectString(Struct, Defaults.GetStructMemory(), DefaultValues);
+	StringifiedStruct.Append(DefaultValues);
+
+	FBlake3 Hash;
+	Hash.Update(StringifiedStruct.GetCharArray().GetData(), StringifiedStruct.GetCharArray().Num());
+	return LexToString(Hash.Finalize());
 }
 
 static FString UE::Private::HashNativeClass( const UClass* NativeClass)
@@ -145,47 +378,74 @@ void UE::Private::ValidateBPCookDependenciesImpl(FCbFieldViewIterator Args, UE::
 	{
 		return;
 	}
-
+	
+	TRACE_CPUPROFILER_EVENT_SCOPE(BlueprintDependencies::ValidateBPCookDependenciesImpl);
 	// confirm dependencies do/do not match the artifact - must correspond to 
-	for(const FInheritanceDependency& InheritanceDependency : Dependencies.InheritanceChain)
+	// GenerateBlueprintDependencies (hashpackage<->hashfile, hashnative class, hashnative struct)
+	for(const FBlueprintDependency& BPDep : Dependencies.BlueprintDependencies)
 	{
-		// if the dependency is native hash the native class and confirm it matches:
-		const FName NativeClassName = InheritanceDependency.NativeClassName;
-		const FName PackageName = InheritanceDependency.PackageName;
-		if( NativeClassName != FName())
+		switch(BPDep.DependencyType)
 		{
-			const UPackage* Package = FindObjectFast<UPackage>(nullptr, PackageName);
-			if(!Package)
-			{
-				Context.LogInvalidated(
-					FString::Printf(TEXT("Native Package Missing: %s"), *PackageName.ToString()));
-				return;
-			}
-			const UClass* Class = FindObjectFast<UClass>(
-				const_cast<UPackage*>(Package), NativeClassName);
-			if(!Class)
+			case EBPDependencyType::Asset:
 			{
-				Context.LogInvalidated(
-					FString::Printf(TEXT("Native Class Missing: %s"), *NativeClassName.ToString()));
-				return;
+				FString PackageFilename;
+				const bool bSuceeded = FPackageName::TryConvertLongPackageNameToFilename(
+					BPDep.PackageName.ToString(), PackageFilename,  FPackageName::GetAssetPackageExtension());
+				check(bSuceeded);
+				if(!BPDep.Hash.Equals(UE::Private::HashFile(PackageFilename)))
+				{
+					Context.LogInvalidated(FString::Printf(TEXT("Peer Package Changed: %s"), *PackageFilename));
+					return;
+				}
+				break;
 			}
-			if(InheritanceDependency.Hash != UE::Private::HashClass(Class))
+			case EBPDependencyType::Struct:
 			{
-				Context.LogInvalidated(FString::Printf(TEXT("Native Class Changed: %s"), *Class->GetPathName()));
-				return;
+				const UPackage* Package = FindObjectFast<UPackage>(nullptr, BPDep.PackageName);
+				if(!Package)
+				{
+					Context.LogInvalidated(
+						FString::Printf(TEXT("Native Package Missing: %s"), *BPDep.PackageName.ToString()));
+					return;
+				}
+				const UScriptStruct* Struct = FindObjectFast<UScriptStruct>(
+					const_cast<UPackage*>(Package), BPDep.NativeObjectName);
+				if(!Struct)
+				{
+					Context.LogInvalidated(
+						FString::Printf(TEXT("Native Struct Missing: %s"), *BPDep.NativeObjectName.ToString()));
+					return;
+				}
+				if(!BPDep.Hash.Equals(UE::Private::HashNativeStruct(Struct)))
+				{
+					Context.LogInvalidated(FString::Printf(TEXT("Native Struct Changed: %s"), *Struct->GetPathName()));
+					return;
+				}
+				break;
 			}
-		}
-		// else if the dependency is from a package hash that package and confirm it matches:
-		else
-		{
-			FString PackageFilename;
-			const bool bSuceeded = FPackageName::TryConvertLongPackageNameToFilename(
-				PackageName.ToString(), PackageFilename,  FPackageName::GetAssetPackageExtension());
-			check(bSuceeded);
-			if(InheritanceDependency.Hash != UE::Private::HashFile(PackageFilename))
+			case EBPDependencyType::Class:
 			{
-				Context.LogInvalidated(FString::Printf(TEXT("Package Changed: %s"), *PackageFilename));
-				return;
+				const UPackage* Package = FindObjectFast<UPackage>(nullptr, BPDep.PackageName);
+				if(!Package)
+				{
+					Context.LogInvalidated(
+						FString::Printf(TEXT("Native Package Missing: %s"), *BPDep.PackageName.ToString()));
+					return;
+				}
+				const UClass* Class = FindObjectFast<UClass>(
+					const_cast<UPackage*>(Package), BPDep.NativeObjectName);
+				if(!Class)
+				{
+					Context.LogInvalidated(
+						FString::Printf(TEXT("Native Class Missing: %s"), *BPDep.NativeObjectName.ToString()));
+					return;
+				}
+				if(!BPDep.Hash.Equals(UE::Private::HashNativeClass(Class)))
+				{
+					Context.LogInvalidated(FString::Printf(TEXT("Native Class Changed: %s"), *Class->GetPathName()));
+					return;
+				}
+				break;
 			}
 		}
 	}
@@ -194,21 +454,54 @@ UE_COOK_DEPENDENCY_FUNCTION(ValidateBPCookDependenciesImpl, UE::Private::Validat
 
 void UE::Private::GenerateBlueprintDependencies(const UBlueprint* ForBP, FBlueprintDependencies& OutDependencies)
 {
+	TArray<const UPackage*> AssetDependencies;
+	TArray<const UClass*> NativeClassDependencies;
+	TArray<const UScriptStruct*> NativeStructDependencies;
+	UE::Private::GetAllDependencies(ForBP, AssetDependencies, NativeClassDependencies, NativeStructDependencies);
+
+	TRACE_CPUPROFILER_EVENT_SCOPE(BlueprintDependencies::GenerateBlueprintDependencies);
 	// MUST MATCH UE::Private::ValidateBPCookDependenciesImpl
-	const UClass* Hierarchy = ForBP->GeneratedClass;
-	while(Hierarchy)
+	// gather dependencies on functions/structs/classes/objects:
+	// treat class as transitive, in case we're using them for type comparisons:
+	// treat struct dependencies as aggregate, in case inner structs change:
+	for(const UPackage* Package : AssetDependencies)
+	{
+		FBlueprintDependency Dependency = {
+			.DependencyType = EBPDependencyType::Asset,
+			.PackageName = Package->GetFName(),
+			.NativeObjectName = FName(),
+			.Hash = UE::Private::HashPackage(Package)
+		};
+		
+		OutDependencies.BlueprintDependencies.Emplace(
+			MoveTemp(Dependency)
+		);
+	}
+	for(const UClass* Class : NativeClassDependencies)
+	{
+		FBlueprintDependency Dependency = {
+			.DependencyType = EBPDependencyType::Class,
+			.PackageName = Class->GetPackage()->GetFName(),
+			.NativeObjectName = Class->GetFName(),
+			.Hash = UE::Private::HashNativeClass(Class)
+		};
+		
+		OutDependencies.BlueprintDependencies.Emplace(
+			MoveTemp(Dependency)
+		);
+	}
+	for(const UScriptStruct* Struct : NativeStructDependencies)
 	{
-		FInheritanceDependency Dependency = {
-			.PackageName = Hierarchy->GetPackage()->GetFName(),
-			// class native branch sets NativeClassName, which controls class hashing 
-			// behavior here and above in ValidateBPCookDependenciesImpl
-			.NativeClassName = Hierarchy->HasAnyClassFlags(CLASS_Native) ? Hierarchy->GetFName() : FName(),
-			.Hash = UE::Private::HashClass(Hierarchy)
+		FBlueprintDependency Dependency = {
+			.DependencyType = EBPDependencyType::Struct,
+			.PackageName = Struct->GetPackage()->GetFName(),
+			.NativeObjectName = Struct->GetFName(),
+			.Hash = UE::Private::HashNativeStruct(Struct)
 		};
-		OutDependencies.InheritanceChain.Emplace(
+		
+		OutDependencies.BlueprintDependencies.Emplace(
 			MoveTemp(Dependency)
 		);
-		Hierarchy = Hierarchy->GetSuperClass();
 	}
 }
 // END PAIR VALIDATE/GENERATE
@@ -216,6 +509,7 @@ void UE::Private::GenerateBlueprintDependencies(const UBlueprint* ForBP, FBluepr
 // BEGIN PAIR LOAD/SAVE
 bool UE::Private::LoadBPCookDependenciesImpl(FCbFieldViewIterator Args, UE::Cook::FCookDependencyContext& Context, FBlueprintDependencies& OutDependencies)
 {
+	TRACE_CPUPROFILER_EVENT_SCOPE(BlueprintDependencies::LoadBPCookDependenciesImpl);
 	// MUST MATCH UE::Private::SaveBPCookDepenenciesImpl:
 	FCbFieldViewIterator ArgField(Args);
 	const int32 ArgsVersion = (ArgField++).AsInt32(); // Writer << BlueprintCookDependenciesVersion;
@@ -241,6 +535,7 @@ bool UE::Private::LoadBPCookDependenciesImpl(FCbFieldViewIterator Args, UE::Cook
 
 void UE::Private::SaveBPCookDepenenciesImpl(FCbWriter& Writer, const FBlueprintDependencies& Dependencies)
 {
+	TRACE_CPUPROFILER_EVENT_SCOPE(BlueprintDependencies::SaveBPCookDepenenciesImpl);
 	// MUST MATCH UE::Private::LoadBPCookDependenciesImpl
 	FString DependenciesSerialized;
 	FBlueprintDependencies::StaticStruct()->ExportText(
@@ -255,10 +550,10 @@ UE::Cook::FCookDependency UE::Private::BlueprintDependencies::RecordCookDependen
 {
 	if(!BP->ParentClass || 
 		!BP->GeneratedClass ||
-		!BP->GeneratedClass->IsChildOf(UObject::StaticClass()) ||
-		!FBlueprintEditorUtils::IsDataOnlyBlueprint(BP) )
+		!BP->GeneratedClass->IsChildOf(UObject::StaticClass()) )
 	{
-		// These BPs cannot yet be cached:
+		// blueprints without classes - malformed or some
+		// kind of utility blueprint, ignore them
 		return UE::Cook::FCookDependency::Function(
 			UE_COOK_DEPENDENCY_FUNCTION_CALL(NotCacheable), FCbFieldIterator());
 	}
