Trigger toast instead of assertion failure if user attempts to implement a function in a macro library blueprint
[RN] Fixed crash when attempting to implement an interface inherited function in a Blueprint Macro Library, instead a toast is triggered explaining that this is not permitted

diff --git a/Engine/Source/Editor/Kismet/Private/SMyBlueprint.cpp b/Engine/Source/Editor/Kismet/Private/SMyBlueprint.cpp
--- a/Engine/Source/Editor/Kismet/Private/SMyBlueprint.cpp
+++ b/Engine/Source/Editor/Kismet/Private/SMyBlueprint.cpp
@@ -43,6 +43,7 @@
 #include "Framework/Commands/UICommandInfo.h"
 #include "Framework/Commands/UICommandList.h"
 #include "Framework/MultiBox/MultiBoxBuilder.h"
+#include "Framework/Notifications/NotificationManager.h"
 #include "GenericPlatform/GenericApplication.h"
 #include "GraphActionNode.h"
 #include "GraphEditorActions.h"
@@ -112,6 +113,7 @@
 #include "Widgets/Input/SSearchBox.h"
 #include "Widgets/Layout/SBorder.h"
 #include "Widgets/Layout/SBox.h"
+#include "Widgets/Notifications/SNotificationList.h"
 #include "Widgets/SBoxPanel.h"
 #include "Widgets/SNullWidget.h"
 #include "Widgets/SWidget.h"
@@ -2806,6 +2808,18 @@ void SMyBlueprint::ImplementFunction(FEdGraphSchemaAction_K2Graph* GraphAction)
 {
 	UBlueprint* BlueprintObj = GetBlueprintObj();
 	check(BlueprintObj && BlueprintObj->SkeletonGeneratedClass);
+	if(BlueprintObj->BlueprintType == BPTYPE_MacroLibrary)
+	{
+		// macro libraries cannot have functions, attempting to create a function in 
+		// one would trigger an assertion failure
+		FNotificationInfo Info( LOCTEXT("NoFunctionsInMacro", 
+				"Blueprint Macro Libraries cannot implement functions"));
+		Info.bUseLargeFont = false;
+		Info.ExpireDuration = 5.0f;
+		Info.bFireAndForget = true;
+		FSlateNotificationManager::Get().AddNotification(Info);
+		return;
+	}
 
 	// Ensure that we are conforming to all current interfaces so that if there has been an additional
 	// interface function added we just focus to it instead of creating a new one
