[BugFix] Typo in recent diff code caused old graphs to be treated as the new graph and visa versa

diff --git a/Engine/Source/Editor/Kismet/Private/SBlueprintDiff.cpp b/Engine/Source/Editor/Kismet/Private/SBlueprintDiff.cpp
--- a/Engine/Source/Editor/Kismet/Private/SBlueprintDiff.cpp
+++ b/Engine/Source/Editor/Kismet/Private/SBlueprintDiff.cpp
@@ -910,16 +910,16 @@ void SBlueprintDiff::GenerateDifferencesList()
 		TObjectPtr<UEdGraph> GraphNew;
 	};
 	TMap<FString, FGraphPair> PathToGraphs;
-	for (UEdGraph* Graph : GraphsNew)
+	for (UEdGraph* Graph : GraphsOld)
 	{
 		FString Path;
-		Graph->GetPathName(PanelNew.Blueprint, Path);
+		Graph->GetPathName(PanelOld.Blueprint, Path);
 		PathToGraphs.FindOrAdd(Path).GraphOld = Graph;
 	}
-	for (UEdGraph* Graph : GraphsOld)
+	for (UEdGraph* Graph : GraphsNew)
 	{
 		FString Path;
-		Graph->GetPathName(PanelOld.Blueprint, Path);
+		Graph->GetPathName(PanelNew.Blueprint, Path);
 		PathToGraphs.FindOrAdd(Path).GraphNew = Graph;
 	}
 
