Disable dependency calculation for non-dobp, seems to have discovered some loading state machine problem

#rnx

diff --git a/Engine/Source/Editor/Kismet/Private/Blueprints/BlueprintDependencies.cpp b/Engine/Source/Editor/Kismet/Private/Blueprints/BlueprintDependencies.cpp
--- a/Engine/Source/Editor/Kismet/Private/Blueprints/BlueprintDependencies.cpp
+++ b/Engine/Source/Editor/Kismet/Private/Blueprints/BlueprintDependencies.cpp
@@ -550,7 +550,8 @@ UE::Cook::FCookDependency UE::Private::BlueprintDependencies::RecordCookDependen
 {
 	if(!BP->ParentClass || 
 		!BP->GeneratedClass ||
-		!BP->GeneratedClass->IsChildOf(UObject::StaticClass()) )
+		!BP->GeneratedClass->IsChildOf(UObject::StaticClass())  ||
+		!FBlueprintEditorUtils::IsDataOnlyBlueprint(BP))
 	{
 		// blueprints without classes - malformed or some
 		// kind of utility blueprint, ignore them
