Fully disable blueprint dependency tracking, as tracking of structs may be causing crashes

#rnx

diff --git a/Engine/Source/Editor/Kismet/Private/Blueprints/BlueprintDependencies.cpp b/Engine/Source/Editor/Kismet/Private/Blueprints/BlueprintDependencies.cpp
--- a/Engine/Source/Editor/Kismet/Private/Blueprints/BlueprintDependencies.cpp
+++ b/Engine/Source/Editor/Kismet/Private/Blueprints/BlueprintDependencies.cpp
@@ -548,6 +548,9 @@ void UE::Private::SaveBPCookDepenenciesImpl(FCbWriter& Writer, const FBlueprintD
 
 UE::Cook::FCookDependency UE::Private::BlueprintDependencies::RecordCookDependencies(const UBlueprint* BP)
 {
+	return UE::Cook::FCookDependency::Function(
+		UE_COOK_DEPENDENCY_FUNCTION_CALL(NotCacheable), FCbFieldIterator());
+#if 0
 	if(!BP->ParentClass || 
 		!BP->GeneratedClass ||
 		!BP->GeneratedClass->IsChildOf(UObject::StaticClass())  ||
@@ -567,4 +570,5 @@ UE::Cook::FCookDependency UE::Private::BlueprintDependencies::RecordCookDependen
 
 	return UE::Cook::FCookDependency::Function(
 		UE_COOK_DEPENDENCY_FUNCTION_CALL(ValidateBPCookDependenciesImpl), Writer.Save());
+#endif
 }
