Avoid use of UPS for reference collection in stringify because UPS can't serialize CDOs (assumes that every instance has a default value in its archetype, and CDOs will not for properties of the same type) (do we really need reference collection at all? probably not)
Reenable blueprint deps tracking

#rnx

diff --git a/Engine/Source/Editor/Kismet/Private/Blueprints/BlueprintDependencies.cpp b/Engine/Source/Editor/Kismet/Private/Blueprints/BlueprintDependencies.cpp
--- a/Engine/Source/Editor/Kismet/Private/Blueprints/BlueprintDependencies.cpp
+++ b/Engine/Source/Editor/Kismet/Private/Blueprints/BlueprintDependencies.cpp
@@ -548,13 +548,9 @@ void UE::Private::SaveBPCookDepenenciesImpl(FCbWriter& Writer, const FBlueprintD
 
 UE::Cook::FCookDependency UE::Private::BlueprintDependencies::RecordCookDependencies(const UBlueprint* BP)
 {
-	return UE::Cook::FCookDependency::Function(
-		UE_COOK_DEPENDENCY_FUNCTION_CALL(NotCacheable), FCbFieldIterator());
-#if 0
 	if(!BP->ParentClass || 
 		!BP->GeneratedClass ||
-		!BP->GeneratedClass->IsChildOf(UObject::StaticClass())  ||
-		!FBlueprintEditorUtils::IsDataOnlyBlueprint(BP))
+		!BP->GeneratedClass->IsChildOf(UObject::StaticClass()))
 	{
 		// blueprints without classes - malformed or some
 		// kind of utility blueprint, ignore them
@@ -570,5 +566,4 @@ UE::Cook::FCookDependency UE::Private::BlueprintDependencies::RecordCookDependen
 
 	return UE::Cook::FCookDependency::Function(
 		UE_COOK_DEPENDENCY_FUNCTION_CALL(ValidateBPCookDependenciesImpl), Writer.Save());
-#endif
 }
diff --git a/Engine/Source/Runtime/Experimental/JsonObjectGraph/Private/JsonStringifyImpl.cpp b/Engine/Source/Runtime/Experimental/JsonObjectGraph/Private/JsonStringifyImpl.cpp
--- a/Engine/Source/Runtime/Experimental/JsonObjectGraph/Private/JsonStringifyImpl.cpp
+++ b/Engine/Source/Runtime/Experimental/JsonObjectGraph/Private/JsonStringifyImpl.cpp
@@ -32,12 +32,12 @@ struct FPackageReferenceFinder : public FArchiveUObject
 		SetIsPersistent(true);
 		SetIsSaving(true);
 		SetFilterEditorOnly(bFilterEditorOnly);
-		SetUseUnversionedPropertySerialization(true);
+		ArNoDelta = true;
 		ArIsObjectReferenceCollector = true;
 		ArShouldSkipBulkData = true;
 		if (Obj->HasAnyFlags(RF_ClassDefaultObject))
 		{
-			Obj->GetClass()->SerializeDefaultObject((UObject*)Obj, *this);
+			Obj->GetClass()->SerializeBin(*this, (UObject*)Obj);
 		}
 		else
 		{
