[Backout] (CL 38187980) Add a new editor subsystem to assist with tracking contextual asset/path permissions in a restricted editor mode.

Notes:
- Deferring this work for now; keeping the unused code out of engine release side while we revisit.

diff --git a/Engine/Source/Editor/AssetPermissions/AssetPermissions.Build.cs b/Engine/Source/Editor/AssetPermissions/AssetPermissions.Build.cs
deleted file mode 100644
diff --git a/Engine/Source/Editor/AssetPermissions/Private/AssetPermissions.cpp b/Engine/Source/Editor/AssetPermissions/Private/AssetPermissions.cpp
deleted file mode 100644
diff --git a/Engine/Source/Editor/AssetPermissions/Private/AssetPermissionsModule.cpp b/Engine/Source/Editor/AssetPermissions/Private/AssetPermissionsModule.cpp
deleted file mode 100644
diff --git a/Engine/Source/Editor/AssetPermissions/Private/AssetPermissionsSubsystem.cpp b/Engine/Source/Editor/AssetPermissions/Private/AssetPermissionsSubsystem.cpp
deleted file mode 100644
diff --git a/Engine/Source/Editor/AssetPermissions/Public/AssetPermissions.h b/Engine/Source/Editor/AssetPermissions/Public/AssetPermissions.h
deleted file mode 100644
diff --git a/Engine/Source/Editor/AssetPermissions/Public/AssetPermissionsSubsystem.h b/Engine/Source/Editor/AssetPermissions/Public/AssetPermissionsSubsystem.h
deleted file mode 100644
diff --git a/Engine/Source/Editor/Kismet/Kismet.Build.cs b/Engine/Source/Editor/Kismet/Kismet.Build.cs
--- a/Engine/Source/Editor/Kismet/Kismet.Build.cs
+++ b/Engine/Source/Editor/Kismet/Kismet.Build.cs
@@ -61,7 +61,6 @@ public class Kismet : ModuleRules
 				"SubobjectDataInterface",
 				"ToolWidgets",
 				"TraceLog",
-				"AssetPermissions",
 			}
 			);
 
diff --git a/Engine/Source/Editor/Kismet/Private/BlueprintAssetPermissions.cpp b/Engine/Source/Editor/Kismet/Private/BlueprintAssetPermissions.cpp
deleted file mode 100644
diff --git a/Engine/Source/Editor/Kismet/Public/BlueprintAssetPermissions.h b/Engine/Source/Editor/Kismet/Public/BlueprintAssetPermissions.h
deleted file mode 100644
