Fix missing initializer

diff --git a/Engine/Source/Editor/Kismet/Internal/Blueprints/BlueprintDependencies.h b/Engine/Source/Editor/Kismet/Internal/Blueprints/BlueprintDependencies.h
--- a/Engine/Source/Editor/Kismet/Internal/Blueprints/BlueprintDependencies.h
+++ b/Engine/Source/Editor/Kismet/Internal/Blueprints/BlueprintDependencies.h
@@ -20,7 +20,7 @@ struct FBlueprintDependency
 	GENERATED_BODY()
 	
 	UPROPERTY()
-	EBPDependencyType DependencyType;
+	EBPDependencyType DependencyType = EBPDependencyType::Asset;
 
 	UPROPERTY()
 	FName PackageName;
