Disable fast reinstancing until a regression in anim blueprint reinstancing is fixed

#rnx

diff --git a/Engine/Source/Editor/Kismet/Private/BlueprintCompilationManager.cpp b/Engine/Source/Editor/Kismet/Private/BlueprintCompilationManager.cpp
--- a/Engine/Source/Editor/Kismet/Private/BlueprintCompilationManager.cpp
+++ b/Engine/Source/Editor/Kismet/Private/BlueprintCompilationManager.cpp
@@ -124,7 +124,7 @@ namespace UE::Kismet::BlueprintCompilationManager::Private
 		 * by compilation. This can only work when compiling in place (e.g. reusing UClasses) - so
 		 * can't be done for asset reload, live coding, or verse 
 		 */
-		static bool bAllowFastReinstancing = true;
+		static bool bAllowFastReinstancing = false; // Disabled while a crash when reinstancing animation blueprints is investigated
 		static FAutoConsoleVariableRef CVarAllowFastReinstancing(
 			TEXT("BP.bAllowFastReinstancing"), bAllowFastReinstancing,
 			TEXT("True by default, this flag enables a fast path in the logic that updates class \
