Cache BP Dependency hashes to avoid repeated identity calculation - e.g. repeatedly hashing a class's default values and layout. Cache is not threadsafe, but we could locok it if we like. Includes debugging facility for inspecting identity notion. This cache is necessary because repeatedly calculating identity for native dependencies takes on the order of hours during a full cook - with this cache we spend a total of ~2 minutes per worker in blueprint presave

#rnx

diff --git a/Engine/Source/Editor/Kismet/Private/Blueprints/BlueprintDependencies.cpp b/Engine/Source/Editor/Kismet/Private/Blueprints/BlueprintDependencies.cpp
--- a/Engine/Source/Editor/Kismet/Private/Blueprints/BlueprintDependencies.cpp
+++ b/Engine/Source/Editor/Kismet/Private/Blueprints/BlueprintDependencies.cpp
@@ -22,8 +22,21 @@
 	have their file contents hashed. There is ample opportunity for optimization here, but
 	there is even more ample opportunity for false positives and negatives. The current emphasis
 	is on simplicity, durability and correctness, at the possible expense of raw throughput.
+
+	Some notes on engine facilities we cannot (yet?) use:
+	::GetSchemaHash(UStruct*) does not include default values (for either scriptstructs or classes)
+		nor does it include any detection of versioning logic in user's (C++) serialize(FArchive&) 
+		routines
+	FAssetPackageData::GetPackageSavedHash could be used for uasset based dependencies, but is hard to access
+		and long term is imprecise (e.g. many changes to an asset do not have effects
+		on dependents)
+
+	@todo - include detection of custom object versions published by native class dependencies
+	@todo - include detection of custom object versions published by native scriptstruct dependencies? The payoff here may be small
 */
 
+#define UE_STORE_DEPENDENCY_SNAPSHOT 0
+
 namespace UE::Private
 {
 constexpr int32 BlueprintCookDependenciesVersion = 2;
@@ -46,6 +59,31 @@ static FString HashNativeClass( const UClass* NativeClass);
 // Asset canonization:
 static FString HashPackage( const UPackage* Package );
 static FString HashFile(const FString& Filename);
+
+// cache of dependency hashes, so they repeatedly recalculated:
+struct FBPDependencyCacheEntry
+{
+	FString Hash;
+#if UE_STORE_DEPENDENCY_SNAPSHOT
+	FUtf8String Source;
+#endif
+};
+
+class FBPDependencyCache
+{
+public:
+	static FBPDependencyCacheEntry* LookupCache(const UClass* Class);
+	static FBPDependencyCacheEntry* LookupCache(const UScriptStruct* Struct);
+	static FBPDependencyCacheEntry* LookupCache(const UPackage* Package);
+private:
+	static FBPDependencyCache& GetCache();
+	static TUniquePtr<FBPDependencyCache> GBPDependencyCache;
+
+	TMap<const UClass*, FBPDependencyCacheEntry> ClassToHash;
+	TMap<const UScriptStruct*, FBPDependencyCacheEntry> StructToHash;
+	TMap<const UPackage*, FBPDependencyCacheEntry> PackageToHash;
+};
+
 // Filesystem helper:
 static bool ReadBytesFromFile(const TCHAR* Filename, TArray64<uint8>& OutBytes);
 // Fail caching helper:
@@ -69,10 +107,12 @@ static TArray<const UObject*> UE::Private::GetAllImportedObjects(TConstArrayView
 			, SearchRoots(InRoots)
 			, References(InReferences)
 		{
-			// Copying FPackageHarvester:
-			SetIsPersistent(true);
 			SetIsSaving(true);
+			SetShouldSkipCompilingAssets(true);
+			SetWantBinaryPropertySerialization(true);
 			SetUseUnversionedPropertySerialization(true);
+			SetShouldSkipUpdateCustomVersion(true);
+			ArIsModifyingWeakAndStrongReferences = true;
 			ArIsObjectReferenceCollector = true;
 			ArShouldSkipBulkData = true;
 			if (Obj->HasAnyFlags(RF_ClassDefaultObject))
@@ -281,6 +321,13 @@ static void UE::Private::GetAllDependencies(
 
 static FString UE::Private::HashNativeStruct(const UScriptStruct* Struct)
 {
+	TRACE_CPUPROFILER_EVENT_SCOPE(BlueprintDependencies::HashNativeStruct);
+	FBPDependencyCacheEntry* CachedValue = FBPDependencyCache::LookupCache(Struct);
+	if(CachedValue->Hash.Len() > 0)
+	{
+		return *CachedValue->Hash;
+	}
+
 	check(Struct->StructFlags & STRUCT_Native);
 	FUtf8String StringifiedStruct = UE::JsonObjectGraph::Stringify(
 		{Struct}, 
@@ -297,11 +344,24 @@ static FString UE::Private::HashNativeStruct(const UScriptStruct* Struct)
 
 	FBlake3 Hash;
 	Hash.Update(StringifiedStruct.GetCharArray().GetData(), StringifiedStruct.GetCharArray().Num());
-	return LexToString(Hash.Finalize());
+	CachedValue->Hash = LexToString(Hash.Finalize());
+	
+#if UE_STORE_DEPENDENCY_SNAPSHOT
+	CachedValue->Source = StringifiedStruct;
+#endif
+
+	return CachedValue->Hash;
 }
 
 static FString UE::Private::HashNativeClass( const UClass* NativeClass)
 {
+	TRACE_CPUPROFILER_EVENT_SCOPE(BlueprintDependencies::HashNativeClass);
+	FBPDependencyCacheEntry* CachedValue = FBPDependencyCache::LookupCache(NativeClass);
+	if(CachedValue->Hash.Len() > 0)
+	{
+		return CachedValue->Hash;
+	}
+
 	// We're using UE::JsonObjectGraph because it is robust compared to ExportText and 
 	// the other core level facilities. Exporting a class to text using ExportText has
 	// always produced an empty object, which is not meaningful. JsonObjectGraph::Stringify
@@ -311,11 +371,24 @@ static FString UE::Private::HashNativeClass( const UClass* NativeClass)
 		FJsonStringifyOptions(EJsonStringifyFlags::FilterEditorOnlyData));
 	FBlake3 Hash;
 	Hash.Update(StringifiedClass.GetCharArray().GetData(), StringifiedClass.GetCharArray().Num());
-	return LexToString(Hash.Finalize());
+	CachedValue->Hash = LexToString(Hash.Finalize());
+
+#if UE_STORE_DEPENDENCY_SNAPSHOT
+	CachedValue->Source = StringifiedClass;
+#endif
+
+	return CachedValue->Hash;
 }
 
 static FString UE::Private::HashPackage(const UPackage* Package)
 {
+	TRACE_CPUPROFILER_EVENT_SCOPE(BlueprintDependencies::HashPackage);
+	FBPDependencyCacheEntry* CachedValue = FBPDependencyCache::LookupCache(Package);
+	if(CachedValue->Hash.Len() > 0)
+	{
+		return CachedValue->Hash;
+	}
+
 	const FName ParentPackageName = Package->GetFName();
 	// we don't want to have to load the asset so hash its file rather than the class itself:
 	FString Result;
@@ -323,7 +396,13 @@ static FString UE::Private::HashPackage(const UPackage* Package)
 		ParentPackageName.ToString(), Result, FPackageName::GetAssetPackageExtension());
 	if(bSuceeded)
 	{
-		return UE::Private::HashFile(Result);
+		CachedValue->Hash = UE::Private::HashFile(Result);
+		
+#if UE_STORE_DEPENDENCY_SNAPSHOT
+		CachedValue->Source = "File hashed on disk";
+#endif
+
+		return CachedValue->Hash;
 	}
 	return FString();
 }
@@ -342,6 +421,37 @@ static FString UE::Private::HashFile(const FString& Filename)
 	return LexToString(Hash.Finalize());
 }
 
+UE::Private::FBPDependencyCacheEntry* UE::Private::FBPDependencyCache::LookupCache(const UClass* Class)
+{
+	FBPDependencyCache& Cache = GetCache();
+	FBPDependencyCacheEntry& Value = Cache.ClassToHash.FindOrAdd(Class);
+	return &Value;
+}
+
+UE::Private::FBPDependencyCacheEntry* UE::Private::FBPDependencyCache::LookupCache(const UScriptStruct* Struct)
+{
+	FBPDependencyCache& Cache = GetCache();
+	FBPDependencyCacheEntry& Value = Cache.StructToHash.FindOrAdd(Struct);
+	return &Value;
+}
+
+UE::Private::FBPDependencyCacheEntry* UE::Private::FBPDependencyCache::LookupCache(const UPackage* Package)
+{
+	FBPDependencyCache& Cache = GetCache();
+	FBPDependencyCacheEntry& Value = Cache.PackageToHash.FindOrAdd(Package);
+	return &Value;
+}
+
+TUniquePtr<UE::Private::FBPDependencyCache> UE::Private::FBPDependencyCache::GBPDependencyCache;
+UE::Private::FBPDependencyCache& UE::Private::FBPDependencyCache::GetCache()
+{
+	if(!GBPDependencyCache)
+	{
+		GBPDependencyCache = MakeUnique<FBPDependencyCache>();
+	}
+	return *GBPDependencyCache;
+}
+
 static bool UE::Private::ReadBytesFromFile(const TCHAR* Filename, TArray64<uint8>& OutBytes)
 {
 	OutBytes.Reset();
@@ -548,9 +658,6 @@ void UE::Private::SaveBPCookDepenenciesImpl(FCbWriter& Writer, const FBlueprintD
 
 UE::Cook::FCookDependency UE::Private::BlueprintDependencies::RecordCookDependencies(const UBlueprint* BP)
 {
-	return UE::Cook::FCookDependency::Function(
-		UE_COOK_DEPENDENCY_FUNCTION_CALL(NotCacheable), FCbFieldIterator());
-#if 0
 	if(!BP->ParentClass || 
 		!BP->GeneratedClass ||
 		!BP->GeneratedClass->IsChildOf(UObject::StaticClass()))
@@ -569,5 +676,6 @@ UE::Cook::FCookDependency UE::Private::BlueprintDependencies::RecordCookDependen
 
 	return UE::Cook::FCookDependency::Function(
 		UE_COOK_DEPENDENCY_FUNCTION_CALL(ValidateBPCookDependenciesImpl), Writer.Save());
-#endif
 }
+
+#undef UE_STORE_DEPENDENCY_SNAPSHOT
