Added data validation to entire BP at compile time, rather than just the CDO. This makes it so any components added at the BP level will be validated. Previously, only native components would be validated.  Includes BP.bDoFullDataValidationDuringCompilation cvar to revert to old behavior if this causes any problems.

diff --git a/Engine/Source/Editor/Kismet/Private/BlueprintCompilationManager.cpp b/Engine/Source/Editor/Kismet/Private/BlueprintCompilationManager.cpp
--- a/Engine/Source/Editor/Kismet/Private/BlueprintCompilationManager.cpp
+++ b/Engine/Source/Editor/Kismet/Private/BlueprintCompilationManager.cpp
@@ -132,6 +132,13 @@ namespace UE::Kismet::BlueprintCompilationManager::Private
 				construction data for a match, if the layout and construction data match then we \
 				don't replace instances of the class."),
 			ECVF_Default);
+
+		/** Flag to choose between full Blueprint data validation at compile time, or just the CDO. */
+		static bool bDoFullDataValidationDuringCompilation = true;
+		static FAutoConsoleVariableRef CVarDoFullDataValidationDuringCompilation(
+			TEXT("BP.bDoFullDataValidationDuringCompilation"), bDoFullDataValidationDuringCompilation,
+			TEXT("If true, data validation will be performed on the entire BP when compiled. Otherwise, only the CDO will be validated. Consider disabling if errors or performance are a concern."),
+			ECVF_Default);
 	}
 }
 
@@ -1688,9 +1695,18 @@ void FBlueprintCompilationManagerImpl::FlushCompilationQueueImpl(bool bSuppressB
 							{
 								FKismetCompilerUtilities::ValidateEnumProperties(ClassDefaultObject, *CompilerData.ActiveResultsLog);
 
-								// Make sure any class-specific validation passes on the CDO
 								FDataValidationContext Context;
-								EDataValidationResult ValidateCDOResult = ClassDefaultObject->IsDataValid(/*out*/ Context);
+
+								// Make sure any class-specific validation passes 
+								if (UE::Kismet::BlueprintCompilationManager::Private::ConsoleVariables::bDoFullDataValidationDuringCompilation)
+								{
+									EDataValidationResult ValidateBlueprintResult = BP->IsDataValid(/*out*/ Context);
+								}
+								else
+								{
+									EDataValidationResult ValidateCDOResult = ClassDefaultObject->IsDataValid(/*out*/ Context);
+								}
+
 								if (Context.GetNumErrors() + Context.GetNumWarnings() > 0)
 								{
 									TArray<FText> Warnings, Errors;
