Harden build target set to avoid rare crashes when the project manager resets itself
#rnx

diff --git a/Engine/Source/Editor/Kismet/Internal/ProjectUtilities/BuildTargetSet.h b/Engine/Source/Editor/Kismet/Internal/ProjectUtilities/BuildTargetSet.h
--- a/Engine/Source/Editor/Kismet/Internal/ProjectUtilities/BuildTargetSet.h
+++ b/Engine/Source/Editor/Kismet/Internal/ProjectUtilities/BuildTargetSet.h
@@ -5,7 +5,7 @@
 #include "Containers/UnrealString.h"
 #include "Misc/EnumClassFlags.h"
 
-struct FModuleDescriptor;
+namespace EHostType { enum Type : int; }
 class UClass;
 class UFunction;
 class UPackage;
@@ -56,8 +56,8 @@ namespace UE::ProjectUtilities
 	ENUM_CLASS_FLAGS(FBuildTargetSet::EBuildTargetFlags);
 
 	/**
-	 * Given a native UPackage attempts to find the module descriptor.
-	 * Returns nullptr if it can not be found.
+	 * Given a native UPackage attempts to find the module descriptor host type, returning EHostType::Max 
+	 * if it cannot be determined.
 	 */
-	KISMET_API const FModuleDescriptor* FindModuleDescriptor(const UPackage* ForNativePackage);
+	KISMET_API const EHostType::Type FindModuleDescriptorHostType(const UPackage* ForNativePackage);
 } 
diff --git a/Engine/Source/Editor/Kismet/Kismet.Build.cs b/Engine/Source/Editor/Kismet/Kismet.Build.cs
--- a/Engine/Source/Editor/Kismet/Kismet.Build.cs
+++ b/Engine/Source/Editor/Kismet/Kismet.Build.cs
@@ -64,7 +64,15 @@ public class Kismet : ModuleRules
 			}
 			);
 
-        DynamicallyLoadedModuleNames.AddRange(
+		// we only 'need' live coding so that we can clear cached 
+		// native visibility information when a live coding event 
+		// occurs:
+		if (Target.bWithLiveCoding)
+		{
+			PrivateDependencyModuleNames.Add("LiveCoding");
+		}
+
+		DynamicallyLoadedModuleNames.AddRange(
             new string[] {
                 "BlueprintRuntime",
                 "ClassViewer",
diff --git a/Engine/Source/Editor/Kismet/Private/ProjectUtilities/BuildTargetSet.cpp b/Engine/Source/Editor/Kismet/Private/ProjectUtilities/BuildTargetSet.cpp
--- a/Engine/Source/Editor/Kismet/Private/ProjectUtilities/BuildTargetSet.cpp
+++ b/Engine/Source/Editor/Kismet/Private/ProjectUtilities/BuildTargetSet.cpp
@@ -8,11 +8,20 @@
 #include "Logging/StructuredLog.h"
 #include "ProfilingDebugging/ScopedTimers.h"
 #include "ProjectDescriptor.h"
+#include "IHotReload.h"
+#if WITH_LIVE_CODING
+#include "ILiveCodingModule.h"
+#endif
 
 namespace UE::Private
 {
 	const UClass* GetNativeClass(const UFunction* ForFunction);
 	const UClass* GetNativeClass(const UClass* ForClass);
+	void InvalidateCache(bool bIsAsyncCompile);
+	void InvalidateCache();
+
+	static TMap<FName, EHostType::Type> DescriptorCache;
+	static FCriticalSection DescriptorCacheLock;
 }
 
 const UClass* UE::Private::GetNativeClass(const UFunction* ForFunction)
@@ -32,17 +41,25 @@ const UClass* UE::Private::GetNativeClass(const UClass* ForClass)
 	return FBlueprintEditorUtils::FindFirstNativeClass(const_cast<UClass*>(ForClass));
 }
 
+void UE::Private::InvalidateCache(bool bIsAsyncCompile)
+{
+	InvalidateCache();
+}
 
-namespace UE::ProjectUtilities
+void UE::Private::InvalidateCache()
 {
-	static TMap<FName, const FModuleDescriptor*> DescriptorCache;
-	static FCriticalSection DescriptorCacheLock;
+	FScopeLock WriteLock(&DescriptorCacheLock);
+	DescriptorCache.Reset();
+}
 
-const FModuleDescriptor* FindModuleDescriptor(const UPackage* ForNativePackage)
+
+namespace UE::ProjectUtilities
+{
+const EHostType::Type FindModuleDescriptorHostType(const UPackage* ForNativePackage)
 {
 	if (!ForNativePackage)
 	{
-		return nullptr;
+		return EHostType::Max;
 	}
 
 	// The loops at the bottom are naive and may need to be optimized if we begin
@@ -50,37 +67,54 @@ const FModuleDescriptor* FindModuleDescriptor(const UPackage* ForNativePackage)
 	FString FunctionPackageName = ForNativePackage->GetName();
 	if (!FunctionPackageName.StartsWith(TEXT("/Script/")))
 	{
-		return nullptr;
+		return EHostType::Max;
 	}
 
 	FunctionPackageName = FunctionPackageName.RightChop(8); // Chop "/Script/" from the name
 	if (FunctionPackageName == TEXT("Engine"))
 	{ // common case, there is no ModuleDescriptor for Engine
-		return nullptr;
+		return EHostType::Max;
 	}
 
 	const IProjectManager& ProjectManager = IProjectManager::Get();
 	const FProjectDescriptor* PD = ProjectManager.GetCurrentProject();
 	if (!PD)
 	{
-		return nullptr;
+		return EHostType::Max;
 	}
 
 	const FName FunctionPackageFName(FunctionPackageName.Len(), *FunctionPackageName);
-	const auto AddToCache = [](FName Name, const FModuleDescriptor* Descriptor)
+	const auto AddToCache = [](FName Name, EHostType::Type HostType)
 		{
 			// must be a write lock:
-			FScopeLock WriteLock(&DescriptorCacheLock);
-			DescriptorCache.Add(Name, Descriptor);
+			FScopeLock WriteLock(&UE::Private::DescriptorCacheLock);
+			static bool bCheckedForHotReloadOrLiveCoding = false;
+			if(!bCheckedForHotReloadOrLiveCoding)
+			{
+				bCheckedForHotReloadOrLiveCoding = true;
+				// Register the hot-reload callback - invalidating our cache when a module is compiled:
+				if(IHotReloadModule::IsAvailable())
+				{
+					IHotReloadModule::Get().OnModuleCompilerStarted().AddStatic( &UE::Private::InvalidateCache );
+				}
+			
+	#if WITH_LIVE_CODING
+				if (ILiveCodingModule* LiveCoding = FModuleManager::LoadModulePtr<ILiveCodingModule>(LIVE_CODING_MODULE_NAME))
+				{
+					LiveCoding->GetOnPatchCompleteDelegate().AddStatic(&UE::Private::InvalidateCache);
+				}
+	#endif
+			}
+			UE::Private::DescriptorCache.Add(Name, HostType);
 		};
-	const auto FindFromCache = [](FName Name) -> TOptional<const FModuleDescriptor*>
+	const auto FindFromCache = [](FName Name) -> TOptional<EHostType::Type>
 		{
 			// could be a read lock:
-			FScopeLock FindLock(&DescriptorCacheLock);
-			const FModuleDescriptor** Result = DescriptorCache.Find(Name);
-			return Result ? *Result : TOptional<const FModuleDescriptor*>();
+			FScopeLock FindLock(&UE::Private::DescriptorCacheLock);
+			const EHostType::Type* Result = UE::Private::DescriptorCache.Find(Name);
+			return Result ? *Result : TOptional<EHostType::Type>();
 		};
-	TOptional<const FModuleDescriptor*> CachedDescriptor = FindFromCache(FunctionPackageFName);
+	TOptional<EHostType::Type> CachedDescriptor = FindFromCache(FunctionPackageFName);
 	if (CachedDescriptor)
 	{
 		return *CachedDescriptor;
@@ -92,8 +126,8 @@ const FModuleDescriptor* FindModuleDescriptor(const UPackage* ForNativePackage)
 	{
 		if (MD.Name == FunctionPackageFName)
 		{
-			AddToCache(FunctionPackageFName, &MD);
-			return &MD;
+			AddToCache(FunctionPackageFName, MD.Type);
+			return MD.Type;
 		}
 	}
 
@@ -111,14 +145,14 @@ const FModuleDescriptor* FindModuleDescriptor(const UPackage* ForNativePackage)
 		{
 			if (MD.Name == FunctionPackageFName)
 			{
-				AddToCache(FunctionPackageFName, &MD);
-				return &MD;
+				AddToCache(FunctionPackageFName, MD.Type);
+				return MD.Type;
 			}
 		}
 	}
 
-	AddToCache(FunctionPackageFName, nullptr);
-	return nullptr;
+	AddToCache(FunctionPackageFName, EHostType::Max);
+	return EHostType::Max;
 }
 
 FBuildTargetSet FBuildTargetSet::GetCallerTargetsUnsupportedByCallee(const UClass* Caller, const UFunction* Callee)
@@ -188,11 +222,12 @@ FBuildTargetSet::EBuildTargetFlags FBuildTargetSet::GetSupportedTargetsForNative
 		| EBuildTargetFlags::Client
 		| EBuildTargetFlags::Editor;
 	using namespace UE::Private;
-	const FModuleDescriptor* ModuleDescriptor = UE::ProjectUtilities::FindModuleDescriptor(NativeBase->GetPackage());
+	const UPackage* Package = NativeBase->GetPackage();
+	const EHostType::Type ModuleHostType = UE::ProjectUtilities::FindModuleDescriptorHostType(Package);
 	EBuildTargetFlags SupportedTargets = ALL;
-	if (ModuleDescriptor)
+	if (ModuleHostType != EHostType::Max)
 	{
-		switch (ModuleDescriptor->Type)
+		switch (ModuleHostType)
 		{
 		case EHostType::Runtime:
 		case EHostType::RuntimeNoCommandlet:
@@ -220,7 +255,7 @@ FBuildTargetSet::EBuildTargetFlags FBuildTargetSet::GetSupportedTargetsForNative
 				| EBuildTargetFlags::Editor;
 			break;
 		default:
-			ensureMsgf(false, TEXT("Encountered unexpected module type: %d in module %s"), ModuleDescriptor->Type, *ModuleDescriptor->Name.ToString());
+			ensureMsgf(false, TEXT("Encountered unexpected module type: %d in module %s"), ModuleHostType, Package ? *Package->GetName() : TEXT("nullpackage"));
 			break;
 		}
 	}
diff --git a/Engine/Source/Runtime/Projects/Public/ModuleDescriptor.h b/Engine/Source/Runtime/Projects/Public/ModuleDescriptor.h
--- a/Engine/Source/Runtime/Projects/Public/ModuleDescriptor.h
+++ b/Engine/Source/Runtime/Projects/Public/ModuleDescriptor.h
@@ -81,7 +81,7 @@ namespace ELoadingPhase
  */
 namespace EHostType
 {
-	enum Type
+	enum Type : int
 	{
 		// Loads on all targets, except programs.
 		Runtime,
