Changes to fix non-determinism in blueprint dependencies.

Specifically:
1. Serialization(Class): mask out the CLASS_ReplicationDataIsSetUp flag on classes as it is transient.
2. Serialization(Property): mask out the RepIndex member as it is non-deterministic and meant to be computed at runtime.
3. Don't traverse transient fields (fields which are marked as transient, not that they refer to data in a transient packagke) when gathering blueprint dependencies.

diff --git a/Engine/Source/Editor/Kismet/Private/Blueprints/BlueprintDependencies.cpp b/Engine/Source/Editor/Kismet/Private/Blueprints/BlueprintDependencies.cpp
--- a/Engine/Source/Editor/Kismet/Private/Blueprints/BlueprintDependencies.cpp
+++ b/Engine/Source/Editor/Kismet/Private/Blueprints/BlueprintDependencies.cpp
@@ -126,6 +126,22 @@ static TArray<const UObject*> UE::Private::GetAllImportedObjects(TConstArrayView
 		}
 
 	private:
+		bool ShouldTraverseProperty()
+		{
+			FProperty* SerializedProperty = GetSerializedProperty();
+			if (!SerializedProperty || !SerializedProperty->HasAnyFlags(RF_Transient))
+			{
+				return true;
+			}
+			return false;
+		}
+
+		virtual FArchive& operator<<(FLazyObjectPtr& Value) override { return ShouldTraverseProperty() ? FArchiveUObject::operator<<(Value) : *this; }
+		virtual FArchive& operator<<(FObjectPtr& Value) override { return ShouldTraverseProperty() ? FArchiveUObject::operator<<(Value) : *this; }
+		virtual FArchive& operator<<(FSoftObjectPtr& Value) override { return ShouldTraverseProperty() ? FArchiveUObject::operator<<(Value) : *this; }
+		virtual FArchive& operator<<(FSoftObjectPath& Value) override { return ShouldTraverseProperty() ? FArchiveUObject::operator<<(Value) : *this; }
+		virtual FArchive& operator<<(FWeakObjectPtr& Value) override { return ShouldTraverseProperty() ? FArchiveUObject::operator<<(Value) : *this; }
+
 		virtual FArchive& operator<<(UObject*& ObjRef) override
 		{
 			if (ObjRef != nullptr &&
diff --git a/Engine/Source/Runtime/CoreUObject/Private/UObject/Class.cpp b/Engine/Source/Runtime/CoreUObject/Private/UObject/Class.cpp
--- a/Engine/Source/Runtime/CoreUObject/Private/UObject/Class.cpp
+++ b/Engine/Source/Runtime/CoreUObject/Private/UObject/Class.cpp
@@ -5813,13 +5813,13 @@ void UClass::Serialize( FArchive& Ar )
 	if (Ar.IsSaving())
 	{
 		uint32 SavedClassFlags = ClassFlags;
-		SavedClassFlags &= ~(CLASS_ShouldNeverBeLoaded | CLASS_TokenStreamAssembled);
+		SavedClassFlags &= ~(CLASS_ShouldNeverBeLoaded | CLASS_TokenStreamAssembled | CLASS_ReplicationDataIsSetUp);
 		Ar << SavedClassFlags;
 	}
 	else if (Ar.IsLoading())
 	{
 		Ar << (uint32&)ClassFlags;
-		ClassFlags &= ~(CLASS_ShouldNeverBeLoaded | CLASS_TokenStreamAssembled);
+		ClassFlags &= ~(CLASS_ShouldNeverBeLoaded | CLASS_TokenStreamAssembled | CLASS_ReplicationDataIsSetUp);
 	}
 	else 
 	{
diff --git a/Engine/Source/Runtime/CoreUObject/Private/UObject/Property.cpp b/Engine/Source/Runtime/CoreUObject/Private/UObject/Property.cpp
--- a/Engine/Source/Runtime/CoreUObject/Private/UObject/Property.cpp
+++ b/Engine/Source/Runtime/CoreUObject/Private/UObject/Property.cpp
@@ -922,7 +922,8 @@ PRAGMA_ENABLE_DEPRECATION_WARNINGS
 		check( !IsEditorOnlyProperty() );
 	}
 
-	Ar << RepIndex;
+	decltype(RepIndex) DefaultRepIndex = 0;
+	Ar << DefaultRepIndex; // RepIndex is built lazily on demand. Serialize default value for backwards compat, but discard it.
 	Ar << RepNotifyFunc;
 
 	if (Ar.IsLoading())
