Change blueprint dependency gathering/evaluation to use the asset registry's knowledge of a package's hash (internally derived from its save guid) as an alternative to reading the package and hashing its bytes.

diff --git a/Engine/Source/Editor/Kismet/Private/Blueprints/BlueprintDependencies.cpp b/Engine/Source/Editor/Kismet/Private/Blueprints/BlueprintDependencies.cpp
--- a/Engine/Source/Editor/Kismet/Private/Blueprints/BlueprintDependencies.cpp
+++ b/Engine/Source/Editor/Kismet/Private/Blueprints/BlueprintDependencies.cpp
@@ -2,11 +2,15 @@
 
 #include "Blueprints/BlueprintDependencies.h"
 
+#include "AssetRegistry/AssetRegistryModule.h"
+#include "AssetRegistry/IAssetRegistry.h"
 #include "Engine/Blueprint.h"
 #include "HAL/FileManager.h"
+#include "IO/IoHash.h"
 #include "JsonObjectConverter.h"
 #include "JsonObjectGraph/Stringify.h"
 #include "Kismet2/BlueprintEditorUtils.h"
+#include "Misc/Optional.h"
 #include "Misc/ScopeExit.h"
 #include "Serialization/CompactBinary.h"
 #include "Serialization/CompactBinaryWriter.h"
@@ -39,7 +43,7 @@
 
 namespace UE::Private
 {
-constexpr int32 BlueprintCookDependenciesVersion = 2;
+constexpr int32 BlueprintCookDependenciesVersion = 3;
 
 // gathering helpers:
 static TArray<const UObject*> GetAllImportedObjects(TConstArrayView<const UObject*> Roots);
@@ -58,7 +62,7 @@ static FString HashNativeStruct(const UScriptStruct* Struct);
 static FString HashNativeClass( const UClass* NativeClass);
 // Asset canonization:
 static FString HashPackage( const UPackage* Package );
-static FString HashFile(const FString& Filename);
+static FString HashPackageFile(FName PackageName);
 
 // cache of dependency hashes, so they repeatedly recalculated:
 struct FBPDependencyCacheEntry
@@ -405,36 +409,22 @@ static FString UE::Private::HashPackage(const UPackage* Package)
 		return CachedValue->Hash;
 	}
 
-	const FName ParentPackageName = Package->GetFName();
-	// we don't want to have to load the asset so hash its file rather than the class itself:
-	FString Result;
-	const bool bSuceeded = FPackageName::TryConvertLongPackageNameToFilename(
-		ParentPackageName.ToString(), Result, FPackageName::GetAssetPackageExtension());
-	if(bSuceeded)
-	{
-		CachedValue->Hash = UE::Private::HashFile(Result);
-		
+	CachedValue->Hash = UE::Private::HashPackageFile(Package->GetFName());
+	
 #if UE_STORE_DEPENDENCY_SNAPSHOT
-		CachedValue->Source = "File hashed on disk";
+	CachedValue->Source = "File hashed on disk";
 #endif
 
-		return CachedValue->Hash;
-	}
-	return FString();
+	return CachedValue->Hash;
 }
 
-static FString UE::Private::HashFile(const FString& Filename)
+static FString UE::Private::HashPackageFile(FName PackageName)
 {
-	TArray64<uint8> FileData;
-	if(!ReadBytesFromFile(*Filename, FileData))
-	{
-		return FString();
-	}
-
-	FBlake3 Hash;
-	Hash.Update(FileData.GetData(), FileData.Num());
-
-	return LexToString(Hash.Finalize());
+	IAssetRegistry& AssetRegistry = FModuleManager::LoadModuleChecked<FAssetRegistryModule>("AssetRegistry").Get();
+	TOptional<FAssetPackageData> PackageData = AssetRegistry.GetAssetPackageDataCopy(PackageName);
+	check(PackageData);
+	FIoHash PackageHash = PackageData->GetPackageSavedHash();
+	return LexToString(PackageHash);
 }
 
 UE::Private::FBPDependencyCacheEntry* UE::Private::FBPDependencyCache::LookupCache(const UClass* Class)
@@ -507,20 +497,16 @@ void UE::Private::ValidateBPCookDependenciesImpl(FCbFieldViewIterator Args, UE::
 	
 	TRACE_CPUPROFILER_EVENT_SCOPE(BlueprintDependencies::ValidateBPCookDependenciesImpl);
 	// confirm dependencies do/do not match the artifact - must correspond to 
-	// GenerateBlueprintDependencies (hashpackage<->hashfile, hashnative class, hashnative struct)
+	// GenerateBlueprintDependencies (hashpackage<->hashpackagefile, hashnative class, hashnative struct)
 	for(const FBlueprintDependency& BPDep : Dependencies.BlueprintDependencies)
 	{
 		switch(BPDep.DependencyType)
 		{
 			case EBPDependencyType::Asset:
 			{
-				FString PackageFilename;
-				const bool bSuceeded = FPackageName::TryConvertLongPackageNameToFilename(
-					BPDep.PackageName.ToString(), PackageFilename,  FPackageName::GetAssetPackageExtension());
-				check(bSuceeded);
-				if(!BPDep.Hash.Equals(UE::Private::HashFile(PackageFilename)))
+				if(!BPDep.Hash.Equals(UE::Private::HashPackageFile(BPDep.PackageName)))
 				{
-					Context.LogInvalidated(FString::Printf(TEXT("Peer Package Changed: %s"), *PackageFilename));
+					Context.LogInvalidated(FString::Printf(TEXT("Peer Package Changed: %s"), *BPDep.PackageName.ToString()));
 					return;
 				}
 				break;
