[Backout] - CL39528039
#rnx
Original CL Desc
-----------------------------------------------------------------
Change blueprint dependency gathering/evaluation to use the asset registry's knowledge of a package's hash (internally derived from its save guid) as an alternative to reading the package and hashing its bytes.

diff --git a/Engine/Source/Editor/Kismet/Private/Blueprints/BlueprintDependencies.cpp b/Engine/Source/Editor/Kismet/Private/Blueprints/BlueprintDependencies.cpp
--- a/Engine/Source/Editor/Kismet/Private/Blueprints/BlueprintDependencies.cpp
+++ b/Engine/Source/Editor/Kismet/Private/Blueprints/BlueprintDependencies.cpp
@@ -2,15 +2,11 @@
 
 #include "Blueprints/BlueprintDependencies.h"
 
-#include "AssetRegistry/AssetRegistryModule.h"
-#include "AssetRegistry/IAssetRegistry.h"
 #include "Engine/Blueprint.h"
 #include "HAL/FileManager.h"
-#include "IO/IoHash.h"
 #include "JsonObjectConverter.h"
 #include "JsonObjectGraph/Stringify.h"
 #include "Kismet2/BlueprintEditorUtils.h"
-#include "Misc/Optional.h"
 #include "Misc/ScopeExit.h"
 #include "Serialization/CompactBinary.h"
 #include "Serialization/CompactBinaryWriter.h"
@@ -43,7 +39,7 @@
 
 namespace UE::Private
 {
-constexpr int32 BlueprintCookDependenciesVersion = 3;
+constexpr int32 BlueprintCookDependenciesVersion = 2;
 
 // gathering helpers:
 static TArray<const UObject*> GetAllImportedObjects(TConstArrayView<const UObject*> Roots);
@@ -62,7 +58,7 @@ static FString HashNativeStruct(const UScriptStruct* Struct);
 static FString HashNativeClass( const UClass* NativeClass);
 // Asset canonization:
 static FString HashPackage( const UPackage* Package );
-static FString HashPackageFile(FName PackageName);
+static FString HashFile(const FString& Filename);
 
 // cache of dependency hashes, so they repeatedly recalculated:
 struct FBPDependencyCacheEntry
@@ -409,22 +405,36 @@ static FString UE::Private::HashPackage(const UPackage* Package)
 		return CachedValue->Hash;
 	}
 
-	CachedValue->Hash = UE::Private::HashPackageFile(Package->GetFName());
-	
+	const FName ParentPackageName = Package->GetFName();
+	// we don't want to have to load the asset so hash its file rather than the class itself:
+	FString Result;
+	const bool bSuceeded = FPackageName::TryConvertLongPackageNameToFilename(
+		ParentPackageName.ToString(), Result, FPackageName::GetAssetPackageExtension());
+	if(bSuceeded)
+	{
+		CachedValue->Hash = UE::Private::HashFile(Result);
+		
 #if UE_STORE_DEPENDENCY_SNAPSHOT
-	CachedValue->Source = "File hashed on disk";
+		CachedValue->Source = "File hashed on disk";
 #endif
 
-	return CachedValue->Hash;
+		return CachedValue->Hash;
+	}
+	return FString();
 }
 
-static FString UE::Private::HashPackageFile(FName PackageName)
+static FString UE::Private::HashFile(const FString& Filename)
 {
-	IAssetRegistry& AssetRegistry = FModuleManager::LoadModuleChecked<FAssetRegistryModule>("AssetRegistry").Get();
-	TOptional<FAssetPackageData> PackageData = AssetRegistry.GetAssetPackageDataCopy(PackageName);
-	check(PackageData);
-	FIoHash PackageHash = PackageData->GetPackageSavedHash();
-	return LexToString(PackageHash);
+	TArray64<uint8> FileData;
+	if(!ReadBytesFromFile(*Filename, FileData))
+	{
+		return FString();
+	}
+
+	FBlake3 Hash;
+	Hash.Update(FileData.GetData(), FileData.Num());
+
+	return LexToString(Hash.Finalize());
 }
 
 UE::Private::FBPDependencyCacheEntry* UE::Private::FBPDependencyCache::LookupCache(const UClass* Class)
@@ -497,16 +507,20 @@ void UE::Private::ValidateBPCookDependenciesImpl(FCbFieldViewIterator Args, UE::
 	
 	TRACE_CPUPROFILER_EVENT_SCOPE(BlueprintDependencies::ValidateBPCookDependenciesImpl);
 	// confirm dependencies do/do not match the artifact - must correspond to 
-	// GenerateBlueprintDependencies (hashpackage<->hashpackagefile, hashnative class, hashnative struct)
+	// GenerateBlueprintDependencies (hashpackage<->hashfile, hashnative class, hashnative struct)
 	for(const FBlueprintDependency& BPDep : Dependencies.BlueprintDependencies)
 	{
 		switch(BPDep.DependencyType)
 		{
 			case EBPDependencyType::Asset:
 			{
-				if(!BPDep.Hash.Equals(UE::Private::HashPackageFile(BPDep.PackageName)))
+				FString PackageFilename;
+				const bool bSuceeded = FPackageName::TryConvertLongPackageNameToFilename(
+					BPDep.PackageName.ToString(), PackageFilename,  FPackageName::GetAssetPackageExtension());
+				check(bSuceeded);
+				if(!BPDep.Hash.Equals(UE::Private::HashFile(PackageFilename)))
 				{
-					Context.LogInvalidated(FString::Printf(TEXT("Peer Package Changed: %s"), *BPDep.PackageName.ToString()));
+					Context.LogInvalidated(FString::Printf(TEXT("Peer Package Changed: %s"), *PackageFilename));
 					return;
 				}
 				break;
