Reenable fast reinstancing - Don't include objects that have been moved to a new CDO by a custom blueprint compiler in the reference replacement map, this prevents replacing them with references to other objects.
#rnx

diff --git a/Engine/Source/Editor/Kismet/Private/BlueprintCompilationManager.cpp b/Engine/Source/Editor/Kismet/Private/BlueprintCompilationManager.cpp
--- a/Engine/Source/Editor/Kismet/Private/BlueprintCompilationManager.cpp
+++ b/Engine/Source/Editor/Kismet/Private/BlueprintCompilationManager.cpp
@@ -124,7 +124,7 @@ namespace UE::Kismet::BlueprintCompilationManager::Private
 		 * by compilation. This can only work when compiling in place (e.g. reusing UClasses) - so
 		 * can't be done for asset reload, live coding, or verse 
 		 */
-		static bool bAllowFastReinstancing = false;
+		static bool bAllowFastReinstancing = true;
 		static FAutoConsoleVariableRef CVarAllowFastReinstancing(
 			TEXT("BP.bAllowFastReinstancing"), bAllowFastReinstancing,
 			TEXT("True by default, this flag enables a fast path in the logic that updates class \
@@ -2659,15 +2659,30 @@ void FBlueprintCompilationManagerImpl::ReinstanceBatch(TArray<FReinstancingJob>&
 					FBlueprintCompileReinstancer::CopyPropertiesForUnrelatedObjects(Pair.Key, Pair.Value, /*bClearExternalReferences*/true, bUseDeltaSerialization, /*bOnlyHandleDirectSubObjects*/true, &OldToNewInstanceMap, &InOutOldToNewClassMap);
 				}
 
-				if (OldToNewTemplates && !OldToNewInstanceMap.IsEmpty())
-				{
-					OldToNewTemplates->FindOrAdd(OldClass).Append(OldToNewInstanceMap);
-				}
-
 				if (ReinstancingJob.Compiler.IsValid())
 				{
 					ReinstancingJob.Compiler->PropagateValuesToCDO(NewCDO, OldCDO);
 				}
+				
+				// PropagateValuesToCDO may 'reuse' objects from the old CDO (e.g. 
+				// FDisplayClusterConfiguratorKismetCompilerContext::CopyTermDefaultsToDefaultObject,
+				// others that make use of SaveSubObjectsFromCleanAndSanitizeClass)
+				// detect this and avoid creating mappings from 'old to new' when
+				// an object has been moved to the new CDO:
+				if (OldToNewTemplates && !OldToNewInstanceMap.IsEmpty())
+				{
+					TMap<UObject*, UObject*>& OldToNewMap = OldToNewTemplates->FindOrAdd(OldClass);
+					for(const TPair<UObject*, UObject*> OldToNew : OldToNewInstanceMap)
+					{
+						// dead objects are definitely old, as are objects that aren't within the NewCDO
+						// hopefully there is no user logic attempting further subobject reuse after 
+						// PropagateValuesToCDO
+						if(!IsValid(OldToNew.Key) || !OldToNew.Key->IsIn(NewCDO) )
+						{
+							OldToNewMap.Add(OldToNew.Key, OldToNew.Value);
+						}
+					}
+				}
 			}
 		}
 
