Disable fast reinstancing, we have found another crash in control rig

#rnx

diff --git a/Engine/Source/Editor/Kismet/Private/BlueprintCompilationManager.cpp b/Engine/Source/Editor/Kismet/Private/BlueprintCompilationManager.cpp
--- a/Engine/Source/Editor/Kismet/Private/BlueprintCompilationManager.cpp
+++ b/Engine/Source/Editor/Kismet/Private/BlueprintCompilationManager.cpp
@@ -120,11 +120,11 @@ namespace UE::Kismet::BlueprintCompilationManager::Private
 			ECVF_Default);
 
 		/** 
-		 * Flag to disable skipping of reinstancing logic when a class is not observabling changed
+		 * Flag to disable skipping of reinstancing logic when a class is not observably changed
 		 * by compilation. This can only work when compiling in place (e.g. reusing UClasses) - so
 		 * can't be done for asset reload, live coding, or verse 
 		 */
-		static bool bAllowFastReinstancing = true;
+		static bool bAllowFastReinstancing = false;
 		static FAutoConsoleVariableRef CVarAllowFastReinstancing(
 			TEXT("BP.bAllowFastReinstancing"), bAllowFastReinstancing,
 			TEXT("True by default, this flag enables a fast path in the logic that updates class \
