Disallow reference replacement on derived class artifacts left over from a previous compile during Blueprint reinstancing.

These may be kept alive due to e.g. transaction records that can still reference reinstanced build artifacts across a Blueprint compile.

diff --git a/Engine/Source/Editor/Kismet/Private/BlueprintCompilationManager.cpp b/Engine/Source/Editor/Kismet/Private/BlueprintCompilationManager.cpp
--- a/Engine/Source/Editor/Kismet/Private/BlueprintCompilationManager.cpp
+++ b/Engine/Source/Editor/Kismet/Private/BlueprintCompilationManager.cpp
@@ -2149,6 +2149,8 @@ void FBlueprintCompilationManagerImpl::ReparentHierarchies(const TMap<UClass*, U
 	const bool bReplaceReferencesToOldClasses = (Options & EReparentClassOptions::ReplaceReferencesToOldClasses) != EReparentClassOptions::None;
 	TRACE_CPUPROFILER_EVENT_SCOPE(ReparentHierarchies);
 
+	TSet<UObject*> OldObjects;
+
 	// something has decided to replace instances of a class. We need to update all the children of those types:
 	TArray<UClass*> ClassesOrdered;
 	// Map used to distinguish between new classes and classes that need to be reinstanced (reparented) via a new reinstancer:
@@ -2184,6 +2186,9 @@ void FBlueprintCompilationManagerImpl::ReparentHierarchies(const TMap<UClass*, U
 				if (DerivedClass->HasAnyClassFlags(CLASS_NewerVersionExists))
 				{
 					// Don't reinstance artifacts from classes that have been previously been reinstanced/trashed but aren't cleaned up yet.
+					// Include class in old objects that don't require reference replacement; otherwise, we might unintentionally invalidate
+					// parts of its linked structure (e.g. w/ a reinstanced base) which would no longer correlate to artifacts of the old subtype.
+					OldObjects.Add(DerivedClass);
 					continue;
 				}
 
@@ -2287,7 +2292,6 @@ void FBlueprintCompilationManagerImpl::ReparentHierarchies(const TMap<UClass*, U
 	BatchOptions.bReplaceReferencesToOldClasses = bReplaceReferencesToOldClasses;
 
 	// Make sure we don't replace old instances that are in the *callers* old to new TMap!
-	TSet<UObject*> OldObjects;
 	for (TPair<UClass*, UClass*> OldToNew : OldClassToNewClassDerivedTypes)
 	{
 		ensure(OldToNew.Value->HasAnyClassFlags(CLASS_TokenStreamAssembled));
