[BPCM] Retain archetypes cached during the CDO/archetype reinstancing phase that may be needed later to resolve old instance subgraph archetypes during non-archetype reinstancing.

Also hooked the editor's archetype cache into core UObject archetype lookup logic for outer chains. This maintains archetype coherency during reinstancing if the chain is not fully reinstanced.

diff --git a/Engine/Source/Editor/Kismet/Private/BlueprintCompilationManager.cpp b/Engine/Source/Editor/Kismet/Private/BlueprintCompilationManager.cpp
--- a/Engine/Source/Editor/Kismet/Private/BlueprintCompilationManager.cpp
+++ b/Engine/Source/Editor/Kismet/Private/BlueprintCompilationManager.cpp
@@ -210,6 +210,7 @@ struct FBlueprintCompilationManagerImpl : public FGCObject
 	// populated by FlushCompilationQueueImpl, cleared by FlushReinstancingQueueImpl:
 	TMap<TObjectPtr<UClass>, TObjectPtr<UClass>> ClassesToReinstance;
 	TMap<TObjectPtr<const UClass>, FBlake3Hash> ClassHashes;
+	TSet<TObjectPtr<UObject>> ObjectsWithCachedArchetypesForReinstancing;
 	
 	// Map to old default values, useful for providing access to this data throughout
 	// the compilation process:
@@ -266,6 +267,7 @@ void FBlueprintCompilationManagerImpl::AddReferencedObjects(FReferenceCollector&
 
 	Collector.AddReferencedObjects(ClassesToReinstance);
 	Collector.AddReferencedObjects(ClassHashes);
+	Collector.AddReferencedObjects(ObjectsWithCachedArchetypesForReinstancing);
 	Collector.AddReferencedObjects(OldCDOs);
 }
 
@@ -585,6 +587,12 @@ struct FReinstancingJob
 
 	// Old archetype to re-instantiate and its old associated template
 	TArray<FArchetypeInfo> OldArchetypeObjects;
+
+	// Keeps track of old archetype objects with cached archetypes which may be needed for reinstancing
+	// of non-archetypes to assist with archetype lookup when they have partially-reinstanced outer chains.
+	// We track these objects because they are annotated by the global archetype cache manager, and we'll
+	// need to clear out those annotations after we're done reinstancing all of the non-template instances.
+	TSet<UObject*> ArchetypeInstancesWithCachedArchetypes;
 };
 
 FReinstancingJob::FReinstancingJob(TSharedPtr<FBlueprintCompileReinstancer> InReinstancer)
@@ -1612,6 +1620,16 @@ void FBlueprintCompilationManagerImpl::FlushCompilationQueueImpl(bool bSuppressB
 			FScopedDurationTimer ReinstTimer(TimeReinstancing);
 			ReinstanceBatch(Reinstancers, MutableView(ClassesToReinstance), InLoadContext, OldToNewTemplates);
 
+			// Take ownership of objects with cached archetypes and defer cleanup until after reinstancing is complete
+			for (FReinstancingJob& ReinstancingJob : Reinstancers)
+			{
+				if (ReinstancingJob.ArchetypeInstancesWithCachedArchetypes.Num() > 0)
+				{
+					ObjectsWithCachedArchetypesForReinstancing.Append(ObjectPtrWrap(MoveTemp(ReinstancingJob.ArchetypeInstancesWithCachedArchetypes)));
+					ReinstancingJob.ArchetypeInstancesWithCachedArchetypes = {};
+				}
+			}
+
 			// We purposefully do not remove the OldCDOs yet, need to keep them in memory past first GC
 		}
 
@@ -1996,6 +2014,20 @@ void FBlueprintCompilationManagerImpl::FlushReinstancingQueueImpl(bool bFindAndR
 		FBlueprintCompileReinstancer::BatchReplaceInstancesOfClass(ClassesToReinstanceOwned, Options);
 		ClassHashes.Reset();
 
+		// Clean up objects with cached archetypes now that reinstancing is complete - these will have been
+		// carried over from any archetype objects we reinstanced during the compile phase
+		if (ObjectsWithCachedArchetypesForReinstancing.Num() > 0)
+		{
+			TSet<UObject*> ObjectsWithCachedArchetypes = ObjectPtrDecay(MoveTemp(ObjectsWithCachedArchetypesForReinstancing));
+			ObjectsWithCachedArchetypesForReinstancing = {};
+
+			FEditorCacheArchetypeManager& CacheManager = FEditorCacheArchetypeManager::Get();
+			for (UObject* ObjectWithCachedArchetype : ObjectsWithCachedArchetypes)
+			{
+				CacheManager.ResetCacheArchetype(ObjectWithCachedArchetype);
+			}
+		}
+
 		// Special case when we run on ALT, we want to cleanup all classes flagged for reinstantiation right away.
 		const bool bIsInActualAsyncLoadingThread = IsInAsyncLoadingThread() && !IsInGameThread();
 		if (IsAsyncLoading() && (!IsAsyncLoadingMultithreaded() || !bIsInActualAsyncLoadingThread))
@@ -2282,9 +2314,17 @@ void FBlueprintCompilationManagerImpl::ReparentHierarchies(const TMap<UClass*, U
 		OldClassToNewClassDerivedTypes = OldClassToNewClassIncludingChildren;
 	}
 	
-	for (const FReinstancingJob& ReinstancingJob : Reinstancers)
+	TSet<UObject*> ObjectsWithCachedArchetypes;
+	for (FReinstancingJob& ReinstancingJob : Reinstancers)
 	{
 		OldClassToNewClassDerivedTypes.Add(ReinstancingJob.OldToNew);
+
+		// Consolidate the set of objects with cached archetypes that we need for reinstancing
+		if (ReinstancingJob.ArchetypeInstancesWithCachedArchetypes.Num() > 0)
+		{
+			ObjectsWithCachedArchetypes.Append(MoveTemp(ReinstancingJob.ArchetypeInstancesWithCachedArchetypes));
+			ReinstancingJob.ArchetypeInstancesWithCachedArchetypes = {};
+		}
 	}
 	TGuardValue<std::atomic<bool>, bool> ReinstancingGuard(GIsReinstancing, true);
 	FReplaceInstancesOfClassParameters BatchOptions;
@@ -2313,6 +2353,16 @@ void FBlueprintCompilationManagerImpl::ReparentHierarchies(const TMap<UClass*, U
 	BatchOptions.OldToNewTemplates = &OldToNewTemplates;
 
 	FBlueprintCompileReinstancer::BatchReplaceInstancesOfClass(OldClassToNewClassDerivedTypes, BatchOptions );
+
+	// Clean up cached archetype objects now that reinstancing is complete
+	if (ObjectsWithCachedArchetypes.Num() > 0)
+	{
+		FEditorCacheArchetypeManager& CacheManager = FEditorCacheArchetypeManager::Get();
+		for (UObject* ObjectWithCachedArchetype : ObjectsWithCachedArchetypes)
+		{
+			CacheManager.ResetCacheArchetype(ObjectWithCachedArchetype);
+		}
+	}
 }
 
 
@@ -2881,10 +2931,12 @@ void FBlueprintCompilationManagerImpl::ReinstanceBatch(TArray<FReinstancingJob>&
 		}
 	}
 
+	FEditorCacheArchetypeManager& CacheManager = FEditorCacheArchetypeManager::Get();
+
 	// This loop finishes the reinstancing of archetypes after the entire Outer hierarchy has been updated with new instances:
-	for (const FReinstancingJob* ReinstancingJobPtr : ReinstancersPtr)
+	for (FReinstancingJob* ReinstancingJobPtr : ReinstancersPtr)
 	{
-		const FReinstancingJob& ReinstancingJob = *ReinstancingJobPtr;
+		FReinstancingJob& ReinstancingJob = *ReinstancingJobPtr;
 		UClass* OldClass = ReinstancingJob.OldToNew.Key;
 		if(OldClass)
 		{
@@ -2892,8 +2944,6 @@ void FBlueprintCompilationManagerImpl::ReinstanceBatch(TArray<FReinstancingJob>&
 
 			UClass* NewClass = ReinstancingJob.OldToNew.Value;
 
-			TSet<UObject*> CachedArchetypeObjects;
-			FEditorCacheArchetypeManager& CacheManager = FEditorCacheArchetypeManager::Get();
 			TMap<UObject*, UObject*>* OldToNewTemplatesForClass = OldToNewTemplates ? &OldToNewTemplates->FindOrAdd(OldClass) : nullptr;
 
 			for(const FReinstancingJob::FArchetypeInfo& OldArchetypeInfo : ReinstancingJob.OldArchetypeObjects)
@@ -2908,19 +2958,13 @@ void FBlueprintCompilationManagerImpl::ReinstanceBatch(TArray<FReinstancingJob>&
 				// We need to cache the archetype of the instances of this archetype
 				// as it will not be possible to get them afterwards as it gets renamed
 				// These cached archetypes will not be updated if they were set earlier
+				// Note we might need these archetypes to reinstance non-archetypes later
 				TArray<UObject*> ArchetypeInstances;
 				OldInstance->GetArchetypeInstances(ArchetypeInstances);
-				for (UObject* ArchetypeInstance : ArchetypeInstances)
-				{
-					CacheManager.CacheArchetype(ArchetypeInstance);
-					CachedArchetypeObjects.Add(ArchetypeInstance);
-
-					ForEachObjectWithOuter(ArchetypeInstance, [&CachedArchetypeObjects, &CacheManager](UObject* SubObject)
-					{
-						CacheManager.CacheArchetype(SubObject);
-						CachedArchetypeObjects.Add(SubObject);
-					});
-				}
+				FBlueprintCompileReinstancer::CacheArchetypes(
+					ArchetypeInstances,
+					CacheManager,
+					ReinstancingJob.ArchetypeInstancesWithCachedArchetypes);
 
 				UObject* Destination = GetTransientOuterForRename(OldInstance->GetClass());
 				OldInstance->Rename(
@@ -2987,13 +3031,6 @@ void FBlueprintCompilationManagerImpl::ReinstanceBatch(TArray<FReinstancingJob>&
 					OldInstance->MarkAsGarbage();
 				}
 			}
-
-			// Reset any cached archetypes
-			for (UObject* CachedArchetypeObject : CachedArchetypeObjects)
-			{
-				CacheManager.ResetCacheArchetype(CachedArchetypeObject);
-			}
-
 		}
 	}
 
diff --git a/Engine/Source/Editor/UnrealEd/Public/Kismet2/KismetReinstanceUtilities.h b/Engine/Source/Editor/UnrealEd/Public/Kismet2/KismetReinstanceUtilities.h
--- a/Engine/Source/Editor/UnrealEd/Public/Kismet2/KismetReinstanceUtilities.h
+++ b/Engine/Source/Editor/UnrealEd/Public/Kismet2/KismetReinstanceUtilities.h
@@ -338,7 +338,7 @@ private:
 	static UNREALED_API bool ReinstancerOrderingFunction(UClass* A, UClass* B);
 	
 	/** Helper for caching archetypes before objects are reconstructed */
-	static void CacheArchetypes(	
+	static UNREALED_API void CacheArchetypes(	
 		const TArray<UObject*>& ObjectsToReplace, 
 		FEditorCacheArchetypeManager& CacheManager, 
 		TSet<UObject*>& CachedArchetypeObjects );
diff --git a/Engine/Source/Runtime/CoreUObject/Private/UObject/UObjectArchetype.cpp b/Engine/Source/Runtime/CoreUObject/Private/UObject/UObjectArchetype.cpp
--- a/Engine/Source/Runtime/CoreUObject/Private/UObject/UObjectArchetype.cpp
+++ b/Engine/Source/Runtime/CoreUObject/Private/UObject/UObjectArchetype.cpp
@@ -111,6 +111,10 @@ UObject* GetArchetypeFromRequiredInfoImpl(const UClass* Class, const UObject* Ou
 		
 #if WITH_EDITOR
 			ArchetypeToSearch = Policy ? Policy->GetArchetype(Outer) : nullptr;
+			if (!ArchetypeToSearch)
+			{
+				ArchetypeToSearch = FEditorCacheArchetypeManager::Get().GetCachedArchetype(Outer);
+			}
 #endif
 
 			if (!ArchetypeToSearch)
