[UE5]: Fix for a regression in Blueprint reinstancing that could lead to data loss of an editor-spawned nested instanced subobject archetype.

Notes:
- This CL also adds a guard against a potential assert should we happen to encounter a NULL entry after reference replacement in the archetype instance cache (followup to CL 39949853).

diff --git a/Engine/Source/Editor/Kismet/Private/BlueprintCompilationManager.cpp b/Engine/Source/Editor/Kismet/Private/BlueprintCompilationManager.cpp
--- a/Engine/Source/Editor/Kismet/Private/BlueprintCompilationManager.cpp
+++ b/Engine/Source/Editor/Kismet/Private/BlueprintCompilationManager.cpp
@@ -2024,7 +2024,13 @@ void FBlueprintCompilationManagerImpl::FlushReinstancingQueueImpl(bool bFindAndR
 			FEditorCacheArchetypeManager& CacheManager = FEditorCacheArchetypeManager::Get();
 			for (UObject* ObjectWithCachedArchetype : ObjectsWithCachedArchetypes)
 			{
-				CacheManager.ResetCacheArchetype(ObjectWithCachedArchetype);
+				// It's possible that references to old instances for which we've cached archetypes can get replaced with
+				// NULL in the source map via ARO during reference replacement, which could otherwise assert here, because
+				// this API requires the input reference to be valid as a precondition.
+				if (ObjectWithCachedArchetype)
+				{
+					CacheManager.ResetCacheArchetype(ObjectWithCachedArchetype);
+				}
 			}
 		}
 
@@ -3017,7 +3023,7 @@ void FBlueprintCompilationManagerImpl::ReinstanceBatch(TArray<FReinstancingJob>&
 					{
 						OldToNewInstanceMap.Append(*OldToNewTemplateMap);
 					}
-					for (const auto& Pair : OrderedListOfObjectToCopy)
+					for (const TTuple<UObject*, UObject*>& Pair : OrderedListOfObjectToCopy)
 					{
 						FBlueprintCompileReinstancer::CopyPropertiesForUnrelatedObjects(Pair.Key, Pair.Value, /*bClearExternalReferences*/false, bUseDeltaSerialization, /*bOnlyHandleDirectSubObjects*/true, &OldToNewInstanceMap, &InOutOldToNewClassMap);
 					}
@@ -3027,6 +3033,17 @@ void FBlueprintCompilationManagerImpl::ReinstanceBatch(TArray<FReinstancingJob>&
 						OldToNewTemplates->FindOrAdd(OldClass).Append(OldToNewInstanceMap);
 					}
 
+					// Ensure that our archetype instance mapping is up to date, as this will be used for reference replacement below. Note that this
+					// might be populated by BuildDSO() initially to contain instance mapping pairs for which the new instance won't be valid until we
+					// run through each of the PreCreateSubObjectsForReinstantiation() and subsequent CopyPropertiesForUnrelatedObjects() steps above.
+					for (const TPair<UObject*, UObject*>& Pair : OldToNewInstanceMap)
+					{
+						if (UObject** NewArchetypeMapping = OldArchetypeToNewArchetype.Find(Pair.Key))
+						{
+							*NewArchetypeMapping = Pair.Value;
+						}
+					}
+
 					OldInstance->RemoveFromRoot();
 					OldInstance->MarkAsGarbage();
 				}
