[Truncation Warnings] Update GraphEditor Drag Drop Actions API to use FVector2f

diff --git a/Engine/Plugins/Animation/DeformerGraph/Source/OptimusEditor/Private/Widgets/SOptimusEditorGraphExplorerActions.cpp b/Engine/Plugins/Animation/DeformerGraph/Source/OptimusEditor/Private/Widgets/SOptimusEditorGraphExplorerActions.cpp
--- a/Engine/Plugins/Animation/DeformerGraph/Source/OptimusEditor/Private/Widgets/SOptimusEditorGraphExplorerActions.cpp
+++ b/Engine/Plugins/Animation/DeformerGraph/Source/OptimusEditor/Private/Widgets/SOptimusEditorGraphExplorerActions.cpp
@@ -57,8 +57,8 @@ TSharedRef<FOptimusEditorGraphDragAction_Binding> FOptimusEditorGraphDragAction_
 
 FReply FOptimusEditorGraphDragAction_Binding::DroppedOnPanel(
 	const TSharedRef<SWidget>& InPanel,
-	FVector2D InScreenPosition,
-	FVector2D InGraphPosition,
+	const FVector2f& InScreenPosition,
+	const FVector2f& InGraphPosition,
 	UEdGraph& InGraph
 	)
 {
@@ -81,7 +81,7 @@ FReply FOptimusEditorGraphDragAction_Binding::DroppedOnPanel(
 
 	UOptimusNodeGraph* ModelGraph = Graph->GetModelGraph();
 
-	ModelGraph->AddComponentBindingGetNode(Binding, InGraphPosition);
+	ModelGraph->AddComponentBindingGetNode(Binding, FDeprecateSlateVector2D(InGraphPosition));
 
 	return FReply::Handled();
 }
@@ -109,8 +109,8 @@ void FOptimusEditorGraphDragAction_Variable::HoverTargetChanged()
 
 FReply FOptimusEditorGraphDragAction_Variable::DroppedOnPanel(
 	const TSharedRef<SWidget>& InPanel, 
-	FVector2D InScreenPosition, 
-	FVector2D InGraphPosition, 
+	const FVector2f& InScreenPosition,
+	const FVector2f& InGraphPosition,
 	UEdGraph& InGraph
 	)
 {
@@ -133,7 +133,7 @@ FReply FOptimusEditorGraphDragAction_Variable::DroppedOnPanel(
 
 	UOptimusNodeGraph* ModelGraph = Graph->GetModelGraph();
 
-	ModelGraph->AddVariableGetNode(VariableDesc, InGraphPosition);
+	ModelGraph->AddVariableGetNode(VariableDesc, FDeprecateSlateVector2D(InGraphPosition));
 
 	return FReply::Handled();
 }
@@ -172,8 +172,8 @@ void FOptimusEditorGraphDragAction_Resource::HoverTargetChanged()
 
 FReply FOptimusEditorGraphDragAction_Resource::DroppedOnPanel(
 	const TSharedRef<SWidget>& InPanel, 
-	FVector2D InScreenPosition, 
-	FVector2D InGraphPosition, 
+	const FVector2f& InScreenPosition,
+	const FVector2f& InGraphPosition,
 	UEdGraph& InGraph
 	)
 {
@@ -207,7 +207,7 @@ FReply FOptimusEditorGraphDragAction_Resource::DroppedOnPanel(
 		FSlateIcon(),
 		FUIAction(
 			FExecuteAction::CreateLambda([ModelGraph, ResourceDesc, InGraphPosition]() {
-				ModelGraph->AddResourceNode(ResourceDesc, InGraphPosition);
+				ModelGraph->AddResourceNode(ResourceDesc, FDeprecateSlateVector2D(InGraphPosition));
 			}),
 			FCanExecuteAction()));
 
@@ -217,7 +217,7 @@ FReply FOptimusEditorGraphDragAction_Resource::DroppedOnPanel(
 	    FSlateIcon(),
 	    FUIAction(
 	        FExecuteAction::CreateLambda([ModelGraph, ResourceDesc, InGraphPosition]() { 
-				ModelGraph->AddResourceGetNode(ResourceDesc, InGraphPosition); 
+				ModelGraph->AddResourceGetNode(ResourceDesc, FDeprecateSlateVector2D(InGraphPosition)); 
 			}), 
 			FCanExecuteAction()));
 
@@ -227,7 +227,7 @@ FReply FOptimusEditorGraphDragAction_Resource::DroppedOnPanel(
 	    FSlateIcon(),
 	    FUIAction(
 	        FExecuteAction::CreateLambda([ModelGraph, ResourceDesc, InGraphPosition]() {
-		        ModelGraph->AddResourceSetNode(ResourceDesc, InGraphPosition);
+		        ModelGraph->AddResourceSetNode(ResourceDesc, FDeprecateSlateVector2D(InGraphPosition));
 	        }), 
 	        FCanExecuteAction()));
 
diff --git a/Engine/Plugins/Animation/DeformerGraph/Source/OptimusEditor/Private/Widgets/SOptimusEditorGraphExplorerActions.h b/Engine/Plugins/Animation/DeformerGraph/Source/OptimusEditor/Private/Widgets/SOptimusEditorGraphExplorerActions.h
--- a/Engine/Plugins/Animation/DeformerGraph/Source/OptimusEditor/Private/Widgets/SOptimusEditorGraphExplorerActions.h
+++ b/Engine/Plugins/Animation/DeformerGraph/Source/OptimusEditor/Private/Widgets/SOptimusEditorGraphExplorerActions.h
@@ -26,8 +26,8 @@ public:
 	// FGraphEditorDragDropAction overrides
 	FReply DroppedOnPanel(
 		const TSharedRef<SWidget>& InPanel, 
-		FVector2D InScreenPosition, 
-		FVector2D InGraphPosition, 
+		const FVector2f& InScreenPosition,
+		const FVector2f& InGraphPosition,
 		UEdGraph& InGraph
 		) override;
 	
@@ -51,8 +51,8 @@ public:
 	void HoverTargetChanged() override;
 	FReply DroppedOnPanel(
 		const TSharedRef<SWidget>& InPanel, 
-		FVector2D InScreenPosition, 
-		FVector2D InGraphPosition, 
+		const FVector2f& InScreenPosition,
+		const FVector2f& InGraphPosition,
 		UEdGraph& InGraph
 		) override;
 	
@@ -86,8 +86,8 @@ public:
 	void HoverTargetChanged() override;
 	FReply DroppedOnPanel(
 	    const TSharedRef<SWidget>& InPanel,
-	    FVector2D InScreenPosition,
-	    FVector2D InGraphPosition,
+		const FVector2f& InScreenPosition,
+		const FVector2f& InGraphPosition,
 	    UEdGraph& InGraph) override;
 
 protected:
diff --git a/Engine/Plugins/Enterprise/DataprepEditor/Source/DataprepEditor/Private/SchemaActions/DataprepDragDropOp.cpp b/Engine/Plugins/Enterprise/DataprepEditor/Source/DataprepEditor/Private/SchemaActions/DataprepDragDropOp.cpp
--- a/Engine/Plugins/Enterprise/DataprepEditor/Source/DataprepEditor/Private/SchemaActions/DataprepDragDropOp.cpp
+++ b/Engine/Plugins/Enterprise/DataprepEditor/Source/DataprepEditor/Private/SchemaActions/DataprepDragDropOp.cpp
@@ -174,7 +174,7 @@ EVisibility FDataprepDragDropOp::GetErrorIconVisible() const
 	return bDropTargetValid || bCopyRequested ? EVisibility::Collapsed : EVisibility::Visible;
 }
 
-FReply FDataprepDragDropOp::DroppedOnNode(FVector2D ScreenPosition, FVector2D GraphPosition)
+FReply FDataprepDragDropOp::DroppedOnNode(const FVector2f& ScreenPosition, const FVector2f& GraphPosition)
 {
 	if(UDataprepGraphActionStepNode* ActionStepNode = Cast<UDataprepGraphActionStepNode>(GetHoveredNode()))
 	{
diff --git a/Engine/Plugins/Enterprise/DataprepEditor/Source/DataprepEditor/Private/SchemaActions/DataprepDragDropOp.h b/Engine/Plugins/Enterprise/DataprepEditor/Source/DataprepEditor/Private/SchemaActions/DataprepDragDropOp.h
--- a/Engine/Plugins/Enterprise/DataprepEditor/Source/DataprepEditor/Private/SchemaActions/DataprepDragDropOp.h
+++ b/Engine/Plugins/Enterprise/DataprepEditor/Source/DataprepEditor/Private/SchemaActions/DataprepDragDropOp.h
@@ -45,7 +45,7 @@ public:
 
 	// FGraphEditorDragDropAction Interface
 	virtual void HoverTargetChanged() override;
-	virtual FReply DroppedOnNode(FVector2D ScreenPosition, FVector2D GraphPosition) override;
+	virtual FReply DroppedOnNode(const FVector2f& ScreenPosition, const FVector2f& GraphPosition) override;
 	virtual void OnDragged (const class FDragDropEvent& DragDropEvent ) override;
 	virtual EVisibility GetIconVisible() const;
 	virtual EVisibility GetErrorIconVisible() const;
diff --git a/Engine/Plugins/Enterprise/DataprepEditor/Source/DataprepEditor/Private/Widgets/DataprepGraph/SDataprepGraphActionNode.cpp b/Engine/Plugins/Enterprise/DataprepEditor/Source/DataprepEditor/Private/Widgets/DataprepGraph/SDataprepGraphActionNode.cpp
--- a/Engine/Plugins/Enterprise/DataprepEditor/Source/DataprepEditor/Private/Widgets/DataprepGraph/SDataprepGraphActionNode.cpp
+++ b/Engine/Plugins/Enterprise/DataprepEditor/Source/DataprepEditor/Private/Widgets/DataprepGraph/SDataprepGraphActionNode.cpp
@@ -318,7 +318,7 @@ public:
 		{
 			if(StepIndex == ParentPtr.Pin()->GetDataprepAction()->GetStepsCount())
 			{
-				const FVector2D NodeAddPosition = ParentPtr.Pin()->NodeCoordToGraphCoord( MyGeometry.AbsoluteToLocal( DragDropEvent.GetScreenSpacePosition() ) );
+				const FVector2f NodeAddPosition = ParentPtr.Pin()->NodeCoordToGraphCoord( MyGeometry.AbsoluteToLocal( DragDropEvent.GetScreenSpacePosition() ) );
 				return DragActionStepNodeOp->DroppedOnNode(DragDropEvent.GetScreenSpacePosition(), NodeAddPosition);
 			}
 			else
diff --git a/Engine/Plugins/Enterprise/DataprepEditor/Source/DataprepEditor/Private/Widgets/DataprepGraph/SDataprepGraphActionStepNode.cpp b/Engine/Plugins/Enterprise/DataprepEditor/Source/DataprepEditor/Private/Widgets/DataprepGraph/SDataprepGraphActionStepNode.cpp
--- a/Engine/Plugins/Enterprise/DataprepEditor/Source/DataprepEditor/Private/Widgets/DataprepGraph/SDataprepGraphActionStepNode.cpp
+++ b/Engine/Plugins/Enterprise/DataprepEditor/Source/DataprepEditor/Private/Widgets/DataprepGraph/SDataprepGraphActionStepNode.cpp
@@ -511,7 +511,7 @@ FReply SDataprepGraphActionStepNode::OnDrop(const FGeometry& MyGeometry, const F
 	TSharedPtr<FDataprepDragDropOp> DragActionStepNodeOp = DragDropEvent.GetOperationAs<FDataprepDragDropOp>();
 	if (DragActionStepNodeOp.IsValid() && !GetOwnerPanel()->SelectionManager.IsNodeSelected(GraphNode))
 	{
-		const FVector2D NodeAddPosition = NodeCoordToGraphCoord( MyGeometry.AbsoluteToLocal( DragDropEvent.GetScreenSpacePosition() ) );
+		const FVector2f NodeAddPosition = NodeCoordToGraphCoord( MyGeometry.AbsoluteToLocal( DragDropEvent.GetScreenSpacePosition() ) );
 		return DragActionStepNodeOp->DroppedOnNode(DragDropEvent.GetScreenSpacePosition(), NodeAddPosition);
 	}
 
diff --git a/Engine/Plugins/Experimental/BlueprintSnapNodes/Source/BlueprintSnapNodes/Private/SGraphSnapContainerRow.cpp b/Engine/Plugins/Experimental/BlueprintSnapNodes/Source/BlueprintSnapNodes/Private/SGraphSnapContainerRow.cpp
--- a/Engine/Plugins/Experimental/BlueprintSnapNodes/Source/BlueprintSnapNodes/Private/SGraphSnapContainerRow.cpp
+++ b/Engine/Plugins/Experimental/BlueprintSnapNodes/Source/BlueprintSnapNodes/Private/SGraphSnapContainerRow.cpp
@@ -555,10 +555,10 @@ FReply SGraphSnapContainerEntry::OnDrop(const FGeometry& MyGeometry, const FDrag
 		DragConnectionOp->SetHoveredNode(TargetNode);
 		DragConnectionOp->SetHoveredPin(TargetPin);
 
-		//const FVector2D NodeAddPosition = NodeCoordToGraphCoord(MyGeometry.AbsoluteToLocal(DragDropEvent.GetScreenSpacePosition()));
-		const FVector2D NodeAddPosition = FVector2D(TargetNode->NodePosX, TargetNode->NodePosY) +
-			FVector2D(0.0f, 64.0f) * TargetNode->Pins.Find(TargetPin) +
-			FVector2D((TargetPin->Direction == EGPD_Output) ? 200.0f : -180.0f, 0.0f);
+		//const FVector2f NodeAddPosition = NodeCoordToGraphCoord(MyGeometry.AbsoluteToLocal(DragDropEvent.GetScreenSpacePosition()));
+		const FVector2f NodeAddPosition = FVector2f(TargetNode->NodePosX, TargetNode->NodePosY) +
+			FVector2f(0.0f, 64.0f) * TargetNode->Pins.Find(TargetPin) +
+			FVector2f((TargetPin->Direction == EGPD_Output) ? 200.0f : -180.0f, 0.0f);
 
 		FReply Result = DragConnectionOp->DroppedOnPin(DragDropEvent.GetScreenSpacePosition(), NodeAddPosition);
 
diff --git a/Engine/Plugins/FX/Niagara/Source/NiagaraEditor/Private/NiagaraActions.cpp b/Engine/Plugins/FX/Niagara/Source/NiagaraEditor/Private/NiagaraActions.cpp
--- a/Engine/Plugins/FX/Niagara/Source/NiagaraEditor/Private/NiagaraActions.cpp
+++ b/Engine/Plugins/FX/Niagara/Source/NiagaraEditor/Private/NiagaraActions.cpp
@@ -388,7 +388,7 @@ void FNiagaraParameterGraphDragOperation::HoverTargetChanged()
 	FGraphSchemaActionDragDropAction::HoverTargetChanged();
 }
 
-FReply FNiagaraParameterGraphDragOperation::DroppedOnNode(FVector2D ScreenPosition, FVector2D GraphPosition)
+FReply FNiagaraParameterGraphDragOperation::DroppedOnNode(const FVector2f& ScreenPosition, const FVector2f& GraphPosition)
 {
 	FNiagaraParameterAction* ParameterAction = (FNiagaraParameterAction*)SourceAction.Get();
 	if (ParameterAction)
@@ -433,7 +433,7 @@ FReply FNiagaraParameterGraphDragOperation::DroppedOnNode(FVector2D ScreenPositi
 	return FReply::Handled();
 }
 
-FReply FNiagaraParameterGraphDragOperation::DroppedOnPanel(const TSharedRef<SWidget>& Panel, FVector2D ScreenPosition, FVector2D GraphPosition, UEdGraph& Graph)
+FReply FNiagaraParameterGraphDragOperation::DroppedOnPanel(const TSharedRef<SWidget>& Panel, const FVector2f& ScreenPosition, const FVector2f& GraphPosition, UEdGraph& Graph)
 {
 	if (Graph.GetSchema()->IsA<UEdGraphSchema_Niagara>())
 	{
diff --git a/Engine/Plugins/FX/Niagara/Source/NiagaraEditor/Public/NiagaraActions.h b/Engine/Plugins/FX/Niagara/Source/NiagaraEditor/Public/NiagaraActions.h
--- a/Engine/Plugins/FX/Niagara/Source/NiagaraEditor/Public/NiagaraActions.h
+++ b/Engine/Plugins/FX/Niagara/Source/NiagaraEditor/Public/NiagaraActions.h
@@ -296,8 +296,8 @@ public:
 
 	// FGraphEditorDragDropAction interface
 	NIAGARAEDITOR_API virtual void HoverTargetChanged() override;
-	NIAGARAEDITOR_API virtual FReply DroppedOnNode(FVector2D ScreenPosition, FVector2D GraphPosition) override;
-	NIAGARAEDITOR_API virtual FReply DroppedOnPanel(const TSharedRef< class SWidget >& Panel, FVector2D ScreenPosition, FVector2D GraphPosition, UEdGraph& Graph) override;
+	NIAGARAEDITOR_API virtual FReply DroppedOnNode(const FVector2f& ScreenPosition, const FVector2f& GraphPosition) override;
+	NIAGARAEDITOR_API virtual FReply DroppedOnPanel(const TSharedRef< class SWidget >& Panel, const FVector2f& ScreenPosition, const FVector2f& GraphPosition, UEdGraph& Graph) override;
 	// End of FGraphEditorDragDropAction
 
 	/** Set if operation is modified by alt */
@@ -321,7 +321,7 @@ protected:
 		FNiagaraParameterNodeConstructionParams() = delete;
 
 		FNiagaraParameterNodeConstructionParams(
-			const FVector2D& InGraphPosition,
+			const UE::Slate::FDeprecateVector2DParameter& InGraphPosition,
 			UEdGraph* InGraph,
 			const FNiagaraVariable& InParameter,
 			const UNiagaraScriptVariable* InScriptVar)
@@ -331,7 +331,7 @@ protected:
 			, ScriptVar(InScriptVar)
 		{};
 
-		const FVector2D GraphPosition;
+		const FDeprecateSlateVector2D GraphPosition;
 		UEdGraph* Graph;
 		const FNiagaraVariable Parameter;
 		const UNiagaraScriptVariable* ScriptVar;
diff --git a/Engine/Plugins/MovieScene/MovieRenderPipeline/Source/MovieRenderPipelineEditor/Private/Graph/SMovieGraphMembersTabContent.cpp b/Engine/Plugins/MovieScene/MovieRenderPipeline/Source/MovieRenderPipelineEditor/Private/Graph/SMovieGraphMembersTabContent.cpp
--- a/Engine/Plugins/MovieScene/MovieRenderPipeline/Source/MovieRenderPipelineEditor/Private/Graph/SMovieGraphMembersTabContent.cpp
+++ b/Engine/Plugins/MovieScene/MovieRenderPipeline/Source/MovieRenderPipelineEditor/Private/Graph/SMovieGraphMembersTabContent.cpp
@@ -624,7 +624,7 @@ FReply FMovieGraphDragAction_Variable::DroppedOnCategory(FText Category)
 }
 
 FReply FMovieGraphDragAction_Variable::DroppedOnPanel(
-	const TSharedRef<SWidget>& InPanel, FVector2D InScreenPosition, FVector2D InGraphPosition, UEdGraph& InGraph)
+	const TSharedRef<SWidget>& InPanel, const FVector2f& InScreenPosition, const FVector2f& InGraphPosition, UEdGraph& InGraph)
 {
 	const UMovieGraphVariable* VariableMember = WeakVariable.Get();
 	if (!VariableMember)
@@ -639,7 +639,7 @@ FReply FMovieGraphDragAction_Variable::DroppedOnPanel(
 	
 	constexpr UEdGraphPin* FromPin = nullptr;
 	constexpr bool bSelectNewNode = true;
-	NewAction->PerformAction(&InGraph, FromPin, InGraphPosition, bSelectNewNode);
+	NewAction->PerformAction(&InGraph, FromPin, FDeprecateSlateVector2D(InGraphPosition), bSelectNewNode);
 
 	return FReply::Handled();
 }
diff --git a/Engine/Plugins/MovieScene/MovieRenderPipeline/Source/MovieRenderPipelineEditor/Private/Graph/SMovieGraphMembersTabContent.h b/Engine/Plugins/MovieScene/MovieRenderPipeline/Source/MovieRenderPipelineEditor/Private/Graph/SMovieGraphMembersTabContent.h
--- a/Engine/Plugins/MovieScene/MovieRenderPipeline/Source/MovieRenderPipelineEditor/Private/Graph/SMovieGraphMembersTabContent.h
+++ b/Engine/Plugins/MovieScene/MovieRenderPipeline/Source/MovieRenderPipelineEditor/Private/Graph/SMovieGraphMembersTabContent.h
@@ -133,7 +133,7 @@ public:
 	virtual FReply DroppedOnAction(TSharedRef<FEdGraphSchemaAction> Action) override;
 	virtual FReply DroppedOnCategory(FText Category) override;
 	virtual FReply DroppedOnPanel(
-		const TSharedRef<SWidget>& InPanel, FVector2D InScreenPosition, FVector2D InGraphPosition, UEdGraph& InGraph) override;
+		const TSharedRef<SWidget>& InPanel, const FVector2f& InScreenPosition, const FVector2f& InGraphPosition, UEdGraph& InGraph) override;
 
 protected:
 	virtual void GetDefaultStatusSymbol(
diff --git a/Engine/Plugins/PCG/Source/PCGEditor/Private/PCGSettingsDragDropAction.cpp b/Engine/Plugins/PCG/Source/PCGEditor/Private/PCGSettingsDragDropAction.cpp
--- a/Engine/Plugins/PCG/Source/PCGEditor/Private/PCGSettingsDragDropAction.cpp
+++ b/Engine/Plugins/PCG/Source/PCGEditor/Private/PCGSettingsDragDropAction.cpp
@@ -6,7 +6,7 @@
 #include "PCGEditorGraphSchema.h"
 #include "PCGEditorGraphSchemaActions.h"
 
-FReply FPCGSettingsDragDropAction::DroppedOnPanel(const TSharedRef<class SWidget>& Panel, FVector2D ScreenPosition, FVector2D GraphPosition, UEdGraph& Graph)
+FReply FPCGSettingsDragDropAction::DroppedOnPanel(const TSharedRef<class SWidget>& Panel, const FVector2f& ScreenPosition, const FVector2f& GraphPosition, UEdGraph& Graph)
 {
 	if (!Graph.GetSchema()->IsA<UPCGEditorGraphSchema>())
 	{
@@ -19,7 +19,7 @@ FReply FPCGSettingsDragDropAction::DroppedOnPanel(const TSharedRef<class SWidget
 		return FReply::Unhandled();
 	}
 
-	FPCGEditorGraphSchemaAction_NewSettingsElement::MakeSettingsNodesOrContextualMenu(Panel, ScreenPosition, &Graph, { SettingsObjectPath }, { GraphPosition }, /*bSelectNewNodes=*/true);
+	FPCGEditorGraphSchemaAction_NewSettingsElement::MakeSettingsNodesOrContextualMenu(Panel, FDeprecateSlateVector2D(ScreenPosition), &Graph, { SettingsObjectPath }, { FDeprecateSlateVector2D(GraphPosition) }, /*bSelectNewNodes=*/true);
 
 	return FReply::Handled();
 }
\ No newline at end of file
diff --git a/Engine/Plugins/PCG/Source/PCGEditor/Private/PCGSettingsDragDropAction.h b/Engine/Plugins/PCG/Source/PCGEditor/Private/PCGSettingsDragDropAction.h
--- a/Engine/Plugins/PCG/Source/PCGEditor/Private/PCGSettingsDragDropAction.h
+++ b/Engine/Plugins/PCG/Source/PCGEditor/Private/PCGSettingsDragDropAction.h
@@ -15,7 +15,7 @@ public:
 	DRAG_DROP_OPERATOR_TYPE(FPCGSettingsDragDropAction, FGraphSchemaActionDragDropAction)
 
 	// ~Begin FGraphEditorDragDropAction interface
-	virtual FReply DroppedOnPanel(const TSharedRef<class SWidget>& Panel, FVector2D ScreenPosition, FVector2D GraphPosition, UEdGraph& Graph) override;
+	virtual FReply DroppedOnPanel(const TSharedRef<class SWidget>& Panel, const FVector2f& ScreenPosition, const FVector2f& GraphPosition, UEdGraph& Graph) override;
 	// ~End FGraphEditorGraphDropAction interface
 
 	static TSharedRef<FPCGSettingsDragDropAction> New(TSharedPtr<FEdGraphSchemaAction> InAction, const FSoftObjectPath& InSettingsObjectPath)
diff --git a/Engine/Plugins/Runtime/Metasound/Source/MetasoundEditor/Private/MetasoundEditor.cpp b/Engine/Plugins/Runtime/Metasound/Source/MetasoundEditor/Private/MetasoundEditor.cpp
--- a/Engine/Plugins/Runtime/Metasound/Source/MetasoundEditor/Private/MetasoundEditor.cpp
+++ b/Engine/Plugins/Runtime/Metasound/Source/MetasoundEditor/Private/MetasoundEditor.cpp
@@ -329,7 +329,7 @@ namespace Metasound
 
 			DRAG_DROP_OPERATOR_TYPE(FMetaSoundDragDropMemberAction, FGraphSchemaActionDragDropAction)
 
-			virtual FReply DroppedOnPanel(const TSharedRef<SWidget>& InPanel, FVector2D InScreenPosition, FVector2D InGraphPosition, UEdGraph& InGraph) override
+			virtual FReply DroppedOnPanel(const TSharedRef<SWidget>& InPanel, const FVector2f& InScreenPosition, const FVector2f& InGraphPosition, UEdGraph& InGraph) override
 			{
 				if (!GraphMember.IsValid() || &InGraph != GraphMember->GetOwningGraph())
 				{
@@ -338,8 +338,8 @@ namespace Metasound
 
 				return DroppedOnPin(InScreenPosition, InGraphPosition);
 			}
-			virtual FReply DroppedOnNode(FVector2D ScreenPosition, FVector2D GraphPosition) override { return FReply::Unhandled(); }
-			virtual FReply DroppedOnPin(FVector2D InScreenPosition, FVector2D InGraphPosition) override
+			virtual FReply DroppedOnNode(const FVector2f& ScreenPosition, const FVector2f& GraphPosition) override { return FReply::Unhandled(); }
+			virtual FReply DroppedOnPin(const FVector2f& InScreenPosition, const FVector2f& InGraphPosition) override
 			{
 				using namespace Engine;
 				using namespace Frontend;
@@ -370,7 +370,7 @@ namespace Metasound
 					if (UMetasoundEditorGraphNode* NewGraphNode = FGraphBuilder::AddInputNode(ParentMetaSound, TemplateNode->GetID()))
 					{
 						NewGraphNode->Modify();
-						NewGraphNode->UpdateFrontendNodeLocation(InGraphPosition);
+						NewGraphNode->UpdateFrontendNodeLocation(FDeprecateSlateVector2D(InGraphPosition));
 						NewGraphNode->SyncLocationFromFrontendNode();
 
 						TryConnectToHoveredPin(NewGraphNode);
@@ -394,7 +394,7 @@ namespace Metasound
 						if (UMetasoundEditorGraphOutputNode* NewGraphNode = FGraphBuilder::AddOutputNode(ParentMetaSound, Output->NodeID))
 						{
 							NewGraphNode->Modify();
-							NewGraphNode->UpdateFrontendNodeLocation(InGraphPosition);
+							NewGraphNode->UpdateFrontendNodeLocation(FDeprecateSlateVector2D(InGraphPosition));
 							NewGraphNode->SyncLocationFromFrontendNode();
 
 							TryConnectToHoveredPin(NewGraphNode);
@@ -486,7 +486,7 @@ namespace Metasound
 					{
 						if (UMetasoundEditorGraphVariableNode* NewGraphNode = FGraphBuilder::AddVariableNode(ParentMetaSound, NewVariableNode->GetID()))
 						{
-							NewGraphNode->UpdateFrontendNodeLocation(InGraphPosition);
+							NewGraphNode->UpdateFrontendNodeLocation(FDeprecateSlateVector2D(InGraphPosition));
 							NewGraphNode->SyncLocationFromFrontendNode();
 
 							TryConnectToHoveredPin(NewGraphNode);
diff --git a/Engine/Plugins/Runtime/RigVM/Source/RigVMDeveloper/Private/EdGraph/RigVMEdGraphSchema.cpp b/Engine/Plugins/Runtime/RigVM/Source/RigVMDeveloper/Private/EdGraph/RigVMEdGraphSchema.cpp
--- a/Engine/Plugins/Runtime/RigVM/Source/RigVMDeveloper/Private/EdGraph/RigVMEdGraphSchema.cpp
+++ b/Engine/Plugins/Runtime/RigVM/Source/RigVMDeveloper/Private/EdGraph/RigVMEdGraphSchema.cpp
@@ -510,7 +510,7 @@ FSlateBrush const* FRigVMEdGraphSchemaAction_Event::GetPaletteIcon() const
 	return EventIcon.GetIcon();
 }
 
-FReply FRigVMFunctionDragDropAction::DroppedOnPanel(const TSharedRef< class SWidget >& Panel, FVector2D ScreenPosition, FVector2D GraphPosition, UEdGraph& Graph)
+FReply FRigVMFunctionDragDropAction::DroppedOnPanel(const TSharedRef< class SWidget >& Panel, const FVector2f& ScreenPosition, const FVector2f& GraphPosition, UEdGraph& Graph)
 {
 	// For local variables
 	if (SourceAction->GetTypeId() == FRigVMEdGraphSchemaAction_LocalVar::StaticGetTypeId())
@@ -537,7 +537,7 @@ FReply FRigVMFunctionDragDropAction::DroppedOnPanel(const TSharedRef< class SWid
 							FUIAction(
 							FExecuteAction::CreateLambda([Controller, LocalVariable, GraphPosition]()
 							{
-								Controller->AddVariableNode(LocalVariable.Name, LocalVariable.CPPType, LocalVariable.CPPTypeObject, true, LocalVariable.DefaultValue, GraphPosition, FString(), true, true);
+								Controller->AddVariableNode(LocalVariable.Name, LocalVariable.CPPType, LocalVariable.CPPTypeObject, true, LocalVariable.DefaultValue, FDeprecateSlateVector2D(GraphPosition), FString(), true, true);
 							}),
 							FCanExecuteAction()));
 
@@ -548,7 +548,7 @@ FReply FRigVMFunctionDragDropAction::DroppedOnPanel(const TSharedRef< class SWid
 							FUIAction(
 							FExecuteAction::CreateLambda([Controller, LocalVariable, GraphPosition]()
 							{
-								Controller->AddVariableNode(LocalVariable.Name, LocalVariable.CPPType, LocalVariable.CPPTypeObject, false, LocalVariable.DefaultValue, GraphPosition, FString(), true, true);
+								Controller->AddVariableNode(LocalVariable.Name, LocalVariable.CPPType, LocalVariable.CPPTypeObject, false, LocalVariable.DefaultValue, FDeprecateSlateVector2D(GraphPosition), FString(), true, true);
 							}),
 							FCanExecuteAction()));
 
@@ -594,7 +594,7 @@ FReply FRigVMFunctionDragDropAction::DroppedOnPanel(const TSharedRef< class SWid
 									}
 								}
 #endif
-								TargetController->AddFunctionReferenceNode(FunctionDefinitionNode, GraphPosition, FString(), true, true);
+								TargetController->AddFunctionReferenceNode(FunctionDefinitionNode, FDeprecateSlateVector2D(GraphPosition), FString(), true, true);
 							}
 						}
 					}
@@ -608,7 +608,7 @@ FReply FRigVMFunctionDragDropAction::DroppedOnPanel(const TSharedRef< class SWid
 	return FReply::Unhandled();
 }
 
-FReply FRigVMFunctionDragDropAction::DroppedOnPin(FVector2D ScreenPosition, FVector2D GraphPosition)
+FReply FRigVMFunctionDragDropAction::DroppedOnPin(const FVector2f& ScreenPosition, const FVector2f& GraphPosition)
 {
 	return FReply::Unhandled();
 }
diff --git a/Engine/Plugins/Runtime/RigVM/Source/RigVMDeveloper/Public/EdGraph/RigVMEdGraphSchema.h b/Engine/Plugins/Runtime/RigVM/Source/RigVMDeveloper/Public/EdGraph/RigVMEdGraphSchema.h
--- a/Engine/Plugins/Runtime/RigVM/Source/RigVMDeveloper/Public/EdGraph/RigVMEdGraphSchema.h
+++ b/Engine/Plugins/Runtime/RigVM/Source/RigVMDeveloper/Public/EdGraph/RigVMEdGraphSchema.h
@@ -262,8 +262,8 @@ public:
 	DRAG_DROP_OPERATOR_TYPE(FRigVMFunctionDragDropAction, FGraphSchemaActionDragDropAction)
 
 	// FGraphEditorDragDropAction interface
-	virtual FReply DroppedOnPanel(const TSharedRef< class SWidget >& Panel, FVector2D ScreenPosition, FVector2D GraphPosition, UEdGraph& Graph) override;
-	virtual FReply DroppedOnPin(FVector2D ScreenPosition, FVector2D GraphPosition) override;
+	virtual FReply DroppedOnPanel(const TSharedRef< class SWidget >& Panel, const FVector2f& ScreenPosition, const FVector2f& GraphPosition, UEdGraph& Graph) override;
+	virtual FReply DroppedOnPin(const FVector2f& ScreenPosition, const FVector2f& GraphPosition) override;
 	virtual FReply DroppedOnAction(TSharedRef<FEdGraphSchemaAction> Action) override;
 	virtual FReply DroppedOnCategory(FText Category) override;
 	virtual void HoverTargetChanged() override;
diff --git a/Engine/Plugins/Runtime/RigVM/Source/RigVMEditor/Private/Widgets/SRigVMEditorGraphExplorer.cpp b/Engine/Plugins/Runtime/RigVM/Source/RigVMEditor/Private/Widgets/SRigVMEditorGraphExplorer.cpp
--- a/Engine/Plugins/Runtime/RigVM/Source/RigVMEditor/Private/Widgets/SRigVMEditorGraphExplorer.cpp
+++ b/Engine/Plugins/Runtime/RigVM/Source/RigVMEditor/Private/Widgets/SRigVMEditorGraphExplorer.cpp
@@ -90,13 +90,13 @@ void FRigVMGraphExplorerDragDropOp::Construct()
 	);
 }
 
-FReply FRigVMGraphExplorerDragDropOp::DroppedOnPanel(const TSharedRef<SWidget>& Panel, FVector2D ScreenPosition, FVector2D GraphPosition, UEdGraph& Graph)
+FReply FRigVMGraphExplorerDragDropOp::DroppedOnPanel(const TSharedRef<SWidget>& Panel, const FVector2f& ScreenPosition, const FVector2f& GraphPosition, UEdGraph& Graph)
 {
 	if (URigVMEdGraph* TargetRigGraph = Cast<URigVMEdGraph>(&Graph))
 	{
 		if(URigVMBlueprint* Blueprint = TargetRigGraph->GetBlueprint())
 		{
-			//  Find the appropiate asset editor where the drop is happening
+			//  Find the appropriate asset editor where the drop is happening
 			const TArray<IAssetEditorInstance*> AssetEditors = GEditor->GetEditorSubsystem<UAssetEditorSubsystem>()->FindEditorsForAsset(Blueprint);
 			IAssetEditorInstance* FocusedAssetEditor = nullptr;
 			for (IAssetEditorInstance* AssetEditor : AssetEditors)
@@ -130,7 +130,7 @@ FReply FRigVMGraphExplorerDragDropOp::DroppedOnPanel(const TSharedRef<SWidget>&
 			{
 				if (FRigVMEditorBase* RigVMEditor = FRigVMEditorBase::GetFromAssetEditorInstance(FocusedAssetEditor))
 				{
-					RigVMEditor->OnGraphNodeDropToPerform(SharedThis(this), TargetRigGraph, GraphPosition, ScreenPosition);
+					RigVMEditor->OnGraphNodeDropToPerform(SharedThis(this), TargetRigGraph, FDeprecateSlateVector2D(GraphPosition), FDeprecateSlateVector2D(ScreenPosition));
 				}
 			}
 		}
@@ -138,7 +138,7 @@ FReply FRigVMGraphExplorerDragDropOp::DroppedOnPanel(const TSharedRef<SWidget>&
 	return FGraphEditorDragDropAction::DroppedOnPanel(Panel, ScreenPosition, GraphPosition, Graph);
 }
 
-FReply FRigVMGraphExplorerDragDropOp::DroppedOnPin(FVector2D ScreenPosition, FVector2D GraphPosition)
+FReply FRigVMGraphExplorerDragDropOp::DroppedOnPin(const FVector2f& ScreenPosition, const FVector2f& GraphPosition)
 {
 	UEdGraph* TargetGraph = GetHoveredGraph();
 	UEdGraphPin* TargetPin = GetHoveredPin();
@@ -150,7 +150,7 @@ FReply FRigVMGraphExplorerDragDropOp::DroppedOnPin(FVector2D ScreenPosition, FVe
 	{
 		if (FProperty* Property = Blueprint->SkeletonGeneratedClass->FindPropertyByName(*GetElement().Name))
 		{
-			if (EdSchema->RequestVariableDropOnPin(TargetGraph, Property, TargetPin, GraphPosition, ScreenPosition))
+			if (EdSchema->RequestVariableDropOnPin(TargetGraph, Property, TargetPin, FDeprecateSlateVector2D(GraphPosition), FDeprecateSlateVector2D(ScreenPosition)))
 			{
 				return FReply::Handled();
 			}
@@ -162,7 +162,7 @@ FReply FRigVMGraphExplorerDragDropOp::DroppedOnPin(FVector2D ScreenPosition, FVe
 		{
 			if (Variable.Name == Element.Name)
 			{
-				if (EdSchema->RequestVariableDropOnPin(TargetGraph, Variable.ToExternalVariable(), TargetPin, GraphPosition, ScreenPosition))
+				if (EdSchema->RequestVariableDropOnPin(TargetGraph, Variable.ToExternalVariable(), TargetPin, FDeprecateSlateVector2D(GraphPosition), FDeprecateSlateVector2D(ScreenPosition)))
 				{
 					return FReply::Handled();
 				}
diff --git a/Engine/Plugins/Runtime/RigVM/Source/RigVMEditor/Public/Widgets/SRigVMEditorGraphExplorer.h b/Engine/Plugins/Runtime/RigVM/Source/RigVMEditor/Public/Widgets/SRigVMEditorGraphExplorer.h
--- a/Engine/Plugins/Runtime/RigVM/Source/RigVMEditor/Public/Widgets/SRigVMEditorGraphExplorer.h
+++ b/Engine/Plugins/Runtime/RigVM/Source/RigVMEditor/Public/Widgets/SRigVMEditorGraphExplorer.h
@@ -44,8 +44,8 @@ public:
 
 	virtual void Construct() override;
 
-	virtual FReply DroppedOnPanel(const TSharedRef<SWidget>& Panel, FVector2D ScreenPosition, FVector2D GraphPosition, UEdGraph& Graph) override;
-	virtual FReply DroppedOnPin(FVector2D ScreenPosition, FVector2D GraphPosition) override;
+	virtual FReply DroppedOnPanel(const TSharedRef<SWidget>& Panel, const FVector2f& ScreenPosition, const FVector2f& GraphPosition, UEdGraph& Graph) override;
+	virtual FReply DroppedOnPin(const FVector2f& ScreenPosition, const FVector2f& GraphPosition) override;
 
 private:
 
diff --git a/Engine/Source/Editor/AnimGraph/Private/AnimationGraphSchema.cpp b/Engine/Source/Editor/AnimGraph/Private/AnimationGraphSchema.cpp
--- a/Engine/Source/Editor/AnimGraph/Private/AnimationGraphSchema.cpp
+++ b/Engine/Source/Editor/AnimGraph/Private/AnimationGraphSchema.cpp
@@ -76,9 +76,9 @@ class ANIMGRAPH_API FAnimationLayerDragDropAction : public FGraphSchemaActionDra
 public:
 	DRAG_DROP_OPERATOR_TYPE(FAnimationLayerDragDropAction, FGraphSchemaActionDragDropAction)
 
-	virtual FReply DroppedOnPanel(const TSharedRef< class SWidget >& Panel, FVector2D ScreenPosition, FVector2D GraphPosition, UEdGraph& Graph) override;
-	virtual FReply DroppedOnNode(FVector2D ScreenPosition, FVector2D GraphPosition) override;
-	virtual FReply DroppedOnPin(FVector2D ScreenPosition, FVector2D GraphPosition) override;
+	virtual FReply DroppedOnPanel(const TSharedRef< class SWidget >& Panel, const FVector2f& ScreenPosition, const FVector2f& GraphPosition, UEdGraph& Graph) override;
+	virtual FReply DroppedOnNode(const FVector2f& ScreenPosition, const FVector2f& GraphPosition) override;
+	virtual FReply DroppedOnPin(const FVector2f& ScreenPosition, const FVector2f& GraphPosition) override;
 	virtual FReply DroppedOnAction(TSharedRef<FEdGraphSchemaAction> Action) override;
 	virtual FReply DroppedOnCategory(FText Category) override;
 	virtual void HoverTargetChanged() override;
@@ -106,7 +106,7 @@ FAnimationLayerDragDropAction::FAnimationLayerDragDropAction()
 {
 }
 
-FReply FAnimationLayerDragDropAction::DroppedOnPanel(const TSharedRef< class SWidget >& Panel, FVector2D ScreenPosition, FVector2D GraphPosition, UEdGraph& Graph)
+FReply FAnimationLayerDragDropAction::DroppedOnPanel(const TSharedRef< class SWidget >& Panel, const FVector2f& ScreenPosition, const FVector2f& GraphPosition, UEdGraph& Graph)
 {
 	if (UAnimationGraph* TargetRigGraph = Cast<UAnimationGraph>(&Graph))
 	{
@@ -135,7 +135,7 @@ FReply FAnimationLayerDragDropAction::DroppedOnPanel(const TSharedRef< class SWi
 	return FReply::Unhandled();
 }
 
-FReply FAnimationLayerDragDropAction::DroppedOnNode(FVector2D ScreenPosition, FVector2D GraphPosition)
+FReply FAnimationLayerDragDropAction::DroppedOnNode(const FVector2f& ScreenPosition, const FVector2f& GraphPosition)
 {
 	if (UEdGraphNode* TargetNode = GetHoveredNode())
 	{
@@ -148,7 +148,7 @@ FReply FAnimationLayerDragDropAction::DroppedOnNode(FVector2D ScreenPosition, FV
 	return FReply::Unhandled();
 }
 
-FReply FAnimationLayerDragDropAction::DroppedOnPin(FVector2D ScreenPosition, FVector2D GraphPosition)
+FReply FAnimationLayerDragDropAction::DroppedOnPin(const FVector2f& ScreenPosition, const FVector2f& GraphPosition)
 {
 	return FReply::Unhandled();
 }
diff --git a/Engine/Source/Editor/GraphEditor/Private/DragConnection.cpp b/Engine/Source/Editor/GraphEditor/Private/DragConnection.cpp
--- a/Engine/Source/Editor/GraphEditor/Private/DragConnection.cpp
+++ b/Engine/Source/Editor/GraphEditor/Private/DragConnection.cpp
@@ -264,7 +264,7 @@ FDragConnection::FDragConnection(const TSharedRef<SGraphPanel>& GraphPanelIn, co
 			const UEdGraphPin* PinObj = FDraggedPinTable::TConstIterator(DraggedPinsIn)->GetPinObj(*GraphPanelIn);
 			if (PinObj && PinObj->Direction == EGPD_Input)
 			{
-				DecoratorAdjust *= FVector2D(-1.0f, 1.0f);
+				DecoratorAdjust *= FVector2f(-1.0f, 1.0f);
 			}
 		}
 
@@ -282,7 +282,7 @@ FDragConnection::FDragConnection(const TSharedRef<SGraphPanel>& GraphPanelIn, co
 			const UEdGraphPin* PinObj = SourcePinHandle.GetPinObj(*GraphPanel.Get());
 			if (PinObj && PinObj->Direction == EGPD_Input)
 			{
-				DecoratorAdjust *= FVector2D(-1.0f, 1.0f);
+				DecoratorAdjust *= FVector2f(-1.0f, 1.0f);
 			}
 		}
 
@@ -292,7 +292,7 @@ FDragConnection::FDragConnection(const TSharedRef<SGraphPanel>& GraphPanelIn, co
 	}
 }
 
-FReply FDragConnection::DroppedOnPin(FVector2D ScreenPosition, FVector2D GraphPosition)
+FReply FDragConnection::DroppedOnPin(const FVector2f& ScreenPosition, const FVector2f& GraphPosition)
 {
 	bool bError = false;
 	TSet<FEdGraphNodeHandle> NodeList;
@@ -430,7 +430,7 @@ FReply FDragConnection::DroppedOnPin(FVector2D ScreenPosition, FVector2D GraphPo
 	return FReply::Handled();
 }
 
-FReply FDragConnection::DroppedOnNode(FVector2D ScreenPosition, FVector2D GraphPosition)
+FReply FDragConnection::DroppedOnNode(const FVector2f& ScreenPosition, const FVector2f& GraphPosition)
 {
 	if (DragMode != EDragMode::CreateConnection)
 	{
@@ -489,7 +489,7 @@ FReply FDragConnection::DroppedOnNode(FVector2D ScreenPosition, FVector2D GraphP
 	return bHandledPinDropOnNode? FReply::Handled() : FReply::Unhandled();
 }
 
-FReply FDragConnection::DroppedOnPanel( const TSharedRef< SWidget >& Panel, FVector2D ScreenPosition, FVector2D GraphPosition, UEdGraph& Graph)
+FReply FDragConnection::DroppedOnPanel( const TSharedRef< SWidget >& Panel, const FVector2f& ScreenPosition, const FVector2f& GraphPosition, UEdGraph& Graph)
 {
 	if (DragMode != EDragMode::CreateConnection)
 	{
diff --git a/Engine/Source/Editor/GraphEditor/Private/DragConnection.h b/Engine/Source/Editor/GraphEditor/Private/DragConnection.h
--- a/Engine/Source/Editor/GraphEditor/Private/DragConnection.h
+++ b/Engine/Source/Editor/GraphEditor/Private/DragConnection.h
@@ -37,9 +37,9 @@ public:
 
 	// FGraphEditorDragDropAction interface
 	virtual void HoverTargetChanged() override;
-	virtual FReply DroppedOnPin(FVector2D ScreenPosition, FVector2D GraphPosition) override;
-	virtual FReply DroppedOnNode(FVector2D ScreenPosition, FVector2D GraphPosition) override;
-	virtual FReply DroppedOnPanel(const TSharedRef< SWidget >& Panel, FVector2D ScreenPosition, FVector2D GraphPosition, UEdGraph& Graph) override;
+	virtual FReply DroppedOnPin(const FVector2f& ScreenPosition, const FVector2f& GraphPosition) override;
+	virtual FReply DroppedOnNode(const FVector2f& ScreenPosition, const FVector2f& GraphPosition) override;
+	virtual FReply DroppedOnPanel(const TSharedRef< SWidget >& Panel, const FVector2f& ScreenPosition, const FVector2f& GraphPosition, UEdGraph& Graph) override;
 	virtual void OnDragged(const class FDragDropEvent& DragDropEvent) override;
 	// End of FGraphEditorDragDropAction interface
 
@@ -60,7 +60,7 @@ protected:
 	EDragMode DragMode = EDragMode::CreateConnection;
 
 	/** Offset information for the decorator widget */
-	FVector2D DecoratorAdjust;
+	FDeprecateSlateVector2D DecoratorAdjust;
 
 	FGraphPinHandle SourcePinHandle;
 	FGraphPinHandle TargetPinHandle;
diff --git a/Engine/Source/Editor/GraphEditor/Private/DragNode.cpp b/Engine/Source/Editor/GraphEditor/Private/DragNode.cpp
--- a/Engine/Source/Editor/GraphEditor/Private/DragNode.cpp
+++ b/Engine/Source/Editor/GraphEditor/Private/DragNode.cpp
@@ -163,12 +163,12 @@ void FDragNode::HoverTargetChanged()
 	}
 }
 
-FReply FDragNode::DroppedOnNode(FVector2D ScreenPosition, FVector2D GraphPosition)
+FReply FDragNode::DroppedOnNode(const FVector2f& ScreenPosition, const FVector2f& GraphPosition)
 {
 	return FReply::Handled();
 }
 
-FReply FDragNode::DroppedOnPanel( const TSharedRef< SWidget >& Panel, FVector2D ScreenPosition, FVector2D GraphPosition, UEdGraph& Graph)
+FReply FDragNode::DroppedOnPanel( const TSharedRef< SWidget >& Panel, const FVector2f& ScreenPosition, const FVector2f& GraphPosition, UEdGraph& Graph)
 {
 	return FReply::Handled();
 }
diff --git a/Engine/Source/Editor/GraphEditor/Private/DragNode.h b/Engine/Source/Editor/GraphEditor/Private/DragNode.h
--- a/Engine/Source/Editor/GraphEditor/Private/DragNode.h
+++ b/Engine/Source/Editor/GraphEditor/Private/DragNode.h
@@ -27,8 +27,8 @@ public:
 	
 	// FGraphEditorDragDropAction interface
 	virtual void HoverTargetChanged() override;
-	virtual FReply DroppedOnNode(FVector2D ScreenPosition, FVector2D GraphPosition) override;
-	virtual FReply DroppedOnPanel( const TSharedRef< SWidget >& Panel, FVector2D ScreenPosition, FVector2D GraphPosition, UEdGraph& Graph) override;
+	virtual FReply DroppedOnNode(const FVector2f& ScreenPosition, const FVector2f& GraphPosition) override;
+	virtual FReply DroppedOnPanel( const TSharedRef< SWidget >& Panel, const FVector2f& ScreenPosition, const FVector2f& GraphPosition, UEdGraph& Graph) override;
 	//virtual void OnDragBegin( const TSharedRef<class SGraphPin>& InPin) override;
 	virtual void OnDragged (const class FDragDropEvent& DragDropEvent ) override;
 	// End of FGraphEditorDragDropAction interface
diff --git a/Engine/Source/Editor/GraphEditor/Private/GraphEditorDragDropAction.cpp b/Engine/Source/Editor/GraphEditor/Private/GraphEditorDragDropAction.cpp
--- a/Engine/Source/Editor/GraphEditor/Private/GraphEditorDragDropAction.cpp
+++ b/Engine/Source/Editor/GraphEditor/Private/GraphEditorDragDropAction.cpp
@@ -115,6 +115,41 @@ void FGraphEditorDragDropAction::SetHoveredAction(TSharedPtr<struct FEdGraphSche
 	}
 }
 
+FReply FGraphEditorDragDropAction::DroppedOnPin(FVector2D ScreenPosition, FVector2D GraphPosition)
+{
+	return FReply::Unhandled();
+}
+
+FReply FGraphEditorDragDropAction::DroppedOnPin(const FVector2f& ScreenPosition, const FVector2f& GraphPosition)
+{
+PRAGMA_DISABLE_DEPRECATION_WARNINGS
+	return DroppedOnPin(FVector2D(ScreenPosition), FVector2D(GraphPosition));
+PRAGMA_ENABLE_DEPRECATION_WARNINGS
+}
+
+FReply FGraphEditorDragDropAction::DroppedOnNode(FVector2D ScreenPosition, FVector2D GraphPosition)
+{
+	return FReply::Unhandled();
+}
+
+FReply FGraphEditorDragDropAction::DroppedOnNode(const FVector2f& ScreenPosition, const FVector2f& GraphPosition)
+{
+PRAGMA_DISABLE_DEPRECATION_WARNINGS
+	return DroppedOnNode(FVector2D(ScreenPosition), FVector2D(GraphPosition));
+PRAGMA_ENABLE_DEPRECATION_WARNINGS
+}
+
+FReply FGraphEditorDragDropAction::DroppedOnPanel(const TSharedRef<SWidget>& Panel, FVector2D ScreenPosition, FVector2D GraphPosition, UEdGraph& Graph)
+{
+	return FReply::Unhandled();
+}
+
+FReply FGraphEditorDragDropAction::DroppedOnPanel(const TSharedRef<SWidget>& Panel, const FVector2f& ScreenPosition, const FVector2f& GraphPosition, UEdGraph& Graph)
+{
+PRAGMA_DISABLE_DEPRECATION_WARNINGS
+	return DroppedOnPanel(Panel, FVector2D(ScreenPosition), FVector2D(GraphPosition), Graph);
+PRAGMA_ENABLE_DEPRECATION_WARNINGS
+}
 
 void FGraphEditorDragDropAction::Construct()
 {
@@ -237,19 +272,19 @@ void FGraphSchemaActionDragDropAction::GetDefaultStatusSymbol(const FSlateBrush*
 	SecondaryColorOut = FLinearColor::White;
 }
 
-FReply FGraphSchemaActionDragDropAction::DroppedOnPanel( const TSharedRef< SWidget >& Panel, FVector2D ScreenPosition, FVector2D GraphPosition, UEdGraph& Graph)
+FReply FGraphSchemaActionDragDropAction::DroppedOnPanel( const TSharedRef< SWidget >& Panel, const FVector2f& ScreenPosition, const FVector2f& GraphPosition, UEdGraph& Graph)
 {
 	if (SourceAction.IsValid())
 	{
 		TArray<UEdGraphPin*> DummyPins;
-		SourceAction->PerformAction(&Graph, DummyPins, GraphPosition);
+		SourceAction->PerformAction(&Graph, DummyPins, FDeprecateSlateVector2D(GraphPosition));
 
 		return FReply::Handled();
 	}
 	return FReply::Unhandled();
 }
 
-FReply FGraphSchemaActionDragDropAction::DroppedOnPin(FVector2D ScreenPosition, FVector2D GraphPosition)
+FReply FGraphSchemaActionDragDropAction::DroppedOnPin(const FVector2f& ScreenPosition, const FVector2f& GraphPosition)
 {
 	if (UEdGraph* Graph = GetHoveredGraph())
 	{
@@ -261,7 +296,7 @@ FReply FGraphSchemaActionDragDropAction::DroppedOnPin(FVector2D ScreenPosition,
 				DummyPins.Add(Pin);
 			}
 
-			SourceAction->PerformAction(Graph, DummyPins, GraphPosition);
+			SourceAction->PerformAction(Graph, DummyPins, FDeprecateSlateVector2D(GraphPosition));
 
 			return FReply::Handled();
 		}
diff --git a/Engine/Source/Editor/GraphEditor/Private/SGraphNode.cpp b/Engine/Source/Editor/GraphEditor/Private/SGraphNode.cpp
--- a/Engine/Source/Editor/GraphEditor/Private/SGraphNode.cpp
+++ b/Engine/Source/Editor/GraphEditor/Private/SGraphNode.cpp
@@ -385,7 +385,7 @@ FReply SGraphNode::OnDrop( const FGeometry& MyGeometry, const FDragDropEvent& Dr
 	{
 		TSharedPtr<FGraphEditorDragDropAction> DragConnectionOp = StaticCastSharedPtr<FGraphEditorDragDropAction>(Operation);
 
-		const FVector2D NodeAddPosition = NodeCoordToGraphCoord( MyGeometry.AbsoluteToLocal( DragDropEvent.GetScreenSpacePosition() ) );
+		const FVector2f NodeAddPosition = NodeCoordToGraphCoord( MyGeometry.AbsoluteToLocal( DragDropEvent.GetScreenSpacePosition() ) );
 
 		FReply Result = DragConnectionOp->DroppedOnNode(DragDropEvent.GetScreenSpacePosition(), NodeAddPosition);
 		
diff --git a/Engine/Source/Editor/GraphEditor/Private/SGraphPanel.cpp b/Engine/Source/Editor/GraphEditor/Private/SGraphPanel.cpp
--- a/Engine/Source/Editor/GraphEditor/Private/SGraphPanel.cpp
+++ b/Engine/Source/Editor/GraphEditor/Private/SGraphPanel.cpp
@@ -1347,7 +1347,7 @@ FReply SGraphPanel::OnDrop( const FGeometry& MyGeometry, const FDragDropEvent& D
 		TSharedPtr<FGraphEditorDragDropAction> DragConn = StaticCastSharedPtr<FGraphEditorDragDropAction>(Operation);
 		if (DragConn.IsValid() && DragConn->IsSupportedBySchema(GraphObj->GetSchema()))
 		{
-			return DragConn->DroppedOnPanel(SharedThis(this), FVector2D(DragDropEvent.GetScreenSpacePosition()), FVector2D(NodeAddPosition), *GraphObj);
+			return DragConn->DroppedOnPanel(SharedThis(this), DragDropEvent.GetScreenSpacePosition(), NodeAddPosition, *GraphObj);
 		}
 
 		return FReply::Unhandled();
diff --git a/Engine/Source/Editor/GraphEditor/Private/SGraphPin.cpp b/Engine/Source/Editor/GraphEditor/Private/SGraphPin.cpp
--- a/Engine/Source/Editor/GraphEditor/Private/SGraphPin.cpp
+++ b/Engine/Source/Editor/GraphEditor/Private/SGraphPin.cpp
@@ -876,7 +876,7 @@ FReply SGraphPin::OnDrop( const FGeometry& MyGeometry, const FDragDropEvent& Dra
 			
 		}
 
-		return DragConnectionOp->DroppedOnPin(DragDropEvent.GetScreenSpacePosition(), FVector2D(NodeAddPosition));
+		return DragConnectionOp->DroppedOnPin(DragDropEvent.GetScreenSpacePosition(), NodeAddPosition);
 	}
 	// handle dropping an asset on the pin
 	else if (Operation->IsOfType<FAssetDragDropOp>() && NodeWidget.IsValid())
diff --git a/Engine/Source/Editor/GraphEditor/Public/GraphEditorDragDropAction.h b/Engine/Source/Editor/GraphEditor/Public/GraphEditorDragDropAction.h
--- a/Engine/Source/Editor/GraphEditor/Public/GraphEditorDragDropAction.h
+++ b/Engine/Source/Editor/GraphEditor/Public/GraphEditorDragDropAction.h
@@ -43,9 +43,15 @@ public:
 
 	// Interface to override
 	virtual void HoverTargetChanged() {}
-	virtual FReply DroppedOnPin(FVector2D ScreenPosition, FVector2D GraphPosition) { return FReply::Unhandled(); }
-	virtual FReply DroppedOnNode(FVector2D ScreenPosition, FVector2D GraphPosition) { return FReply::Unhandled(); }
-	virtual FReply DroppedOnPanel( const TSharedRef< SWidget >& Panel, FVector2D ScreenPosition, FVector2D GraphPosition, UEdGraph& Graph) { return FReply::Unhandled(); }
+	UE_DEPRECATED(5.6, "Use the version of the function accepting FVector2f; this Slate API no longer interfaces directly with double-precision scalars and vectors.")
+	virtual FReply DroppedOnPin(FVector2D ScreenPosition, FVector2D GraphPosition) UE_SLATE_DEPRECATED_VECTOR_VIRTUAL_FUNCTION;
+	virtual FReply DroppedOnPin(const FVector2f& ScreenPosition, const FVector2f& GraphPosition);
+	UE_DEPRECATED(5.6, "Use the version of the function accepting FVector2f; this Slate API no longer interfaces directly with double-precision scalars and vectors.")
+	virtual FReply DroppedOnNode(FVector2D ScreenPosition, FVector2D GraphPosition) UE_SLATE_DEPRECATED_VECTOR_VIRTUAL_FUNCTION;
+	virtual FReply DroppedOnNode(const FVector2f& ScreenPosition, const FVector2f& GraphPosition);
+	UE_DEPRECATED(5.6, "Use the version of the function accepting FVector2f; this Slate API no longer interfaces directly with double-precision scalars and vectors.")
+	virtual FReply DroppedOnPanel(const TSharedRef< SWidget >& Panel, FVector2D ScreenPosition, FVector2D GraphPosition, UEdGraph& Graph) UE_SLATE_DEPRECATED_VECTOR_VIRTUAL_FUNCTION;
+	virtual FReply DroppedOnPanel(const TSharedRef< SWidget >& Panel, const FVector2f& ScreenPosition, const FVector2f& GraphPosition, UEdGraph& Graph);
 	virtual FReply DroppedOnAction(TSharedRef<struct FEdGraphSchemaAction> Action) { return FReply::Unhandled(); }
 	virtual FReply DroppedOnCategory(FText Category) { return FReply::Unhandled(); }
 	// End of interface to override
@@ -98,8 +104,8 @@ public:
 
 	// FGraphEditorDragDropAction interface
 	virtual void HoverTargetChanged() override;
-	virtual FReply DroppedOnPanel(const TSharedRef< SWidget >& Panel, FVector2D ScreenPosition, FVector2D GraphPosition, UEdGraph& Graph) override;
-	virtual FReply DroppedOnPin(FVector2D ScreenPosition, FVector2D GraphPosition) override;
+	virtual FReply DroppedOnPanel(const TSharedRef< SWidget >& Panel, const FVector2f& ScreenPosition, const FVector2f& GraphPosition, UEdGraph& Graph) override;
+	virtual FReply DroppedOnPin(const FVector2f& ScreenPosition, const FVector2f& GraphPosition) override;
 	// End of FGraphEditorDragDropAction
 
 	static TSharedRef<FGraphSchemaActionDragDropAction> New(TSharedPtr<FEdGraphSchemaAction> InActionNode )
diff --git a/Engine/Source/Editor/Kismet/Private/BPDelegateDragDropAction.cpp b/Engine/Source/Editor/Kismet/Private/BPDelegateDragDropAction.cpp
--- a/Engine/Source/Editor/Kismet/Private/BPDelegateDragDropAction.cpp
+++ b/Engine/Source/Editor/Kismet/Private/BPDelegateDragDropAction.cpp
@@ -73,7 +73,7 @@ void FKismetDelegateDragDropAction::AssignEvent(FNodeConstructionParams Params)
 	}
 }
 
-FReply FKismetDelegateDragDropAction::DroppedOnPanel(const TSharedRef< SWidget >& Panel, FVector2D ScreenPosition, FVector2D GraphPosition, UEdGraph& Graph)
+FReply FKismetDelegateDragDropAction::DroppedOnPanel(const TSharedRef< SWidget >& Panel, const FVector2f& ScreenPosition, const FVector2f& GraphPosition, UEdGraph& Graph)
 {
 	if(IsValid())
 	{
diff --git a/Engine/Source/Editor/Kismet/Private/BPDelegateDragDropAction.h b/Engine/Source/Editor/Kismet/Private/BPDelegateDragDropAction.h
--- a/Engine/Source/Editor/Kismet/Private/BPDelegateDragDropAction.h
+++ b/Engine/Source/Editor/Kismet/Private/BPDelegateDragDropAction.h
@@ -32,10 +32,8 @@ public:
 
 	// FGraphEditorDragDropAction interface
 	virtual void HoverTargetChanged() override { FMyBlueprintItemDragDropAction::HoverTargetChanged(); }
-	virtual FReply DroppedOnPin(FVector2D ScreenPosition, FVector2D GraphPosition) override { return FGraphEditorDragDropAction::DroppedOnPin(ScreenPosition, GraphPosition); }
-	virtual FReply DroppedOnNode(FVector2D ScreenPosition, FVector2D GraphPosition) override { return FGraphEditorDragDropAction::DroppedOnNode(ScreenPosition, GraphPosition); }
 
-	virtual FReply DroppedOnPanel(const TSharedRef< SWidget >& Panel, FVector2D ScreenPosition, FVector2D GraphPosition, UEdGraph& Graph) override;
+	virtual FReply DroppedOnPanel(const TSharedRef< SWidget >& Panel, const FVector2f& ScreenPosition, const FVector2f& GraphPosition, UEdGraph& Graph) override;
 
 	virtual bool IsSupportedBySchema(const class UEdGraphSchema* Schema) const override;
 	// End of FGraphEditorDragDropAction
@@ -56,7 +54,7 @@ public:
 	/** Structure for required node construction parameters */
 	struct FNodeConstructionParams
 	{
-		FVector2D GraphPosition;
+		FDeprecateSlateVector2D GraphPosition;
 		UEdGraph* Graph;
 		bool bSelfContext;
 		const FProperty* Property;
diff --git a/Engine/Source/Editor/Kismet/Private/BPFunctionDragDropAction.cpp b/Engine/Source/Editor/Kismet/Private/BPFunctionDragDropAction.cpp
--- a/Engine/Source/Editor/Kismet/Private/BPFunctionDragDropAction.cpp
+++ b/Engine/Source/Editor/Kismet/Private/BPFunctionDragDropAction.cpp
@@ -146,7 +146,7 @@ void FKismetDragDropAction::HoverTargetChanged()
 }
 
 //------------------------------------------------------------------------------
-FReply FKismetDragDropAction::DroppedOnPanel( const TSharedRef< SWidget >& Panel, FVector2D ScreenPosition, FVector2D GraphPosition, UEdGraph& Graph)
+FReply FKismetDragDropAction::DroppedOnPanel( const TSharedRef< SWidget >& Panel, const FVector2f& ScreenPosition, const FVector2f& GraphPosition, UEdGraph& Graph)
 {	
 	FReply Reply = FReply::Unhandled();
 
@@ -243,13 +243,13 @@ FKismetFunctionDragDropAction::FKismetFunctionDragDropAction()
 }
 
 //------------------------------------------------------------------------------
-FReply FKismetFunctionDragDropAction::DroppedOnPanel(TSharedRef<SWidget> const& Panel, FVector2D ScreenPosition, FVector2D GraphPosition, UEdGraph& Graph)
+FReply FKismetFunctionDragDropAction::DroppedOnPanel(TSharedRef<SWidget> const& Panel, const FVector2f& ScreenPosition, const FVector2f& GraphPosition, UEdGraph& Graph)
 {
 	return DroppedOnPin(ScreenPosition, GraphPosition);
 }
 
 //------------------------------------------------------------------------------
-FReply FKismetFunctionDragDropAction::DroppedOnPin(FVector2D ScreenPosition, FVector2D GraphPosition)
+FReply FKismetFunctionDragDropAction::DroppedOnPin(const FVector2f& ScreenPosition, const FVector2f& GraphPosition)
 {
 	FReply Reply = FReply::Unhandled();
 
@@ -275,7 +275,7 @@ FReply FKismetFunctionDragDropAction::DroppedOnPin(FVector2D ScreenPosition, FVe
 					const FScopedTransaction Transaction(LOCTEXT("KismetFunction_DroppedOnPanel", "Function Dropped on Graph"));
 
 					IBlueprintNodeBinder::FBindingSet Bindings;
-					UEdGraphNode* ResultNode = FunctionNodeSpawner->Invoke(Graph, Bindings, GraphPosition);
+					UEdGraphNode* ResultNode = FunctionNodeSpawner->Invoke(Graph, Bindings, FDeprecateSlateVector2D(GraphPosition));
 
 					// Autowire the node if we were dragging on top of a pin
 					if (ResultNode != nullptr)
@@ -376,7 +376,7 @@ FKismetMacroDragDropAction::FKismetMacroDragDropAction()
 }
 
 //------------------------------------------------------------------------------
-FReply FKismetMacroDragDropAction::DroppedOnPanel(TSharedRef<SWidget> const& Panel, FVector2D ScreenPosition, FVector2D GraphPosition, UEdGraph& Graph)
+FReply FKismetMacroDragDropAction::DroppedOnPanel(TSharedRef<SWidget> const& Panel, const FVector2f& ScreenPosition, const FVector2f& GraphPosition, UEdGraph& Graph)
 {	
 	check(Macro);
 	check(CanBeDroppedDelegate.IsBound());
@@ -390,7 +390,7 @@ FReply FKismetMacroDragDropAction::DroppedOnPanel(TSharedRef<SWidget> const& Pan
 	{
 		FEdGraphSchemaAction_K2NewNode::SpawnNode<UK2Node_MacroInstance>(
 			&Graph,
-			GraphPosition,
+			FDeprecateSlateVector2D(GraphPosition),
 			EK2NewNodeFlags::SelectNewNode,
 			[MacroForCapture, &AnalyticCallbackForCapture](UK2Node_MacroInstance* NewInstance)
 			{
diff --git a/Engine/Source/Editor/Kismet/Private/BPFunctionDragDropAction.h b/Engine/Source/Editor/Kismet/Private/BPFunctionDragDropAction.h
--- a/Engine/Source/Editor/Kismet/Private/BPFunctionDragDropAction.h
+++ b/Engine/Source/Editor/Kismet/Private/BPFunctionDragDropAction.h
@@ -31,7 +31,7 @@ public:
 		
 	// FGraphEditorDragDropAction interface
 	virtual void HoverTargetChanged() override;
-	virtual FReply DroppedOnPanel( const TSharedRef< class SWidget >& Panel, FVector2D ScreenPosition, FVector2D GraphPosition, UEdGraph& Graph) override;
+	virtual FReply DroppedOnPanel( const TSharedRef< class SWidget >& Panel, const FVector2f& ScreenPosition, const FVector2f& GraphPosition, UEdGraph& Graph) override;
 	// End of FGraphSchemaActionDragDropAction
 
 	DECLARE_DELEGATE_RetVal_ThreeParams(bool, FCanBeDroppedDelegate, TSharedPtr<FEdGraphSchemaAction> /*DropAction*/, UEdGraph* /*HoveredGraphIn*/, FText& /*ImpededReasonOut*/);
@@ -65,8 +65,8 @@ public:
 	FKismetFunctionDragDropAction();
 
 	// FGraphEditorDragDropAction interface
-	virtual FReply DroppedOnPanel( const TSharedRef< class SWidget >& Panel, FVector2D ScreenPosition, FVector2D GraphPosition, UEdGraph& Graph) override;
-	virtual FReply DroppedOnPin(FVector2D ScreenPosition, FVector2D GraphPosition) override;
+	virtual FReply DroppedOnPanel( const TSharedRef< class SWidget >& Panel, const FVector2f& ScreenPosition, const FVector2f& GraphPosition, UEdGraph& Graph) override;
+	virtual FReply DroppedOnPin(const FVector2f& ScreenPosition, const FVector2f& GraphPosition) override;
 	// End of FGraphEditorDragDropAction
 
 	static TSharedRef<FKismetFunctionDragDropAction> New(TSharedPtr<FEdGraphSchemaAction> InActionNode, FName InFunctionName, UClass* InOwningClass, const FMemberReference& InCallOnMember, FNodeCreationAnalytic AnalyticCallback, FCanBeDroppedDelegate CanBeDroppedDelegate = FCanBeDroppedDelegate());
@@ -98,7 +98,7 @@ public:
 	FKismetMacroDragDropAction();
 
 	// FGraphEditorDragDropAction interface
-	virtual FReply DroppedOnPanel( const TSharedRef< class SWidget >& Panel, FVector2D ScreenPosition, FVector2D GraphPosition, UEdGraph& Graph) override;
+	virtual FReply DroppedOnPanel( const TSharedRef< class SWidget >& Panel, const FVector2f& ScreenPosition, const FVector2f& GraphPosition, UEdGraph& Graph) override;
 	// End of FGraphEditorDragDropAction
 
 	static TSharedRef<FKismetMacroDragDropAction> New(TSharedPtr<FEdGraphSchemaAction> InActionNode, FName InMacroName, UBlueprint* InBlueprint, UEdGraph* InMacro, FNodeCreationAnalytic AnalyticCallback);
diff --git a/Engine/Source/Editor/Kismet/Private/BPVariableDragDropAction.cpp b/Engine/Source/Editor/Kismet/Private/BPVariableDragDropAction.cpp
--- a/Engine/Source/Editor/Kismet/Private/BPVariableDragDropAction.cpp
+++ b/Engine/Source/Editor/Kismet/Private/BPVariableDragDropAction.cpp
@@ -285,7 +285,7 @@ void FKismetVariableDragDropAction::GetDefaultStatusSymbol(const FSlateBrush*& P
 	PrimaryBrushOut = FBlueprintEditor::GetVarIconAndColor(VariableSource.Get(), VariableName, IconColorOut, SecondaryBrushOut, SecondaryColorOut);
 }
 
-FReply FKismetVariableDragDropAction::DroppedOnPin(FVector2D ScreenPosition, FVector2D GraphPosition)
+FReply FKismetVariableDragDropAction::DroppedOnPin(const FVector2f& ScreenPosition, const FVector2f& GraphPosition)
 {
 	if (UEdGraphPin* TargetPin = GetHoveredPin())
 	{
@@ -297,7 +297,7 @@ FReply FKismetVariableDragDropAction::DroppedOnPin(FVector2D ScreenPosition, FVe
 				return FReply::Unhandled();
 			}
 
-			bool DropReply = ((UEdGraphSchema*)TargetPin->GetSchema())->RequestVariableDropOnPin(TargetPin->GetOwningNode()->GetGraph(), VariableProperty, TargetPin, GraphPosition, ScreenPosition);
+			bool DropReply = ((UEdGraphSchema*)TargetPin->GetSchema())->RequestVariableDropOnPin(TargetPin->GetOwningNode()->GetGraph(), VariableProperty, TargetPin, FDeprecateSlateVector2D(GraphPosition), FDeprecateSlateVector2D(ScreenPosition));
 			if (DropReply)
 			{
 				return FReply::Handled();
@@ -338,7 +338,7 @@ FReply FKismetVariableDragDropAction::DroppedOnPin(FVector2D ScreenPosition, FVe
 						UBlueprint* DropOnBlueprint = FBlueprintEditorUtils::FindBlueprintForGraph(TargetPin->GetOwningNode()->GetGraph());
 						UEdGraphSchema_K2::ConfigureVarNode(VarNode, VariableName, VariableSource.Get(), DropOnBlueprint);
 
-						Action.PerformAction(TargetPin->GetOwningNode()->GetGraph(), TargetPin, GraphPosition);
+						Action.PerformAction(TargetPin->GetOwningNode()->GetGraph(), TargetPin, FDeprecateSlateVector2D(GraphPosition));
 					}
 				}
 			}
@@ -348,7 +348,7 @@ FReply FKismetVariableDragDropAction::DroppedOnPin(FVector2D ScreenPosition, FVe
 	return FReply::Handled();
 }
 
-FReply FKismetVariableDragDropAction::DroppedOnNode(FVector2D ScreenPosition, FVector2D GraphPosition)
+FReply FKismetVariableDragDropAction::DroppedOnNode(const FVector2f& ScreenPosition, const FVector2f& GraphPosition)
 {
 	if (UEdGraphNode* TargetNode = GetHoveredNode())
 	{
@@ -358,7 +358,7 @@ FReply FKismetVariableDragDropAction::DroppedOnNode(FVector2D ScreenPosition, FV
 			return FReply::Unhandled();
 		}
 
-		bool DropReply = ((UEdGraphSchema*)TargetNode->GetSchema())->RequestVariableDropOnNode(TargetNode->GetGraph(), VariableProperty, TargetNode, GraphPosition, ScreenPosition);
+		bool DropReply = ((UEdGraphSchema*)TargetNode->GetSchema())->RequestVariableDropOnNode(TargetNode->GetGraph(), VariableProperty, TargetNode, FDeprecateSlateVector2D(GraphPosition), FDeprecateSlateVector2D(ScreenPosition));
 		if (DropReply)
 		{
 			return FReply::Handled();
@@ -467,10 +467,10 @@ bool FKismetVariableDragDropAction::CanExecuteMakeSetter(FNodeConstructionParams
 	return true;
 }
 
-FReply FKismetVariableDragDropAction::DroppedOnPanel( const TSharedRef< SWidget >& Panel, FVector2D ScreenPosition, FVector2D GraphPosition, UEdGraph& Graph)
+FReply FKismetVariableDragDropAction::DroppedOnPanel( const TSharedRef< SWidget >& Panel, const FVector2f& ScreenPosition, const FVector2f& GraphPosition, UEdGraph& Graph)
 {	
 	FProperty* VariableProperty = GetVariableProperty();
-	bool  DropReply = ((UEdGraphSchema*)Graph.GetSchema())->RequestVariableDropOnPanel(GetHoveredGraph(), VariableProperty, GraphPosition, ScreenPosition);
+	bool  DropReply = ((UEdGraphSchema*)Graph.GetSchema())->RequestVariableDropOnPanel(GetHoveredGraph(), VariableProperty, FDeprecateSlateVector2D(GraphPosition), FDeprecateSlateVector2D(ScreenPosition));
 	if (DropReply)
 	{
 		return FReply::Handled();
diff --git a/Engine/Source/Editor/Kismet/Private/SSCSEditor.cpp b/Engine/Source/Editor/Kismet/Private/SSCSEditor.cpp
--- a/Engine/Source/Editor/Kismet/Private/SSCSEditor.cpp
+++ b/Engine/Source/Editor/Kismet/Private/SSCSEditor.cpp
@@ -349,8 +349,8 @@ public:
 
 	// FGraphEditorDragDropAction interface
 	virtual void HoverTargetChanged() override;
-	virtual FReply DroppedOnNode(FVector2D ScreenPosition, FVector2D GraphPosition) override;
-	virtual FReply DroppedOnPanel(const TSharedRef< class SWidget >& Panel, FVector2D ScreenPosition, FVector2D GraphPosition, UEdGraph& Graph) override;
+	virtual FReply DroppedOnNode(const FVector2f& ScreenPosition, const FVector2f& GraphPosition) override;
+	virtual FReply DroppedOnPanel(const TSharedRef< class SWidget >& Panel, const FVector2f& ScreenPosition, const FVector2f& GraphPosition, UEdGraph& Graph) override;
 	// End of FGraphEditorDragDropAction
 
 	/** Node(s) that we started the drag from */
@@ -426,7 +426,7 @@ void FSCSRowDragDropOp::HoverTargetChanged()
 	}
 }
 
-FReply FSCSRowDragDropOp::DroppedOnNode(FVector2D ScreenPosition, FVector2D GraphPosition)
+FReply FSCSRowDragDropOp::DroppedOnNode(const FVector2f& ScreenPosition, const FVector2f& GraphPosition)
 {
 	// Only allow dropping on another node if there is only a single source item
 	if(SourceNodes.Num() == 1)
@@ -436,20 +436,21 @@ FReply FSCSRowDragDropOp::DroppedOnNode(FVector2D ScreenPosition, FVector2D Grap
 	return FReply::Handled();
 }
 
-FReply FSCSRowDragDropOp::DroppedOnPanel(const TSharedRef< class SWidget >& Panel, FVector2D ScreenPosition, FVector2D GraphPosition, UEdGraph& Graph)
+FReply FSCSRowDragDropOp::DroppedOnPanel(const TSharedRef< class SWidget >& Panel, const FVector2f& ScreenPosition, const FVector2f& GraphPosition, UEdGraph& Graph)
 {
 	const FScopedTransaction Transaction(LOCTEXT("SCSEditorAddMultipleNodes", "Add Component Nodes"));
 
 	TArray<UK2Node_VariableGet*> OriginalVariableNodes;
 	Graph.GetNodesOfClass<UK2Node_VariableGet>(OriginalVariableNodes);
 
+	FVector2f GraphPositionToPlace = GraphPosition;
 	// Add source items to the graph in turn
 	for (FSCSEditorTreeNodePtrType& SourceNode : SourceNodes)
 	{
 		VariableName = SourceNode->GetVariableName();
-		FKismetVariableDragDropAction::DroppedOnPanel(Panel, ScreenPosition, GraphPosition, Graph);
+		FKismetVariableDragDropAction::DroppedOnPanel(Panel, ScreenPosition, GraphPositionToPlace, Graph);
 
-		GraphPosition.Y += 50;
+		GraphPositionToPlace.Y += 50;
 	}
 
 	TArray<UK2Node_VariableGet*> ResultVariableNodes;
diff --git a/Engine/Source/Editor/Kismet/Public/BPVariableDragDropAction.h b/Engine/Source/Editor/Kismet/Public/BPVariableDragDropAction.h
--- a/Engine/Source/Editor/Kismet/Public/BPVariableDragDropAction.h
+++ b/Engine/Source/Editor/Kismet/Public/BPVariableDragDropAction.h
@@ -35,9 +35,9 @@ public:
 
 	// FGraphEditorDragDropAction interface
 	virtual void HoverTargetChanged() override;
-	virtual FReply DroppedOnPin(FVector2D ScreenPosition, FVector2D GraphPosition) override;
-	virtual FReply DroppedOnNode(FVector2D ScreenPosition, FVector2D GraphPosition) override;
-	virtual FReply DroppedOnPanel(const TSharedRef< class SWidget >& Panel, FVector2D ScreenPosition, FVector2D GraphPosition, UEdGraph& Graph) override;
+	virtual FReply DroppedOnPin(const FVector2f& ScreenPosition, const FVector2f& GraphPosition) override;
+	virtual FReply DroppedOnNode(const FVector2f& ScreenPosition, const FVector2f& GraphPosition) override;
+	virtual FReply DroppedOnPanel(const TSharedRef< class SWidget >& Panel, const FVector2f& ScreenPosition, const FVector2f& GraphPosition, UEdGraph& Graph) override;
 	// End of FGraphEditorDragDropAction
 
 	static TSharedRef<FKismetVariableDragDropAction> New(TSharedPtr<FEdGraphSchemaAction> InAction, FName InVariableName, UStruct* InVariableSource, FNodeCreationAnalytic AnalyticCallback)
@@ -67,7 +67,7 @@ protected:
 	/** Structure for required node construction parameters */
 	struct FNodeConstructionParams
 	{
-		FVector2D GraphPosition;
+		FDeprecateSlateVector2D GraphPosition;
 		UEdGraph* Graph;
 		FName VariableName;
 		TWeakObjectPtr<UStruct> VariableSource;
diff --git a/Engine/Source/Editor/Sequencer/Private/SequencerOutlinerItemDragDropOp.cpp b/Engine/Source/Editor/Sequencer/Private/SequencerOutlinerItemDragDropOp.cpp
--- a/Engine/Source/Editor/Sequencer/Private/SequencerOutlinerItemDragDropOp.cpp
+++ b/Engine/Source/Editor/Sequencer/Private/SequencerOutlinerItemDragDropOp.cpp
@@ -94,7 +94,7 @@ void FSequencerOutlinerDragDropOp::HoverTargetChanged()
 	}
 }
 
-FReply FSequencerOutlinerDragDropOp::DroppedOnPanel( const TSharedRef< class SWidget >& Panel, FVector2D ScreenPosition, FVector2D GraphPosition, UEdGraph& Graph)
+FReply FSequencerOutlinerDragDropOp::DroppedOnPanel( const TSharedRef< class SWidget >& Panel, const FVector2f& ScreenPosition, const FVector2f& GraphPosition, UEdGraph& Graph)
 {
 	using namespace MovieScene;
 
@@ -106,6 +106,7 @@ FReply FSequencerOutlinerDragDropOp::DroppedOnPanel( const TSharedRef< class SWi
 		FEdGraphSchemaAction_NewNode Action;
 		Action.NodeTemplate = Template;
 
+		FVector2f AddGraphPosition = GraphPosition;
 		for (const TWeakViewModelPtr<IOutlinerExtension>& WeakModel : GetDraggedViewModels())
 		{
 			FViewModelPtr Model = WeakModel.Pin();
@@ -121,11 +122,11 @@ FReply FSequencerOutlinerDragDropOp::DroppedOnPanel( const TSharedRef< class SWi
 				// Fixed bindings always resolve from the root
 				Template->SetSequence(SequenceModel->GetSequencer()->GetRootMovieSceneSequence());
 				Template->Binding = FFixedObjectBindingID(ObjectBinding->GetObjectGuid(), SequenceModel->GetSequenceID());
-				UEdGraphNode* NewNode = Action.PerformAction(&Graph, GetHoveredPin(), GraphPosition, false);
+				UEdGraphNode* NewNode = Action.PerformAction(&Graph, GetHoveredPin(), FDeprecateSlateVector2D(AddGraphPosition), false);
 
 				int32 Offset = FMath::Max(NewNode->NodeHeight, 100);
 				Offset += Offset % 16;
-				GraphPosition.Y += Offset;
+				AddGraphPosition.Y += Offset;
 			}
 		}
 	}
diff --git a/Engine/Source/Editor/Sequencer/Private/SequencerOutlinerItemDragDropOp.h b/Engine/Source/Editor/Sequencer/Private/SequencerOutlinerItemDragDropOp.h
--- a/Engine/Source/Editor/Sequencer/Private/SequencerOutlinerItemDragDropOp.h
+++ b/Engine/Source/Editor/Sequencer/Private/SequencerOutlinerItemDragDropOp.h
@@ -40,7 +40,7 @@ public:
 
 	//~ FGraphEditorDragDropAction interface
 	virtual void HoverTargetChanged() override;
-	virtual FReply DroppedOnPanel( const TSharedRef< class SWidget >& Panel, FVector2D ScreenPosition, FVector2D GraphPosition, UEdGraph& Graph) override;
+	virtual FReply DroppedOnPanel( const TSharedRef< class SWidget >& Panel, const FVector2f& ScreenPosition, const FVector2f& GraphPosition, UEdGraph& Graph) override;
 };
 
 
diff --git a/Engine/Source/Editor/SubobjectEditor/Private/SSubobjectEditor.cpp b/Engine/Source/Editor/SubobjectEditor/Private/SSubobjectEditor.cpp
--- a/Engine/Source/Editor/SubobjectEditor/Private/SSubobjectEditor.cpp
+++ b/Engine/Source/Editor/SubobjectEditor/Private/SSubobjectEditor.cpp
@@ -504,8 +504,8 @@ public:
 
 	// FGraphEditorDragDropAction interface
 	virtual void HoverTargetChanged() override;
-	virtual FReply DroppedOnNode(FVector2D ScreenPosition, FVector2D GraphPosition) override;
-	virtual FReply DroppedOnPanel(const TSharedRef< class SWidget >& Panel, FVector2D ScreenPosition, FVector2D GraphPosition, UEdGraph& Graph) override;
+	virtual FReply DroppedOnNode(const FVector2f& ScreenPosition, const FVector2f& GraphPosition) override;
+	virtual FReply DroppedOnPanel(const TSharedRef< class SWidget >& Panel, const FVector2f& ScreenPosition, const FVector2f& GraphPosition, UEdGraph& Graph) override;
 	// End of FGraphEditorDragDropAction
 
 	/** Node(s) that we started the drag from */
@@ -581,7 +581,7 @@ void FSubobjectRowDragDropOp::HoverTargetChanged()
 	}
 }
 
-FReply FSubobjectRowDragDropOp::DroppedOnNode(FVector2D ScreenPosition, FVector2D GraphPosition)
+FReply FSubobjectRowDragDropOp::DroppedOnNode(const FVector2f& ScreenPosition, const FVector2f& GraphPosition)
 {
 	// Only allow dropping on another node if there is only a single source item
 	if(SourceNodes.Num() == 1)
@@ -591,7 +591,7 @@ FReply FSubobjectRowDragDropOp::DroppedOnNode(FVector2D ScreenPosition, FVector2
 	return FReply::Handled();
 }
 
-FReply FSubobjectRowDragDropOp::DroppedOnPanel(const TSharedRef<class SWidget>& Panel, FVector2D ScreenPosition, FVector2D GraphPosition, UEdGraph& Graph)
+FReply FSubobjectRowDragDropOp::DroppedOnPanel(const TSharedRef<class SWidget>& Panel, const FVector2f& ScreenPosition, const FVector2f& GraphPosition, UEdGraph& Graph)
 {
 	const FScopedTransaction Transaction(LOCTEXT("SSubobjectEditorAddMultipleNodes", "Add Component Nodes"));
 
@@ -599,6 +599,7 @@ FReply FSubobjectRowDragDropOp::DroppedOnPanel(const TSharedRef<class SWidget>&
 	Graph.GetNodesOfClass<UK2Node_VariableGet>(OriginalVariableNodes);
 
 	// Add source items to the graph in turn
+	FVector2f GraphPositionToPlace = GraphPosition;
 	for (FSubobjectEditorTreeNodePtrType& SourceNode : SourceNodes)
 	{
 		FSubobjectData* Data = SourceNode->GetDataSource();
@@ -608,9 +609,9 @@ FReply FSubobjectRowDragDropOp::DroppedOnPanel(const TSharedRef<class SWidget>&
 		}
 		
 		VariableName = Data->GetVariableName();
-		FKismetVariableDragDropAction::DroppedOnPanel(Panel, ScreenPosition, GraphPosition, Graph);
+		FKismetVariableDragDropAction::DroppedOnPanel(Panel, ScreenPosition, GraphPositionToPlace, Graph);
 
-		GraphPosition.Y += 50.0f;
+		GraphPositionToPlace.Y += 50.0f;
 	}
 
 	TArray<UK2Node_VariableGet*> ResultVariableNodes;
