[Truncation Warnings] Deprecate FVector2D delegates in GraphEditor module

diff --git a/Engine/Plugins/Animation/ControlRig/Source/ControlRigEditor/Private/Editor/ControlRigEditor.cpp b/Engine/Plugins/Animation/ControlRig/Source/ControlRigEditor/Private/Editor/ControlRigEditor.cpp
--- a/Engine/Plugins/Animation/ControlRig/Source/ControlRigEditor/Private/Editor/ControlRigEditor.cpp
+++ b/Engine/Plugins/Animation/ControlRig/Source/ControlRigEditor/Private/Editor/ControlRigEditor.cpp
@@ -1868,7 +1868,7 @@ FString FControlRigBaseEditor::GetWorldCentricTabPrefixImpl() const
 	return LOCTEXT("WorldCentricTabPrefix", "Control Rig Editor ").ToString();
 }
 
-FReply FControlRigBaseEditor::OnSpawnGraphNodeByShortcutImpl(FInputChord InChord, const FVector2D& InPosition, UEdGraph* InGraph)
+FReply FControlRigBaseEditor::OnSpawnGraphNodeByShortcutImpl(FInputChord InChord, const FVector2f& InPosition, UEdGraph* InGraph)
 {
 	const FReply SuperReply = OnSpawnGraphNodeByShortcutSuper(InChord, InPosition, InGraph);
 	if(SuperReply.IsEventHandled())
@@ -1882,29 +1882,30 @@ FReply FControlRigBaseEditor::OnSpawnGraphNodeByShortcutImpl(FInputChord InChord
 		{
 			if(URigVMController* Controller = RigGraph->GetController())
 			{
+				FDeprecateSlateVector2D Position = InPosition;
 				if(InChord.Key == EKeys::S)
 				{
-					Controller->AddUnitNode(FRigVMFunction_Sequence::StaticStruct(), FRigUnit::GetMethodName(), InPosition, FString(), true, true);
+					Controller->AddUnitNode(FRigVMFunction_Sequence::StaticStruct(), FRigUnit::GetMethodName(), Position, FString(), true, true);
 				}
 				else if(InChord.Key == EKeys::One)
 				{
-					Controller->AddUnitNode(FRigUnit_GetTransform::StaticStruct(), FRigUnit::GetMethodName(), InPosition, FString(), true, true);
+					Controller->AddUnitNode(FRigUnit_GetTransform::StaticStruct(), FRigUnit::GetMethodName(), Position, FString(), true, true);
 				}
 				else if(InChord.Key == EKeys::Two)
 				{
-					Controller->AddUnitNode(FRigUnit_SetTransform::StaticStruct(), FRigUnit::GetMethodName(), InPosition, FString(), true, true);
+					Controller->AddUnitNode(FRigUnit_SetTransform::StaticStruct(), FRigUnit::GetMethodName(), Position, FString(), true, true);
 				}
 				else if(InChord.Key == EKeys::Three)
 				{
-					Controller->AddUnitNode(FRigUnit_ParentConstraint::StaticStruct(), FRigUnit::GetMethodName(), InPosition, FString(), true, true);
+					Controller->AddUnitNode(FRigUnit_ParentConstraint::StaticStruct(), FRigUnit::GetMethodName(), Position, FString(), true, true);
 				}
 				else if(InChord.Key == EKeys::Four)
 				{
-					Controller->AddUnitNode(FRigUnit_GetControlFloat::StaticStruct(), FRigUnit::GetMethodName(), InPosition, FString(), true, true);
+					Controller->AddUnitNode(FRigUnit_GetControlFloat::StaticStruct(), FRigUnit::GetMethodName(), Position, FString(), true, true);
 				}
 				else if(InChord.Key == EKeys::Five)
 				{
-					Controller->AddUnitNode(FRigUnit_SetCurveValue::StaticStruct(), FRigUnit::GetMethodName(), InPosition, FString(), true, true);
+					Controller->AddUnitNode(FRigUnit_SetCurveValue::StaticStruct(), FRigUnit::GetMethodName(), Position, FString(), true, true);
 				}
 			}
 		}
@@ -5158,7 +5159,7 @@ void FControlRigBaseEditor::CreateRigHierarchyToGraphDragAndDropMenu() const
 	}
 }
 
-void FControlRigBaseEditor::OnGraphNodeDropToPerformImpl(TSharedPtr<FDragDropOperation> InDragDropOp, UEdGraph* InGraph, const FVector2D& InNodePosition, const FVector2D& InScreenPosition)
+void FControlRigBaseEditor::OnGraphNodeDropToPerformImpl(TSharedPtr<FDragDropOperation> InDragDropOp, UEdGraph* InGraph, const FVector2f& InNodePosition, const FVector2f& InScreenPosition)
 {
 	DECLARE_SCOPE_HIERARCHICAL_COUNTER_FUNC()
 
@@ -5176,7 +5177,7 @@ void FControlRigBaseEditor::OnGraphNodeDropToPerformImpl(TSharedPtr<FDragDropOpe
 				FControlRigRigHierarchyToGraphDragAndDropContext(
 					RigHierarchyOp->GetElements(),
 					InGraph,
-					InNodePosition
+					FDeprecateSlateVector2D(InNodePosition)
 				);
 			MenuContext->Init(StaticCastSharedRef<FControlRigBaseEditor>(SharedControlRigEditorRef()), MenuSpecificContext);
 			
diff --git a/Engine/Plugins/Animation/ControlRig/Source/ControlRigEditor/Private/Editor/ControlRigEditor.h b/Engine/Plugins/Animation/ControlRig/Source/ControlRigEditor/Private/Editor/ControlRigEditor.h
--- a/Engine/Plugins/Animation/ControlRig/Source/ControlRigEditor/Private/Editor/ControlRigEditor.h
+++ b/Engine/Plugins/Animation/ControlRig/Source/ControlRigEditor/Private/Editor/ControlRigEditor.h
@@ -159,7 +159,7 @@ protected:
 	virtual void SetEventQueueSuper(TArray<FName> InEventQueue) = 0;
 	virtual void SaveAsset_ExecuteSuper() = 0;
 	virtual void SaveAssetAs_ExecuteSuper() = 0;
-	virtual FReply OnSpawnGraphNodeByShortcutSuper(FInputChord InChord, const FVector2D& InPosition, UEdGraph* InGraph) = 0;
+	virtual FReply OnSpawnGraphNodeByShortcutSuper(FInputChord InChord, const FVector2f& InPosition, UEdGraph* InGraph) = 0;
 	virtual void HandleVMExecutedEventSuper(URigVMHost* InHost, const FName& InEventName) = 0;
 	virtual void SetEditorModeManager(TSharedPtr<FEditorModeTools> InManager) = 0;
 	virtual void TickSuper(float DeltaTime) = 0;
@@ -171,9 +171,9 @@ protected:
 	virtual void BindCommandsSuper() = 0;
 	virtual FMenuBuilder GenerateBulkEditMenuSuper() = 0;
 	virtual TWeakPtr<class SGraphEditor> GetFocusedGraphEd() const = 0;
-	virtual void OnGraphNodeDropToPerformSuper(TSharedPtr<FDragDropOperation> InDragDropOp, UEdGraph* InGraph, const FVector2D& InNodePosition, const FVector2D& InScreenPosition) = 0;
+	virtual void OnGraphNodeDropToPerformSuper(TSharedPtr<FDragDropOperation> InDragDropOp, UEdGraph* InGraph, const FVector2f& InNodePosition, const FVector2f& InScreenPosition) = 0;
 	virtual void HandleRefreshEditorFromBlueprintSuper(URigVMBlueprint* InBlueprint) = 0;
-	virtual void OnGraphNodeDropToPerform(TSharedPtr<FDragDropOperation> InDragDropOp, UEdGraph* InGraph, const FVector2D& InNodePosition, const FVector2D& InScreenPosition) = 0;
+	virtual void OnGraphNodeDropToPerform(TSharedPtr<FDragDropOperation> InDragDropOp, UEdGraph* InGraph, const FVector2f& InNodePosition, const FVector2f& InScreenPosition) = 0;
 	
 	virtual void OnCloseSuper() = 0;
 
@@ -231,7 +231,7 @@ protected:
 	}
 
 	// BlueprintEditor interface
-	virtual FReply OnSpawnGraphNodeByShortcutImpl(FInputChord InChord, const FVector2D& InPosition, UEdGraph* InGraph);
+	virtual FReply OnSpawnGraphNodeByShortcutImpl(FInputChord InChord, const FVector2f& InPosition, UEdGraph* InGraph);
 
 	virtual void PostTransactionImpl(bool bSuccess, const FTransaction* Transaction, bool bIsRedo);
 
@@ -291,7 +291,7 @@ protected:
 
 	const FName RigHierarchyToGraphDragAndDropMenuName = TEXT("ControlRigEditor.RigHierarchyToGraphDragAndDropMenu");
 	void CreateRigHierarchyToGraphDragAndDropMenu() const;
-	virtual void OnGraphNodeDropToPerformImpl(TSharedPtr<FDragDropOperation> InDragDropOp, UEdGraph* InGraph, const FVector2D& InNodePosition, const FVector2D& InScreenPosition);
+	virtual void OnGraphNodeDropToPerformImpl(TSharedPtr<FDragDropOperation> InDragDropOp, UEdGraph* InGraph, const FVector2f& InNodePosition, const FVector2f& InScreenPosition);
 
 	virtual FPersonaViewportKeyDownDelegate& GetKeyDownDelegate() override { return OnKeyDownDelegate; }
 
diff --git a/Engine/Plugins/Animation/ControlRig/Source/ControlRigEditor/Private/Editor/ControlRigLegacyEditor.h b/Engine/Plugins/Animation/ControlRig/Source/ControlRigEditor/Private/Editor/ControlRigLegacyEditor.h
--- a/Engine/Plugins/Animation/ControlRig/Source/ControlRigEditor/Private/Editor/ControlRigLegacyEditor.h
+++ b/Engine/Plugins/Animation/ControlRig/Source/ControlRigEditor/Private/Editor/ControlRigLegacyEditor.h
@@ -4,6 +4,7 @@
 
 #if WITH_RIGVMLEGACYEDITOR
 #include "ControlRigEditor.h"
+#include "Types/SlateVector2.h"
 
 class FControlRigLegacyEditor : public IControlRigLegacyEditor, public FControlRigBaseEditor
 {
@@ -81,8 +82,8 @@ public:
 	virtual FString GetDocumentationLink() const override { return FControlRigBaseEditor::GetDocumentationLinkImpl(); }
 
 	// BlueprintEditor interface
-	virtual FReply OnSpawnGraphNodeByShortcut(FInputChord InChord, const FVector2D& InPosition, UEdGraph* InGraph) override { return FControlRigBaseEditor::OnSpawnGraphNodeByShortcutImpl(InChord, InPosition, InGraph); }
-	virtual FReply OnSpawnGraphNodeByShortcutSuper(FInputChord InChord, const FVector2D& InPosition, UEdGraph* InGraph) override { return IControlRigLegacyEditor::OnSpawnGraphNodeByShortcut(InChord, InPosition, InGraph); }
+	virtual FReply OnSpawnGraphNodeByShortcut(FInputChord InChord, const FVector2f& InPosition, UEdGraph* InGraph) override { return FControlRigBaseEditor::OnSpawnGraphNodeByShortcutImpl(InChord, InPosition, InGraph); }
+	virtual FReply OnSpawnGraphNodeByShortcutSuper(FInputChord InChord, const FVector2f& InPosition, UEdGraph* InGraph) override { return IControlRigLegacyEditor::OnSpawnGraphNodeByShortcut(InChord, InPosition, InGraph); }
 #if WITH_RIGVMLEGACYEDITOR
 	virtual bool IsSectionVisible(NodeSectionID::Type InSectionID) const override;
 	virtual bool NewDocument_IsVisibleForType(FBlueprintEditor::ECreatedDocumentType GraphType) const override;
@@ -105,8 +106,8 @@ public:
 public:
 
 	
-	virtual void OnGraphNodeDropToPerform(TSharedPtr<FDragDropOperation> InDragDropOp, UEdGraph* InGraph, const FVector2D& InNodePosition, const FVector2D& InScreenPosition) override  { return FControlRigBaseEditor::OnGraphNodeDropToPerformImpl(InDragDropOp, InGraph, InNodePosition, InScreenPosition); }
-	virtual void OnGraphNodeDropToPerformSuper(TSharedPtr<FDragDropOperation> InDragDropOp, UEdGraph* InGraph, const FVector2D& InNodePosition, const FVector2D& InScreenPosition) override  { return IControlRigLegacyEditor::OnGraphNodeDropToPerform(InDragDropOp, InGraph, InNodePosition, InScreenPosition); }
+	virtual void OnGraphNodeDropToPerform(TSharedPtr<FDragDropOperation> InDragDropOp, UEdGraph* InGraph, const FVector2f& InNodePosition, const FVector2f& InScreenPosition) override  { return FControlRigBaseEditor::OnGraphNodeDropToPerformImpl(InDragDropOp, InGraph, InNodePosition, InScreenPosition); }
+	virtual void OnGraphNodeDropToPerformSuper(TSharedPtr<FDragDropOperation> InDragDropOp, UEdGraph* InGraph, const FVector2f& InNodePosition, const FVector2f& InScreenPosition) override  { return IControlRigLegacyEditor::OnGraphNodeDropToPerform(InDragDropOp, InGraph, InNodePosition, InScreenPosition); }
 
 	
 protected:
diff --git a/Engine/Plugins/Animation/ControlRig/Source/ControlRigEditor/Private/Editor/ControlRigNewEditor.h b/Engine/Plugins/Animation/ControlRig/Source/ControlRigEditor/Private/Editor/ControlRigNewEditor.h
--- a/Engine/Plugins/Animation/ControlRig/Source/ControlRigEditor/Private/Editor/ControlRigNewEditor.h
+++ b/Engine/Plugins/Animation/ControlRig/Source/ControlRigEditor/Private/Editor/ControlRigNewEditor.h
@@ -82,8 +82,8 @@ public:
 	virtual FString GetDocumentationLink() const override { return FControlRigBaseEditor::GetDocumentationLinkImpl(); }
 
 	// BlueprintEditor interface
-	virtual FReply OnSpawnGraphNodeByShortcut(FInputChord InChord, const FVector2D& InPosition, UEdGraph* InGraph) override { return FControlRigBaseEditor::OnSpawnGraphNodeByShortcutImpl(InChord, InPosition, InGraph); }
-	virtual FReply OnSpawnGraphNodeByShortcutSuper(FInputChord InChord, const FVector2D& InPosition, UEdGraph* InGraph) override { return IControlRigNewEditor::OnSpawnGraphNodeByShortcut(InChord, InPosition, InGraph); }
+	virtual FReply OnSpawnGraphNodeByShortcut(FInputChord InChord, const FVector2f& InPosition, UEdGraph* InGraph) override { return FControlRigBaseEditor::OnSpawnGraphNodeByShortcutImpl(InChord, InPosition, InGraph); }
+	virtual FReply OnSpawnGraphNodeByShortcutSuper(FInputChord InChord, const FVector2f& InPosition, UEdGraph* InGraph) override { return IControlRigNewEditor::OnSpawnGraphNodeByShortcut(InChord, InPosition, InGraph); }
 
 	virtual bool IsSectionVisible(RigVMNodeSectionID::Type InSectionID) const override;
 	virtual bool NewDocument_IsVisibleForType(FRigVMEditorBase::ECreatedDocumentType GraphType) const override;
@@ -104,8 +104,8 @@ public:
 public:
 
 	
-	virtual void OnGraphNodeDropToPerform(TSharedPtr<FDragDropOperation> InDragDropOp, UEdGraph* InGraph, const FVector2D& InNodePosition, const FVector2D& InScreenPosition) override  { return FControlRigBaseEditor::OnGraphNodeDropToPerformImpl(InDragDropOp, InGraph, InNodePosition, InScreenPosition); }
-	virtual void OnGraphNodeDropToPerformSuper(TSharedPtr<FDragDropOperation> InDragDropOp, UEdGraph* InGraph, const FVector2D& InNodePosition, const FVector2D& InScreenPosition) override  { return IControlRigNewEditor::OnGraphNodeDropToPerform(InDragDropOp, InGraph, InNodePosition, InScreenPosition); }
+	virtual void OnGraphNodeDropToPerform(TSharedPtr<FDragDropOperation> InDragDropOp, UEdGraph* InGraph, const FVector2f& InNodePosition, const FVector2f& InScreenPosition) override  { return FControlRigBaseEditor::OnGraphNodeDropToPerformImpl(InDragDropOp, InGraph, InNodePosition, InScreenPosition); }
+	virtual void OnGraphNodeDropToPerformSuper(TSharedPtr<FDragDropOperation> InDragDropOp, UEdGraph* InGraph, const FVector2f& InNodePosition, const FVector2f& InScreenPosition) override  { return IControlRigNewEditor::OnGraphNodeDropToPerform(InDragDropOp, InGraph, InNodePosition, InScreenPosition); }
 
 	
 protected:
diff --git a/Engine/Plugins/Animation/ControlRig/Source/ControlRigEditor/Private/Editor/SRigHierarchy.cpp b/Engine/Plugins/Animation/ControlRig/Source/ControlRigEditor/Private/Editor/SRigHierarchy.cpp
--- a/Engine/Plugins/Animation/ControlRig/Source/ControlRigEditor/Private/Editor/SRigHierarchy.cpp
+++ b/Engine/Plugins/Animation/ControlRig/Source/ControlRigEditor/Private/Editor/SRigHierarchy.cpp
@@ -3002,7 +3002,7 @@ FReply SRigHierarchy::OnDragDetected(const FGeometry& MyGeometry, const FPointer
 			UpdateConnectorMatchesOnDrag(DraggedElements);
 			
 			TSharedRef<FRigElementHierarchyDragDropOp> DragDropOp = FRigElementHierarchyDragDropOp::New(MoveTemp(DraggedElements));
-			DragDropOp->OnPerformDropToGraph.BindSP(EditorPtr.ToSharedRef(), &FControlRigBaseEditor::OnGraphNodeDropToPerform);
+			DragDropOp->OnPerformDropToGraphAtLocation.BindSP(EditorPtr.ToSharedRef(), &FControlRigBaseEditor::OnGraphNodeDropToPerform);
 			return FReply::Handled().BeginDragDrop(DragDropOp);
 		}
 	}
diff --git a/Engine/Plugins/Animation/DeformerGraph/Source/OptimusEditor/Private/OptimusEditor.cpp b/Engine/Plugins/Animation/DeformerGraph/Source/OptimusEditor/Private/OptimusEditor.cpp
--- a/Engine/Plugins/Animation/DeformerGraph/Source/OptimusEditor/Private/OptimusEditor.cpp
+++ b/Engine/Plugins/Animation/DeformerGraph/Source/OptimusEditor/Private/OptimusEditor.cpp
@@ -957,7 +957,7 @@ bool FOptimusEditor::OnVerifyNodeTextCommit(const FText& NewText, UEdGraphNode*
 }
 
 
-FReply FOptimusEditor::OnSpawnGraphNodeByShortcut(FInputChord InChord, const FVector2D& InPosition, UEdGraph* InGraph)
+FReply FOptimusEditor::OnSpawnGraphNodeByShortcut(FInputChord InChord, const FVector2f& InPosition, UEdGraph* InGraph)
 {
 	return FReply::Handled();
 }
@@ -1260,8 +1260,8 @@ TSharedRef<SGraphEditor> FOptimusEditor::CreateGraphEditorWidget()
 	InEvents.OnNodeDoubleClicked = FSingleNodeEvent::CreateSP(this, &FOptimusEditor::OnNodeDoubleClicked);
 	InEvents.OnTextCommitted = FOnNodeTextCommitted::CreateSP(this, &FOptimusEditor::OnNodeTitleCommitted);
 	InEvents.OnVerifyTextCommit = FOnNodeVerifyTextCommit::CreateSP(this, &FOptimusEditor::OnVerifyNodeTextCommit);
-	InEvents.OnSpawnNodeByShortcut = 
-		SGraphEditor::FOnSpawnNodeByShortcut::CreateSP(
+	InEvents.OnSpawnNodeByShortcutAtLocation = 
+		SGraphEditor::FOnSpawnNodeByShortcutAtLocation::CreateSP(
 			this, &FOptimusEditor::OnSpawnGraphNodeByShortcut, 
 			static_cast<UEdGraph*>(EditorGraph));
 
diff --git a/Engine/Plugins/Animation/DeformerGraph/Source/OptimusEditor/Private/OptimusEditor.h b/Engine/Plugins/Animation/DeformerGraph/Source/OptimusEditor/Private/OptimusEditor.h
--- a/Engine/Plugins/Animation/DeformerGraph/Source/OptimusEditor/Private/OptimusEditor.h
+++ b/Engine/Plugins/Animation/DeformerGraph/Source/OptimusEditor/Private/OptimusEditor.h
@@ -181,7 +181,7 @@ private:
 	void OnNodeDoubleClicked(class UEdGraphNode* Node);
 	void OnNodeTitleCommitted(const FText& NewText, ETextCommit::Type CommitInfo, UEdGraphNode* NodeBeingChanged);
 	bool OnVerifyNodeTextCommit(const FText& NewText, UEdGraphNode* NodeBeingChanged, FText& OutErrorMessage);
-	FReply OnSpawnGraphNodeByShortcut(FInputChord InChord, const FVector2D& InPosition, UEdGraph* InGraph);
+	FReply OnSpawnGraphNodeByShortcut(FInputChord InChord, const FVector2f& InPosition, UEdGraph* InGraph);
 
 public:
 	// ----------------------------------------------------------------------------------------
diff --git a/Engine/Plugins/Editor/AssetManagerEditor/Source/AssetManagerEditor/Private/ReferenceViewer/SReferenceViewer.cpp b/Engine/Plugins/Editor/AssetManagerEditor/Source/AssetManagerEditor/Private/ReferenceViewer/SReferenceViewer.cpp
--- a/Engine/Plugins/Editor/AssetManagerEditor/Source/AssetManagerEditor/Private/ReferenceViewer/SReferenceViewer.cpp
+++ b/Engine/Plugins/Editor/AssetManagerEditor/Source/AssetManagerEditor/Private/ReferenceViewer/SReferenceViewer.cpp
@@ -137,7 +137,7 @@ void SReferenceViewer::Construct(const FArguments& InArgs)
 
 	SGraphEditor::FGraphEditorEvents GraphEvents;
 	GraphEvents.OnNodeDoubleClicked = FSingleNodeEvent::CreateSP(this, &SReferenceViewer::OnNodeDoubleClicked);
-	GraphEvents.OnCreateActionMenu = SGraphEditor::FOnCreateActionMenu::CreateSP(this, &SReferenceViewer::OnCreateGraphActionMenu);
+	GraphEvents.OnCreateActionMenuAtLocation = SGraphEditor::FOnCreateActionMenuAtLocation::CreateSP(this, &SReferenceViewer::OnCreateGraphActionMenu);
 
 	// Create the graph editor
 	GraphEditorPtr = SNew(SGraphEditor)
@@ -839,7 +839,7 @@ void SReferenceViewer::RebuildGraph()
 	}
 }
 
-FActionMenuContent SReferenceViewer::OnCreateGraphActionMenu(UEdGraph* InGraph, const FVector2D& InNodePosition, const TArray<UEdGraphPin*>& InDraggedPins, bool bAutoExpand, SGraphEditor::FActionMenuClosed InOnMenuClosed)
+FActionMenuContent SReferenceViewer::OnCreateGraphActionMenu(UEdGraph* InGraph, const FVector2f& InNodePosition, const TArray<UEdGraphPin*>& InDraggedPins, bool bAutoExpand, SGraphEditor::FActionMenuClosed InOnMenuClosed)
 {
 	// no context menu when not over a node
 	return FActionMenuContent();
diff --git a/Engine/Plugins/Editor/AssetManagerEditor/Source/AssetManagerEditor/Private/ReferenceViewer/SReferenceViewer.h b/Engine/Plugins/Editor/AssetManagerEditor/Source/AssetManagerEditor/Private/ReferenceViewer/SReferenceViewer.h
--- a/Engine/Plugins/Editor/AssetManagerEditor/Source/AssetManagerEditor/Private/ReferenceViewer/SReferenceViewer.h
+++ b/Engine/Plugins/Editor/AssetManagerEditor/Source/AssetManagerEditor/Private/ReferenceViewer/SReferenceViewer.h
@@ -68,7 +68,7 @@ private:
 	void RebuildGraph();
 
 	/** Called to create context menu when right-clicking on graph */
-	FActionMenuContent OnCreateGraphActionMenu(UEdGraph* InGraph, const FVector2D& InNodePosition, const TArray<UEdGraphPin*>& InDraggedPins, bool bAutoExpand, SGraphEditor::FActionMenuClosed InOnMenuClosed);
+	FActionMenuContent OnCreateGraphActionMenu(UEdGraph* InGraph, const FVector2f& InNodePosition, const TArray<UEdGraphPin*>& InDraggedPins, bool bAutoExpand, SGraphEditor::FActionMenuClosed InOnMenuClosed);
 
 	/** Called when a node is double clicked */
 	void OnNodeDoubleClicked(class UEdGraphNode* Node);
diff --git a/Engine/Plugins/Enterprise/DataprepEditor/Source/DataprepEditor/Private/Widgets/DataprepGraph/SDataprepGraphEditor.cpp b/Engine/Plugins/Enterprise/DataprepEditor/Source/DataprepEditor/Private/Widgets/DataprepGraph/SDataprepGraphEditor.cpp
--- a/Engine/Plugins/Enterprise/DataprepEditor/Source/DataprepEditor/Private/Widgets/DataprepGraph/SDataprepGraphEditor.cpp
+++ b/Engine/Plugins/Enterprise/DataprepEditor/Source/DataprepEditor/Private/Widgets/DataprepGraph/SDataprepGraphEditor.cpp
@@ -141,7 +141,7 @@ void SDataprepGraphEditor::Construct(const FArguments& InArgs, UDataprepAsset* I
 
 	SGraphEditor::FGraphEditorEvents Events;
 	Events.OnCreateNodeOrPinMenu = SGraphEditor::FOnCreateNodeOrPinMenu::CreateSP(this, &SDataprepGraphEditor::OnCreateNodeOrPinMenu);
-	Events.OnCreateActionMenu = SGraphEditor::FOnCreateActionMenu::CreateSP(this, &SDataprepGraphEditor::OnCreateActionMenu);
+	Events.OnCreateActionMenuAtLocation = SGraphEditor::FOnCreateActionMenuAtLocation::CreateSP(this, &SDataprepGraphEditor::OnCreateActionMenu);
 	Events.OnVerifyTextCommit = FOnNodeVerifyTextCommit::CreateSP(this, &SDataprepGraphEditor::OnNodeVerifyTitleCommit);
 	Events.OnTextCommitted = FOnNodeTextCommitted::CreateSP(this, &SDataprepGraphEditor::OnNodeTitleCommitted);
 
@@ -834,6 +834,11 @@ void SDataprepGraphEditor::CreateFilterFromSelection(UDataprepActionAsset* InTar
 }
 
 FActionMenuContent SDataprepGraphEditor::OnCreateActionMenu(UEdGraph* InGraph, const FVector2D& InNodePosition, const TArray<UEdGraphPin*>& InDraggedPins, bool bAutoExpand, SGraphEditor::FActionMenuClosed InOnMenuClosed)
+{
+	return OnCreateActionMenu(InGraph, UE::Slate::CastToVector2f(InNodePosition), InDraggedPins, bAutoExpand, InOnMenuClosed);
+}
+
+FActionMenuContent SDataprepGraphEditor::OnCreateActionMenu(UEdGraph* InGraph, const FVector2f& InNodePosition, const TArray<UEdGraphPin*>& InDraggedPins, bool bAutoExpand, SGraphEditor::FActionMenuClosed InOnMenuClosed)
 {
 	// bAutoExpand is voluntary ignored for now
 	TUniquePtr<IDataprepMenuActionCollector> ActionCollector = MakeUnique<FDataprepAllMenuActionCollector>();
@@ -842,7 +847,7 @@ FActionMenuContent SDataprepGraphEditor::OnCreateActionMenu(UEdGraph* InGraph, c
 		SNew (SDataprepActionMenu, MoveTemp(ActionCollector))
 		.TransactionText(LOCTEXT("AddingANewActionNode","Add a Action Node"))
 		.GraphObj(InGraph)
-		.NewNodePosition(InNodePosition)
+		.NewNodePosition(FDeprecateSlateVector2D(InNodePosition))
 		.DraggedFromPins(InDraggedPins)
 		.OnClosedCallback(InOnMenuClosed)
 		.OnCollectCustomActions(SDataprepActionMenu::FOnCollectCustomActions::CreateSP(this, &SDataprepGraphEditor::OnCollectCustomActions));
diff --git a/Engine/Plugins/Enterprise/DataprepEditor/Source/DataprepEditor/Private/Widgets/DataprepGraph/SDataprepGraphEditor.h b/Engine/Plugins/Enterprise/DataprepEditor/Source/DataprepEditor/Private/Widgets/DataprepGraph/SDataprepGraphEditor.h
--- a/Engine/Plugins/Enterprise/DataprepEditor/Source/DataprepEditor/Private/Widgets/DataprepGraph/SDataprepGraphEditor.h
+++ b/Engine/Plugins/Enterprise/DataprepEditor/Source/DataprepEditor/Private/Widgets/DataprepGraph/SDataprepGraphEditor.h
@@ -120,7 +120,9 @@ public:
 	void OnNodeTitleCommitted(const FText& NewText, ETextCommit::Type CommitInfo, UEdGraphNode* NodeBeingChanged);
 
 	/** Callback to create contextual menu on graph panel */
+	UE_DEPRECATED(5.6, "Please use the version of the function accepting FVector2f.")
 	FActionMenuContent OnCreateActionMenu(UEdGraph* InGraph, const FVector2D& InNodePosition, const TArray<UEdGraphPin*>& InDraggedPins, bool bAutoExpand, SGraphEditor::FActionMenuClosed InOnMenuClosed);
+	FActionMenuContent OnCreateActionMenu(UEdGraph* InGraph, const FVector2f& InNodePosition, const TArray<UEdGraphPin*>& InDraggedPins, bool bAutoExpand, SGraphEditor::FActionMenuClosed InOnMenuClosed);
 
 	/** Callback to create contextual menu for graph nodes in graph panel */
 	FActionMenuContent OnCreateNodeOrPinMenu(UEdGraph* CurrentGraph, const UEdGraphNode* InGraphNode, const UEdGraphPin* InGraphPin, FMenuBuilder* MenuBuilder, bool bIsDebugging);
diff --git a/Engine/Plugins/Experimental/RenderGrid/Source/RenderGridEditor/Private/Toolkit/RenderGridEditor.cpp b/Engine/Plugins/Experimental/RenderGrid/Source/RenderGridEditor/Private/Toolkit/RenderGridEditor.cpp
--- a/Engine/Plugins/Experimental/RenderGrid/Source/RenderGridEditor/Private/Toolkit/RenderGridEditor.cpp
+++ b/Engine/Plugins/Experimental/RenderGrid/Source/RenderGridEditor/Private/Toolkit/RenderGridEditor.cpp
@@ -194,7 +194,7 @@ void UE::RenderGrid::Private::FRenderGridEditor::SetupGraphEditorEvents(UEdGraph
 {
 	FBlueprintEditor::SetupGraphEditorEvents(InGraph, InEvents);
 
-	InEvents.OnCreateActionMenu = SGraphEditor::FOnCreateActionMenu::CreateSP(this, &FRenderGridEditor::HandleCreateGraphActionMenu);
+	InEvents.OnCreateActionMenuAtLocation = SGraphEditor::FOnCreateActionMenuAtLocation::CreateSP(this, &FRenderGridEditor::HandleCreateGraphActionMenu);
 }
 
 void UE::RenderGrid::Private::FRenderGridEditor::RegisterApplicationModes(const TArray<UBlueprint*>& InBlueprints, bool bShouldOpenInDefaultsMode, bool bNewlyCreated)
@@ -713,7 +713,7 @@ void UE::RenderGrid::Private::FRenderGridEditor::UpdateInstance(UBlueprint* InBl
 	}
 }
 
-FActionMenuContent UE::RenderGrid::Private::FRenderGridEditor::HandleCreateGraphActionMenu(UEdGraph* InGraph, const FVector2D& InNodePosition, const TArray<UEdGraphPin*>& InDraggedPins, bool bAutoExpand, SGraphEditor::FActionMenuClosed InOnMenuClosed)
+FActionMenuContent UE::RenderGrid::Private::FRenderGridEditor::HandleCreateGraphActionMenu(UEdGraph* InGraph, const FVector2f& InNodePosition, const TArray<UEdGraphPin*>& InDraggedPins, bool bAutoExpand, SGraphEditor::FActionMenuClosed InOnMenuClosed)
 {
 	return OnCreateGraphActionMenu(InGraph, InNodePosition, InDraggedPins, bAutoExpand, InOnMenuClosed);
 }
diff --git a/Engine/Plugins/Experimental/RenderGrid/Source/RenderGridEditor/Private/Toolkit/RenderGridEditor.h b/Engine/Plugins/Experimental/RenderGrid/Source/RenderGridEditor/Private/Toolkit/RenderGridEditor.h
--- a/Engine/Plugins/Experimental/RenderGrid/Source/RenderGridEditor/Private/Toolkit/RenderGridEditor.h
+++ b/Engine/Plugins/Experimental/RenderGrid/Source/RenderGridEditor/Private/Toolkit/RenderGridEditor.h
@@ -172,7 +172,7 @@ namespace UE::RenderGrid::Private
 		void UpdateInstance(UBlueprint* InBlueprint, bool bInForceFullUpdate);
 
 		/** Wraps the normal blueprint editor's action menu creation callback. */
-		FActionMenuContent HandleCreateGraphActionMenu(UEdGraph* InGraph, const FVector2D& InNodePosition, const TArray<UEdGraphPin*>& InDraggedPins, bool bAutoExpand, SGraphEditor::FActionMenuClosed InOnMenuClosed);
+		FActionMenuContent HandleCreateGraphActionMenu(UEdGraph* InGraph, const FVector2f& InNodePosition, const TArray<UEdGraphPin*>& InDraggedPins, bool bAutoExpand, SGraphEditor::FActionMenuClosed InOnMenuClosed);
 
 	private:
 		/** The toolbar builder that is used to customize the toolbar of this editor. */
diff --git a/Engine/Plugins/Experimental/Workspace/Source/WorkspaceEditor/Private/SGraphDocument.cpp b/Engine/Plugins/Experimental/Workspace/Source/WorkspaceEditor/Private/SGraphDocument.cpp
--- a/Engine/Plugins/Experimental/Workspace/Source/WorkspaceEditor/Private/SGraphDocument.cpp
+++ b/Engine/Plugins/Experimental/Workspace/Source/WorkspaceEditor/Private/SGraphDocument.cpp
@@ -30,11 +30,11 @@ void SGraphDocument::Construct(const FArguments& InArgs, TSharedRef<FWorkspaceEd
 	OnDuplicateSelectedNodes = InArgs._OnDuplicateSelectedNodes;
 
 	SGraphEditor::FGraphEditorEvents Events;
-	Events.OnCreateActionMenu = SGraphEditor::FOnCreateActionMenu::CreateLambda([this, OnCreateActionMenu = InArgs._OnCreateActionMenu](UEdGraph* InGraph, const FVector2D& InNodePosition, const TArray<UEdGraphPin*>& InDraggedPins, bool bInAutoExpand, SGraphEditor::FActionMenuClosed InOnMenuClosed)
+	Events.OnCreateActionMenuAtLocation = SGraphEditor::FOnCreateActionMenuAtLocation::CreateLambda([this, OnCreateActionMenu = InArgs._OnCreateActionMenu](UEdGraph* InGraph, const FVector2f& InNodePosition, const TArray<UEdGraphPin*>& InDraggedPins, bool bInAutoExpand, SGraphEditor::FActionMenuClosed InOnMenuClosed)
 	{
 		if(OnCreateActionMenu.IsBound())
 		{
-			return OnCreateActionMenu.Execute(FWorkspaceEditorContext(HostingAppPtr.Pin().ToSharedRef(), Document), InGraph, InNodePosition, InDraggedPins, bInAutoExpand, InOnMenuClosed);
+			return OnCreateActionMenu.Execute(FWorkspaceEditorContext(HostingAppPtr.Pin().ToSharedRef(), Document), InGraph, FDeprecateSlateVector2D(InNodePosition), InDraggedPins, bInAutoExpand, InOnMenuClosed);
 		}
 		return FActionMenuContent();
 	});
diff --git a/Engine/Plugins/FX/Niagara/Source/NiagaraEditor/Private/Widgets/SNiagaraScriptGraph.cpp b/Engine/Plugins/FX/Niagara/Source/NiagaraEditor/Private/Widgets/SNiagaraScriptGraph.cpp
--- a/Engine/Plugins/FX/Niagara/Source/NiagaraEditor/Private/Widgets/SNiagaraScriptGraph.cpp
+++ b/Engine/Plugins/FX/Niagara/Source/NiagaraEditor/Private/Widgets/SNiagaraScriptGraph.cpp
@@ -318,8 +318,8 @@ TSharedRef<SGraphEditor> SNiagaraScriptGraph::ConstructGraphEditor()
 	Events.OnNodeDoubleClicked = FSingleNodeEvent::CreateSP(this, &SNiagaraScriptGraph::OnNodeDoubleClicked);
 	Events.OnTextCommitted = FOnNodeTextCommitted::CreateSP(this, &SNiagaraScriptGraph::OnNodeTitleCommitted);
 	Events.OnVerifyTextCommit = FOnNodeVerifyTextCommit::CreateSP(this, &SNiagaraScriptGraph::OnVerifyNodeTextCommit);
-	Events.OnSpawnNodeByShortcut = SGraphEditor::FOnSpawnNodeByShortcut::CreateSP(this, &SNiagaraScriptGraph::OnSpawnGraphNodeByShortcut);
-	Events.OnCreateActionMenu = SGraphEditor::FOnCreateActionMenu::CreateSP(this, &SNiagaraScriptGraph::OnCreateActionMenu);
+	Events.OnSpawnNodeByShortcutAtLocation = SGraphEditor::FOnSpawnNodeByShortcutAtLocation::CreateSP(this, &SNiagaraScriptGraph::OnSpawnGraphNodeByShortcut);
+	Events.OnCreateActionMenuAtLocation = SGraphEditor::FOnCreateActionMenuAtLocation::CreateSP(this, &SNiagaraScriptGraph::OnCreateActionMenu);
 
 	Commands = MakeShared<FUICommandList>();
 	Commands->Append(ViewModel->GetCommands());
@@ -468,7 +468,7 @@ bool SNiagaraScriptGraph::OnVerifyNodeTextCommit(const FText& NewText, UEdGraphN
 	return bValid;
 }
 
-FReply SNiagaraScriptGraph::OnSpawnGraphNodeByShortcut(FInputChord InChord, const FVector2D& InPosition)
+FReply SNiagaraScriptGraph::OnSpawnGraphNodeByShortcut(FInputChord InChord, const FVector2f& InPosition)
 {
 	if (ViewModel.Get() == nullptr)
 	{
@@ -529,7 +529,7 @@ FReply SNiagaraScriptGraph::OnSpawnGraphNodeByShortcut(FInputChord InChord, cons
 				{
 					FScopedTransaction Transaction(LOCTEXT("AddNode", "Add Node"));
 					TArray<UEdGraphPin*> Pins;
-					NiagaraAction->CreateNode(Graph, Pins, InPosition);
+					NiagaraAction->CreateNode(Graph, Pins, FDeprecateSlateVector2D(InPosition));
 					return FReply::Handled();
 				}					
 			}
@@ -571,7 +571,7 @@ FReply SNiagaraScriptGraph::OnSpawnGraphNodeByShortcut(FInputChord InChord, cons
 	*/
 }
 
-FActionMenuContent SNiagaraScriptGraph::OnCreateActionMenu(UEdGraph* Graph, const FVector2D& Position, const TArray<UEdGraphPin*>& DraggedPins, bool bAutoExpandActionMenu, SGraphEditor::FActionMenuClosed OnClosed)
+FActionMenuContent SNiagaraScriptGraph::OnCreateActionMenu(UEdGraph* Graph, const FVector2f& Position, const TArray<UEdGraphPin*>& DraggedPins, bool bAutoExpandActionMenu, SGraphEditor::FActionMenuClosed OnClosed)
 {
 	FActionMenuContent Content;
 	
diff --git a/Engine/Plugins/FX/Niagara/Source/NiagaraEditor/Public/Widgets/SNiagaraGraphActionMenu.h b/Engine/Plugins/FX/Niagara/Source/NiagaraEditor/Public/Widgets/SNiagaraGraphActionMenu.h
--- a/Engine/Plugins/FX/Niagara/Source/NiagaraEditor/Public/Widgets/SNiagaraGraphActionMenu.h
+++ b/Engine/Plugins/FX/Niagara/Source/NiagaraEditor/Public/Widgets/SNiagaraGraphActionMenu.h
@@ -41,12 +41,12 @@ public:
 	
 	SLATE_BEGIN_ARGS( SNiagaraGraphActionMenu )
 		: _GraphObj( static_cast<UEdGraph*>(NULL) )
-		, _NewNodePosition( FVector2D::ZeroVector )
+		, _NewNodePosition( FVector2f::ZeroVector )
 		, _AutoExpandActionMenu( false )
 		{}
 
 		SLATE_ARGUMENT( UEdGraph*, GraphObj )
-		SLATE_ARGUMENT( FVector2D, NewNodePosition )
+		SLATE_ARGUMENT( FDeprecateSlateVector2D, NewNodePosition )
 		SLATE_ARGUMENT( TArray<UEdGraphPin*>, DraggedFromPins )
 		SLATE_ARGUMENT( SGraphEditor::FActionMenuClosed, OnClosedCallback )
 		SLATE_ARGUMENT( bool, AutoExpandActionMenu )
@@ -62,7 +62,7 @@ protected:
 	UEdGraph* GraphObj;
 	FNiagaraScriptGraphViewModel* ViewModel;
 	TArray<UEdGraphPin*> DraggedFromPins;
-	FVector2D NewNodePosition;
+	FDeprecateSlateVector2D NewNodePosition;
 	bool AutoExpandActionMenu;
 
 	SGraphEditor::FActionMenuClosed OnClosedCallback;
diff --git a/Engine/Plugins/FX/Niagara/Source/NiagaraEditor/Public/Widgets/SNiagaraScriptGraph.h b/Engine/Plugins/FX/Niagara/Source/NiagaraEditor/Public/Widgets/SNiagaraScriptGraph.h
--- a/Engine/Plugins/FX/Niagara/Source/NiagaraEditor/Public/Widgets/SNiagaraScriptGraph.h
+++ b/Engine/Plugins/FX/Niagara/Source/NiagaraEditor/Public/Widgets/SNiagaraScriptGraph.h
@@ -77,9 +77,9 @@ private:
 	void OnNodeTitleCommitted(const FText& NewText, ETextCommit::Type CommitInfo, UEdGraphNode* NodeBeingChanged);
 
 	/** Called whenever a user left-clicks on the graph with a chord pressed.*/
-	FReply OnSpawnGraphNodeByShortcut(FInputChord InChord, const FVector2D& InPosition);
+	FReply OnSpawnGraphNodeByShortcut(FInputChord InChord, const FVector2f& InPosition);
 
-	FActionMenuContent OnCreateActionMenu(UEdGraph* Graph, const FVector2D& Position, const TArray<UEdGraphPin*>& DraggedPins, bool bAutoExpandActionMenu, SGraphEditor::FActionMenuClosed OnClosed);	
+	FActionMenuContent OnCreateActionMenu(UEdGraph* Graph, const FVector2f& Position, const TArray<UEdGraphPin*>& DraggedPins, bool bAutoExpandActionMenu, SGraphEditor::FActionMenuClosed OnClosed);	
 	
 	/** Called whenever a user is trying to edit the name inline of a node and we want to make sure that it is valid.*/
 	bool OnVerifyNodeTextCommit(const FText& NewText, UEdGraphNode* NodeBeingChanged, FText& OutErrorMessage);
diff --git a/Engine/Plugins/FX/Niagara/Source/NiagaraEditorWidgets/Private/SNiagaraOverviewGraph.cpp b/Engine/Plugins/FX/Niagara/Source/NiagaraEditorWidgets/Private/SNiagaraOverviewGraph.cpp
--- a/Engine/Plugins/FX/Niagara/Source/NiagaraEditorWidgets/Private/SNiagaraOverviewGraph.cpp
+++ b/Engine/Plugins/FX/Niagara/Source/NiagaraEditorWidgets/Private/SNiagaraOverviewGraph.cpp
@@ -48,7 +48,7 @@ void SNiagaraOverviewGraph::Construct(const FArguments& InArgs, TSharedRef<FNiag
 
 	SGraphEditor::FGraphEditorEvents Events;
 	Events.OnSelectionChanged = SGraphEditor::FOnSelectionChanged::CreateSP(this, &SNiagaraOverviewGraph::GraphSelectionChanged);
-	Events.OnCreateActionMenu = SGraphEditor::FOnCreateActionMenu::CreateSP(this, &SNiagaraOverviewGraph::OnCreateGraphActionMenu);
+	Events.OnCreateActionMenuAtLocation = SGraphEditor::FOnCreateActionMenuAtLocation::CreateSP(this, &SNiagaraOverviewGraph::OnCreateGraphActionMenu);
 	Events.OnVerifyTextCommit = FOnNodeVerifyTextCommit::CreateSP(this, &SNiagaraOverviewGraph::OnVerifyNodeTitle);
 	Events.OnTextCommitted = FOnNodeTextCommitted::CreateSP(this, &SNiagaraOverviewGraph::OnNodeTitleCommitted);
 
@@ -231,7 +231,7 @@ bool SNiagaraOverviewGraph::CanAddEmitters() const
 	return ViewModel->GetSystemViewModel()->GetEditMode() == ENiagaraSystemViewModelEditMode::SystemAsset ? true : false;
 }
 
-FActionMenuContent SNiagaraOverviewGraph::OnCreateGraphActionMenu(UEdGraph* InGraph, const FVector2D& InNodePosition, const TArray<UEdGraphPin*>& InDraggedPins, bool bAutoExpand, SGraphEditor::FActionMenuClosed InOnMenuClosed)
+FActionMenuContent SNiagaraOverviewGraph::OnCreateGraphActionMenu(UEdGraph* InGraph, const FVector2f& InNodePosition, const TArray<UEdGraphPin*>& InDraggedPins, bool bAutoExpand, SGraphEditor::FActionMenuClosed InOnMenuClosed)
 {
 	if (ViewModel->GetSystemViewModel()->GetEditMode() == ENiagaraSystemViewModelEditMode::SystemAsset)
 	{
diff --git a/Engine/Plugins/FX/Niagara/Source/NiagaraEditorWidgets/Private/SNiagaraOverviewGraph.h b/Engine/Plugins/FX/Niagara/Source/NiagaraEditorWidgets/Private/SNiagaraOverviewGraph.h
--- a/Engine/Plugins/FX/Niagara/Source/NiagaraEditorWidgets/Private/SNiagaraOverviewGraph.h
+++ b/Engine/Plugins/FX/Niagara/Source/NiagaraEditorWidgets/Private/SNiagaraOverviewGraph.h
@@ -33,7 +33,7 @@ private:
 	bool CanAddEmitters() const;
 
 	/** Called to create context menu when right-clicking on graph */
-	FActionMenuContent OnCreateGraphActionMenu(UEdGraph* InGraph, const FVector2D& InNodePosition, const TArray<UEdGraphPin*>& InDraggedPins, bool bAutoExpand, SGraphEditor::FActionMenuClosed InOnMenuClosed);
+	FActionMenuContent OnCreateGraphActionMenu(UEdGraph* InGraph, const FVector2f& InNodePosition, const TArray<UEdGraphPin*>& InDraggedPins, bool bAutoExpand, SGraphEditor::FActionMenuClosed InOnMenuClosed);
 
 	void OnCreateEmptyEmitter();
 	bool CanCreateStatelessEmitter() const;
diff --git a/Engine/Plugins/PCG/Source/PCGEditor/Private/PCGEditor.cpp b/Engine/Plugins/PCG/Source/PCGEditor/Private/PCGEditor.cpp
--- a/Engine/Plugins/PCG/Source/PCGEditor/Private/PCGEditor.cpp
+++ b/Engine/Plugins/PCG/Source/PCGEditor/Private/PCGEditor.cpp
@@ -1578,6 +1578,11 @@ void FPCGEditor::OnJumpToSource()
 }
 
 FReply FPCGEditor::OnSpawnNodeByShortcut(FInputChord InChord, const FVector2D& InPosition, UPCGEditorGraph* InGraph)
+{
+	return OnSpawnNodeByShortcut(InChord, UE::Slate::CastToVector2f(InPosition), InGraph);
+}
+
+FReply FPCGEditor::OnSpawnNodeByShortcut(FInputChord InChord, const FVector2f& InPosition, UPCGEditorGraph* InGraph)
 {
 	const TSharedPtr<FEdGraphSchemaAction> Action = FPCGEditorSpawnNodeCommands::Get().GetGraphActionByChord(InChord);
 	if (Action.IsValid())
@@ -2756,7 +2761,7 @@ TSharedRef<SGraphEditor> FPCGEditor::CreateGraphEditorWidget()
 	InEvents.OnVerifyTextCommit = FOnNodeVerifyTextCommit::CreateSP(this, &FPCGEditor::OnValidateNodeTitle);
 	InEvents.OnTextCommitted = FOnNodeTextCommitted::CreateSP(this, &FPCGEditor::OnNodeTitleCommitted);
 	InEvents.OnNodeDoubleClicked = FSingleNodeEvent::CreateSP(this, &FPCGEditor::OnNodeDoubleClicked);
-	InEvents.OnSpawnNodeByShortcut = SGraphEditor::FOnSpawnNodeByShortcut::CreateSP(this, &FPCGEditor::OnSpawnNodeByShortcut, PCGEditorGraph);
+	InEvents.OnSpawnNodeByShortcutAtLocation = SGraphEditor::FOnSpawnNodeByShortcutAtLocation::CreateSP(this, &FPCGEditor::OnSpawnNodeByShortcut, PCGEditorGraph);
 
 	return SNew(SGraphEditor)
 		.AdditionalCommands(GraphEditorCommands)
diff --git a/Engine/Plugins/PCG/Source/PCGEditor/Private/PCGEditor.h b/Engine/Plugins/PCG/Source/PCGEditor/Private/PCGEditor.h
--- a/Engine/Plugins/PCG/Source/PCGEditor/Private/PCGEditor.h
+++ b/Engine/Plugins/PCG/Source/PCGEditor/Private/PCGEditor.h
@@ -141,7 +141,9 @@ public:
 	 *
 	 * @return	FReply	Whether chord was handled.
 	 */
+	UE_DEPRECATED(5.6, "Please use the version of the function accepting FVector2f.")
 	FReply OnSpawnNodeByShortcut(FInputChord InChord, const FVector2D& InPosition, UPCGEditorGraph* InGraph);
+	FReply OnSpawnNodeByShortcut(FInputChord InChord, const FVector2f& InPosition, UPCGEditorGraph* InGraph);
 
 	FOnInspectedStackChanged OnInspectedStackChangedDelegate;
 
diff --git a/Engine/Plugins/Runtime/Metasound/Source/MetasoundEditor/Private/MetasoundEditor.cpp b/Engine/Plugins/Runtime/Metasound/Source/MetasoundEditor/Private/MetasoundEditor.cpp
--- a/Engine/Plugins/Runtime/Metasound/Source/MetasoundEditor/Private/MetasoundEditor.cpp
+++ b/Engine/Plugins/Runtime/Metasound/Source/MetasoundEditor/Private/MetasoundEditor.cpp
@@ -3006,7 +3006,7 @@ namespace Metasound
 			}
 
 			SGraphEditor::FGraphEditorEvents GraphEvents;
-			GraphEvents.OnCreateActionMenu = SGraphEditor::FOnCreateActionMenu::CreateSP(this, &FEditor::OnCreateGraphActionMenu);
+			GraphEvents.OnCreateActionMenuAtLocation = SGraphEditor::FOnCreateActionMenuAtLocation::CreateSP(this, &FEditor::OnCreateGraphActionMenu);
 			GraphEvents.OnNodeDoubleClicked = FSingleNodeEvent::CreateSP(this, &FEditor::ExecuteNode);
 			GraphEvents.OnSelectionChanged = SGraphEditor::FOnSelectionChanged::CreateSP(this, &FEditor::OnSelectedNodesChanged);
 			GraphEvents.OnTextCommitted = FOnNodeTextCommitted::CreateSP(this, &FEditor::OnNodeTitleCommitted);
@@ -4430,7 +4430,7 @@ PRAGMA_ENABLE_DEPRECATION_WARNINGS
 			return false;
 		}
 
-		FActionMenuContent FEditor::OnCreateGraphActionMenu(UEdGraph* InGraph, const FVector2D& InNodePosition, const TArray<UEdGraphPin*>& InDraggedPins, bool bAutoExpand, SGraphEditor::FActionMenuClosed InOnMenuClosed)
+		FActionMenuContent FEditor::OnCreateGraphActionMenu(UEdGraph* InGraph, const FVector2f& InNodePosition, const TArray<UEdGraphPin*>& InDraggedPins, bool bAutoExpand, SGraphEditor::FActionMenuClosed InOnMenuClosed)
 		{
 			TSharedRef<SMetasoundActionMenu> ActionMenu = SNew(SMetasoundActionMenu)
 				.AutoExpandActionMenu(bAutoExpand)
diff --git a/Engine/Plugins/Runtime/Metasound/Source/MetasoundEditor/Private/MetasoundEditor.h b/Engine/Plugins/Runtime/Metasound/Source/MetasoundEditor/Private/MetasoundEditor.h
--- a/Engine/Plugins/Runtime/Metasound/Source/MetasoundEditor/Private/MetasoundEditor.h
+++ b/Engine/Plugins/Runtime/Metasound/Source/MetasoundEditor/Private/MetasoundEditor.h
@@ -265,7 +265,7 @@ namespace Metasound::Editor
 		bool HandleActionMatchesName(FEdGraphSchemaAction* InAction, const FName& InName) const;
 		FReply OnActionDragged(const TArray<TSharedPtr<FEdGraphSchemaAction>>& InActions, const FPointerEvent& MouseEvent);
 		void OnMemberActionDoubleClicked(const TArray< TSharedPtr<FEdGraphSchemaAction> >& InActions);
-		FActionMenuContent OnCreateGraphActionMenu(UEdGraph* InGraph, const FVector2D& InNodePosition, const TArray<UEdGraphPin*>& InDraggedPins, bool bAutoExpand, SGraphEditor::FActionMenuClosed InOnMenuClosed);
+		FActionMenuContent OnCreateGraphActionMenu(UEdGraph* InGraph, const FVector2f& InNodePosition, const TArray<UEdGraphPin*>& InDraggedPins, bool bAutoExpand, SGraphEditor::FActionMenuClosed InOnMenuClosed);
 		void OnActionSelected(const TArray<TSharedPtr<FEdGraphSchemaAction>>& InActions, ESelectInfo::Type InSelectionType);
 		FReply OnAddButtonClickedOnSection(int32 InSectionID);
 		TSharedRef<SWidget> OnGetMenuSectionWidget(TSharedRef<SWidget> RowWidget, int32 InSectionID);
diff --git a/Engine/Plugins/Runtime/Metasound/Source/MetasoundEditor/Private/SMetasoundActionMenu.h b/Engine/Plugins/Runtime/Metasound/Source/MetasoundEditor/Private/SMetasoundActionMenu.h
--- a/Engine/Plugins/Runtime/Metasound/Source/MetasoundEditor/Private/SMetasoundActionMenu.h
+++ b/Engine/Plugins/Runtime/Metasound/Source/MetasoundEditor/Private/SMetasoundActionMenu.h
@@ -50,12 +50,12 @@ namespace Metasound
 
 			SLATE_BEGIN_ARGS(SMetasoundActionMenu)
 				: _Graph(static_cast<UEdGraph*>(nullptr))
-				, _NewNodePosition(FVector2D::ZeroVector)
+				, _NewNodePosition(FVector2f::ZeroVector)
 				, _AutoExpandActionMenu(false)
 			{
 			}
 				SLATE_ARGUMENT(UEdGraph*, Graph)
-				SLATE_ARGUMENT(FVector2D, NewNodePosition)
+				SLATE_ARGUMENT(FDeprecateSlateVector2D, NewNodePosition)
 				SLATE_ARGUMENT(TArray<UEdGraphPin*>, DraggedFromPins)
 				SLATE_ARGUMENT(SGraphEditor::FActionMenuClosed, OnClosedCallback)
 				SLATE_ARGUMENT(bool, AutoExpandActionMenu)
diff --git a/Engine/Plugins/Runtime/RigVM/Source/RigVMEditor/Private/Editor/RigVMEditor.cpp b/Engine/Plugins/Runtime/RigVM/Source/RigVMEditor/Private/Editor/RigVMEditor.cpp
--- a/Engine/Plugins/Runtime/RigVM/Source/RigVMEditor/Private/Editor/RigVMEditor.cpp
+++ b/Engine/Plugins/Runtime/RigVM/Source/RigVMEditor/Private/Editor/RigVMEditor.cpp
@@ -818,7 +818,7 @@ bool FRigVMEditorBase::CanDuplicateNodes() const
 	return CanCopyNodes() && IsEditable(GetFocusedGraph());
 }
 
-FReply FRigVMEditorBase::OnSpawnGraphNodeByShortcut(FInputChord InChord, const FVector2D& InPosition,
+FReply FRigVMEditorBase::OnSpawnGraphNodeByShortcut(FInputChord InChord, const FVector2f& InPosition,
                                                      UEdGraph* InGraph)
 {
 	if(!InChord.HasAnyModifierKeys())
@@ -829,7 +829,7 @@ FReply FRigVMEditorBase::OnSpawnGraphNodeByShortcut(FInputChord InChord, const F
 			{
 				if(InChord.Key == EKeys::B)
 				{
-					Controller->AddUnitNode(FRigVMFunction_ControlFlowBranch::StaticStruct(), FRigVMStruct::ExecuteName, InPosition, FString(), true, true);
+					Controller->AddUnitNode(FRigVMFunction_ControlFlowBranch::StaticStruct(), FRigVMStruct::ExecuteName, FDeprecateSlateVector2D(InPosition), FString(), true, true);
 				}
 			}
 		}
@@ -1328,11 +1328,11 @@ void FRigVMEditorBase::SetupGraphEditorEvents(UEdGraph* InGraph, SGraphEditor::F
 
 	SetupGraphEditorEventsImpl(InGraph, InEvents);
 
-	InEvents.OnCreateActionMenu = SGraphEditor::FOnCreateActionMenu::CreateSP(StaticCastSharedRef<FRigVMEditorBase>(SharedRef()), &FRigVMEditorBase::HandleCreateGraphActionMenu);
+	InEvents.OnCreateActionMenuAtLocation = SGraphEditor::FOnCreateActionMenuAtLocation::CreateSP(StaticCastSharedRef<FRigVMEditorBase>(SharedRef()), &FRigVMEditorBase::HandleCreateGraphActionMenu);
 	InEvents.OnTextCommitted = FOnNodeTextCommitted::CreateSP(StaticCastSharedRef<FRigVMEditorBase>(SharedRef()), &FRigVMEditorBase::OnNodeTitleCommitted);
 }
 
-FActionMenuContent FRigVMEditorBase::HandleCreateGraphActionMenu(UEdGraph* InGraph, const FVector2D& InNodePosition, const TArray<UEdGraphPin*>& InDraggedPins, bool bAutoExpand, SGraphEditor::FActionMenuClosed InOnMenuClosed)
+FActionMenuContent FRigVMEditorBase::HandleCreateGraphActionMenu(UEdGraph* InGraph, const FVector2f& InNodePosition, const TArray<UEdGraphPin*>& InDraggedPins, bool bAutoExpand, SGraphEditor::FActionMenuClosed InOnMenuClosed)
 {
 	return OnCreateGraphActionMenu(InGraph, InNodePosition, InDraggedPins, bAutoExpand, InOnMenuClosed);
 }
@@ -4000,7 +4000,7 @@ void FRigVMEditorBase::SwapAssetReferences()
 	SwapFunctionsDialog->ShowNormal();
 }
 
-void FRigVMEditorBase::OnGraphNodeDropToPerform(TSharedPtr<FDragDropOperation> InDragDropOp, UEdGraph* InGraph, const FVector2D& InNodePosition, const FVector2D& InScreenPosition)
+void FRigVMEditorBase::OnGraphNodeDropToPerform(TSharedPtr<FDragDropOperation> InDragDropOp, UEdGraph* InGraph, const FVector2f& InNodePosition, const FVector2f& InScreenPosition)
 {
 	DECLARE_SCOPE_HIERARCHICAL_COUNTER_FUNC()
 
@@ -4019,7 +4019,7 @@ void FRigVMEditorBase::OnGraphNodeDropToPerform(TSharedPtr<FDragDropOperation> I
 					{
 						if (URigVMLibraryNode* LibraryNode = Library->FindFunction(*ExplorerOp->GetElement().Name))
 						{
-							Schema->RequestFunctionDropOnPanel(InGraph, LibraryNode->GetFunctionIdentifier(), InNodePosition, InScreenPosition);
+							Schema->RequestFunctionDropOnPanel(InGraph, LibraryNode->GetFunctionIdentifier(), FDeprecateSlateVector2D(InNodePosition), FDeprecateSlateVector2D(InScreenPosition));
 						}
 					}
 				}
@@ -4029,7 +4029,7 @@ void FRigVMEditorBase::OnGraphNodeDropToPerform(TSharedPtr<FDragDropOperation> I
 				if (URigVMBlueprint* TargetRigBlueprint = Cast<URigVMBlueprint>(FBlueprintEditorUtils::FindBlueprintForGraph(InGraph)))
 				{
 					FProperty* Property = TargetRigBlueprint->SkeletonGeneratedClass->FindPropertyByName(*ExplorerOp->GetElement().Name);
-					Schema->RequestVariableDropOnPanel(InGraph, Property, UE::Slate::CastToVector2f(InNodePosition), UE::Slate::CastToVector2f(InScreenPosition));
+					Schema->RequestVariableDropOnPanel(InGraph, Property, InNodePosition, InScreenPosition);
 				}
 			}
 			else if (ExplorerOp->GetElement().Type == ERigVMExplorerElementType::LocalVariable)
@@ -4056,7 +4056,7 @@ void FRigVMEditorBase::OnGraphNodeDropToPerform(TSharedPtr<FDragDropOperation> I
 										FUIAction(
 										FExecuteAction::CreateLambda([Controller, LocalVariable, InNodePosition]()
 										{
-											Controller->AddVariableNode(LocalVariable.Name, LocalVariable.CPPType, LocalVariable.CPPTypeObject, true, LocalVariable.DefaultValue, InNodePosition, FString(), true, true);
+											Controller->AddVariableNode(LocalVariable.Name, LocalVariable.CPPType, LocalVariable.CPPTypeObject, true, LocalVariable.DefaultValue, FDeprecateSlateVector2D(InNodePosition), FString(), true, true);
 										}),
 										FCanExecuteAction()));
 
@@ -4067,7 +4067,7 @@ void FRigVMEditorBase::OnGraphNodeDropToPerform(TSharedPtr<FDragDropOperation> I
 										FUIAction(
 										FExecuteAction::CreateLambda([Controller, LocalVariable, InNodePosition]()
 										{
-											Controller->AddVariableNode(LocalVariable.Name, LocalVariable.CPPType, LocalVariable.CPPTypeObject, false, LocalVariable.DefaultValue, InNodePosition, FString(), true, true);
+											Controller->AddVariableNode(LocalVariable.Name, LocalVariable.CPPType, LocalVariable.CPPTypeObject, false, LocalVariable.DefaultValue, FDeprecateSlateVector2D(InNodePosition), FString(), true, true);
 										}),
 										FCanExecuteAction()));
 
diff --git a/Engine/Plugins/Runtime/RigVM/Source/RigVMEditor/Private/Editor/RigVMNewEditor.cpp b/Engine/Plugins/Runtime/RigVM/Source/RigVMEditor/Private/Editor/RigVMNewEditor.cpp
--- a/Engine/Plugins/Runtime/RigVM/Source/RigVMEditor/Private/Editor/RigVMNewEditor.cpp
+++ b/Engine/Plugins/Runtime/RigVM/Source/RigVMEditor/Private/Editor/RigVMNewEditor.cpp
@@ -1595,11 +1595,11 @@ void FRigVMNewEditor::SetupGraphEditorEventsImpl(UEdGraph* InGraph, SGraphEditor
 	// Custom menu for K2 schemas
 	if(InGraph->Schema != nullptr && InGraph->Schema->IsChildOf(UEdGraphSchema_K2::StaticClass()))
 	{
-		InEvents.OnCreateActionMenu = SGraphEditor::FOnCreateActionMenu::CreateSP(this, &FRigVMNewEditor::OnCreateGraphActionMenu);
+		InEvents.OnCreateActionMenuAtLocation = SGraphEditor::FOnCreateActionMenuAtLocation::CreateSP(this, &FRigVMNewEditor::OnCreateGraphActionMenu);
 	}
 }
 
-FActionMenuContent FRigVMNewEditor::OnCreateGraphActionMenu(UEdGraph* InGraph, const FVector2D& InNodePosition,
+FActionMenuContent FRigVMNewEditor::OnCreateGraphActionMenu(UEdGraph* InGraph, const FVector2f& InNodePosition,
 	const TArray<UEdGraphPin*>& InDraggedPins, bool bAutoExpand, SGraphEditor::FActionMenuClosed InOnMenuClosed)
 {
 	// TODO sara-s: Replicate what blueprint editor is doing
diff --git a/Engine/Plugins/Runtime/RigVM/Source/RigVMEditor/Public/Editor/RigVMEditor.h b/Engine/Plugins/Runtime/RigVM/Source/RigVMEditor/Public/Editor/RigVMEditor.h
--- a/Engine/Plugins/Runtime/RigVM/Source/RigVMEditor/Public/Editor/RigVMEditor.h
+++ b/Engine/Plugins/Runtime/RigVM/Source/RigVMEditor/Public/Editor/RigVMEditor.h
@@ -137,7 +137,7 @@ public:
 	virtual bool InEditingMode() const = 0;
 	virtual bool IsEditable(UEdGraph* InGraph) const = 0;
 	
-	virtual void OnGraphNodeDropToPerform(TSharedPtr<FDragDropOperation> InDragDropOp, UEdGraph* InGraph, const FVector2D& InNodePosition, const FVector2D& InScreenPosition) = 0;
+	virtual void OnGraphNodeDropToPerform(TSharedPtr<FDragDropOperation> InDragDropOp, UEdGraph* InGraph, const FVector2f& InNodePosition, const FVector2f& InScreenPosition) = 0;
 	virtual TSharedRef<FUICommandList> GetToolkitCommands() = 0;
 
 	virtual bool SelectLocalVariable(const UEdGraph* Graph, const FName& VariableName) = 0;
@@ -226,7 +226,7 @@ protected:
 	virtual void OnBlueprintChangedInnerImpl(UBlueprint* InBlueprint, bool bIsJustBeingCompiled = false) = 0;
 	virtual void RefreshEditorsImpl(ERefreshRigVMEditorReason::Type Reason = ERefreshRigVMEditorReason::UnknownReason) = 0;
 	virtual void SetupGraphEditorEventsImpl(UEdGraph* InGraph, SGraphEditor::FGraphEditorEvents& InEvents) = 0;
-	virtual FActionMenuContent OnCreateGraphActionMenu(UEdGraph* InGraph, const FVector2D& InNodePosition, const TArray<UEdGraphPin*>& InDraggedPins, bool bAutoExpand, SGraphEditor::FActionMenuClosed InOnMenuClosed) = 0;
+	virtual FActionMenuContent OnCreateGraphActionMenu(UEdGraph* InGraph, const FVector2f& InNodePosition, const TArray<UEdGraphPin*>& InDraggedPins, bool bAutoExpand, SGraphEditor::FActionMenuClosed InOnMenuClosed) = 0;
 	virtual void AddReferencedObjectsImpl(FReferenceCollector& Collector) = 0;
 	virtual FGraphAppearanceInfo GetGraphAppearanceImpl(class UEdGraph* InGraph) const = 0;
 	virtual void NotifyPreChangeImpl( FProperty* PropertyAboutToChange ) = 0;
@@ -289,7 +289,7 @@ public:
 	FText GetGraphDecorationString(UEdGraph* InGraph) const;
 	virtual bool IsEditable(UEdGraph* InGraph) const override;
 
-	virtual void OnGraphNodeDropToPerform(TSharedPtr<FDragDropOperation> InDragDropOp, UEdGraph* InGraph, const FVector2D& InNodePosition, const FVector2D& InScreenPosition) override;
+	virtual void OnGraphNodeDropToPerform(TSharedPtr<FDragDropOperation> InDragDropOp, UEdGraph* InGraph, const FVector2f& InNodePosition, const FVector2f& InScreenPosition) override;
 
 	/**
 	 * Util for finding a glyph and color for a variable.
@@ -355,7 +355,7 @@ protected:
 
 	bool IsNativeParentClassCodeLinkEnabled() const { return false; }
 	bool ReparentBlueprint_IsVisible() const { return false; }
-	FReply OnSpawnGraphNodeByShortcut(FInputChord InChord, const FVector2D& InPosition, UEdGraph* InGraph);
+	FReply OnSpawnGraphNodeByShortcut(FInputChord InChord, const FVector2f& InPosition, UEdGraph* InGraph);
 	bool ShouldLoadBPLibrariesFromAssetRegistry() { return false; }
 	bool JumpToHyperlinkImpl(const UObject* ObjectReference, bool bRequestRename = false);
 	bool ShouldOpenGraphByDefault() const { return true; }
@@ -475,7 +475,7 @@ protected:
 	virtual void GenerateEventQueueMenuContent(FMenuBuilder& MenuBuilder);
 
 	/** Wraps the normal blueprint editor's action menu creation callback */
-	FActionMenuContent HandleCreateGraphActionMenu(UEdGraph* InGraph, const FVector2D& InNodePosition, const TArray<UEdGraphPin*>& InDraggedPins, bool bAutoExpand, SGraphEditor::FActionMenuClosed InOnMenuClosed);
+	FActionMenuContent HandleCreateGraphActionMenu(UEdGraph* InGraph, const FVector2f& InNodePosition, const TArray<UEdGraphPin*>& InDraggedPins, bool bAutoExpand, SGraphEditor::FActionMenuClosed InOnMenuClosed);
 	void OnNodeTitleCommitted(const FText& NewText, ETextCommit::Type CommitInfo, UEdGraphNode* NodeBeingChanged);
 
 	/** Undo Action**/
diff --git a/Engine/Plugins/Runtime/RigVM/Source/RigVMEditor/Public/Editor/RigVMLegacyEditor.h b/Engine/Plugins/Runtime/RigVM/Source/RigVMEditor/Public/Editor/RigVMLegacyEditor.h
--- a/Engine/Plugins/Runtime/RigVM/Source/RigVMEditor/Public/Editor/RigVMLegacyEditor.h
+++ b/Engine/Plugins/Runtime/RigVM/Source/RigVMEditor/Public/Editor/RigVMLegacyEditor.h
@@ -90,7 +90,7 @@ protected:
 	virtual void RefreshEditorsImpl(ERefreshRigVMEditorReason::Type Reason) override;
 	virtual void SetupGraphEditorEvents(UEdGraph* InGraph, SGraphEditor::FGraphEditorEvents& InEvents) override { FRigVMEditorBase::SetupGraphEditorEvents(InGraph, InEvents); }
 	virtual void SetupGraphEditorEventsImpl(UEdGraph* InGraph, SGraphEditor::FGraphEditorEvents& InEvents) override { FBlueprintEditor::SetupGraphEditorEvents(InGraph, InEvents); }
-	virtual FActionMenuContent OnCreateGraphActionMenu(UEdGraph* InGraph, const FVector2D& InNodePosition, const TArray<UEdGraphPin*>& InDraggedPins, bool bAutoExpand, SGraphEditor::FActionMenuClosed InOnMenuClosed) override { return FBlueprintEditor::OnCreateGraphActionMenu(InGraph, InNodePosition, InDraggedPins, bAutoExpand, InOnMenuClosed); }
+	virtual FActionMenuContent OnCreateGraphActionMenu(UEdGraph* InGraph, const FVector2f& InNodePosition, const TArray<UEdGraphPin*>& InDraggedPins, bool bAutoExpand, SGraphEditor::FActionMenuClosed InOnMenuClosed) override { return FBlueprintEditor::OnCreateGraphActionMenu(InGraph, InNodePosition, InDraggedPins, bAutoExpand, InOnMenuClosed); }
 	virtual void AddReferencedObjects(FReferenceCollector& Collector) override { FRigVMEditorBase::AddReferencedObjects(Collector); }
 	virtual void AddReferencedObjectsImpl(FReferenceCollector& Collector) override { FBlueprintEditor::AddReferencedObjects(Collector); }
 	virtual bool NewDocument_IsVisibleForType(FBlueprintEditor::ECreatedDocumentType GraphType) const override;
@@ -146,7 +146,7 @@ protected:
 	virtual bool CanPasteNodes() const override { return FRigVMEditorBase::CanPasteNodes(); }
 	virtual bool IsNativeParentClassCodeLinkEnabled() const override { return FRigVMEditorBase::IsNativeParentClassCodeLinkEnabled(); }
 	virtual bool ReparentBlueprint_IsVisible() const override { return FRigVMEditorBase::ReparentBlueprint_IsVisible(); }
-	virtual FReply OnSpawnGraphNodeByShortcut(FInputChord InChord, const FVector2D& InPosition, UEdGraph* InGraph) override { return FRigVMEditorBase::OnSpawnGraphNodeByShortcut(InChord, InPosition, InGraph); }
+	virtual FReply OnSpawnGraphNodeByShortcut(FInputChord InChord, const FVector2f& InPosition, UEdGraph* InGraph) override { return FRigVMEditorBase::OnSpawnGraphNodeByShortcut(InChord, FDeprecateSlateVector2D(InPosition), InGraph); }
 	virtual bool ShouldLoadBPLibrariesFromAssetRegistry() override { return FRigVMEditorBase::ShouldLoadBPLibrariesFromAssetRegistry(); }
 	virtual bool ShouldOpenGraphByDefault() const { return FRigVMEditorBase::ShouldOpenGraphByDefault(); }
 	virtual void AddNewFunctionVariant(const UEdGraph* InOriginalFunction) override { return FRigVMEditorBase::AddNewFunctionVariant(InOriginalFunction); }
diff --git a/Engine/Plugins/Runtime/RigVM/Source/RigVMEditor/Public/Editor/RigVMNewEditor.h b/Engine/Plugins/Runtime/RigVM/Source/RigVMEditor/Public/Editor/RigVMNewEditor.h
--- a/Engine/Plugins/Runtime/RigVM/Source/RigVMEditor/Public/Editor/RigVMNewEditor.h
+++ b/Engine/Plugins/Runtime/RigVM/Source/RigVMEditor/Public/Editor/RigVMNewEditor.h
@@ -90,7 +90,7 @@ protected:
 	virtual void OnBlueprintChangedInnerImpl(UBlueprint* InBlueprint, bool bIsJustBeingCompiled) override;
 	virtual void RefreshEditorsImpl(ERefreshRigVMEditorReason::Type Reason) override;
 	virtual void SetupGraphEditorEventsImpl(UEdGraph* InGraph, SGraphEditor::FGraphEditorEvents& InEvents) override;
-	virtual FActionMenuContent OnCreateGraphActionMenu(UEdGraph* InGraph, const FVector2D& InNodePosition, const TArray<UEdGraphPin*>& InDraggedPins, bool bAutoExpand, SGraphEditor::FActionMenuClosed InOnMenuClosed) override;
+	virtual FActionMenuContent OnCreateGraphActionMenu(UEdGraph* InGraph, const FVector2f& InNodePosition, const TArray<UEdGraphPin*>& InDraggedPins, bool bAutoExpand, SGraphEditor::FActionMenuClosed InOnMenuClosed) override;
 	virtual void AddReferencedObjectsImpl(FReferenceCollector& Collector) override;
 	virtual bool IsSectionVisible(RigVMNodeSectionID::Type InSectionID) const; 
 	virtual bool NewDocument_IsVisibleForType(ECreatedDocumentType GraphType) const;
@@ -129,7 +129,7 @@ public:
 	virtual bool ShouldOpenGraphByDefault() const { return FRigVMEditorBase::ShouldOpenGraphByDefault(); }
 	virtual void OnFinishedChangingProperties(const FPropertyChangedEvent& PropertyChangedEvent) { FRigVMEditorBase::OnFinishedChangingProperties(PropertyChangedEvent); }
 	virtual void HandleSetObjectBeingDebugged(UObject* InObject) { FRigVMEditorBase::HandleSetObjectBeingDebugged(InObject); }
-	virtual FReply OnSpawnGraphNodeByShortcut(FInputChord InChord, const FVector2D& InPosition, UEdGraph* InGraph) { return FRigVMEditorBase::OnSpawnGraphNodeByShortcut(InChord, InPosition, InGraph); }
+	virtual FReply OnSpawnGraphNodeByShortcut(FInputChord InChord, const FVector2f& InPosition, UEdGraph* InGraph) { return FRigVMEditorBase::OnSpawnGraphNodeByShortcut(InChord, InPosition, InGraph); }
 	virtual FPreviewScene* GetPreviewScene() override { return nullptr; }
 
 	//~ Begin IToolkit Interface
diff --git a/Engine/Plugins/Runtime/RigVM/Source/RigVMEditor/Public/Editor/SRigVMActionMenu.h b/Engine/Plugins/Runtime/RigVM/Source/RigVMEditor/Public/Editor/SRigVMActionMenu.h
--- a/Engine/Plugins/Runtime/RigVM/Source/RigVMEditor/Public/Editor/SRigVMActionMenu.h
+++ b/Engine/Plugins/Runtime/RigVM/Source/RigVMEditor/Public/Editor/SRigVMActionMenu.h
@@ -20,11 +20,11 @@ public:
 
 	SLATE_BEGIN_ARGS( SRigVMActionMenu )
 		: _GraphObj( nullptr )
-		, _NewNodePosition( FVector2D::ZeroVector )
+		, _NewNodePosition( FVector2f::ZeroVector )
 	{}
 
 		SLATE_ARGUMENT( UEdGraph*, GraphObj )
-		SLATE_ARGUMENT( FVector2D, NewNodePosition )
+		SLATE_ARGUMENT( FDeprecateSlateVector2D, NewNodePosition )
 		SLATE_ARGUMENT( TArray<UEdGraphPin*>, DraggedFromPins )
 	SLATE_END_ARGS()
 
diff --git a/Engine/Source/Editor/AudioEditor/Private/SoundClassEditor.cpp b/Engine/Source/Editor/AudioEditor/Private/SoundClassEditor.cpp
--- a/Engine/Source/Editor/AudioEditor/Private/SoundClassEditor.cpp
+++ b/Engine/Source/Editor/AudioEditor/Private/SoundClassEditor.cpp
@@ -302,7 +302,7 @@ TSharedRef<SGraphEditor> FSoundClassEditor::CreateGraphEditorWidget()
 
 	SGraphEditor::FGraphEditorEvents InEvents;
 	InEvents.OnSelectionChanged = SGraphEditor::FOnSelectionChanged::CreateSP(this, &FSoundClassEditor::OnSelectedNodesChanged);
-	InEvents.OnCreateActionMenu = SGraphEditor::FOnCreateActionMenu::CreateSP(this, &FSoundClassEditor::OnCreateGraphActionMenu);
+	InEvents.OnCreateActionMenuAtLocation = SGraphEditor::FOnCreateActionMenuAtLocation::CreateSP(this, &FSoundClassEditor::OnCreateGraphActionMenu);
 
 	return SNew(SGraphEditor)
 		.AdditionalCommands(GraphEditorCommands)
@@ -332,12 +332,12 @@ void FSoundClassEditor::OnSelectedNodesChanged(const TSet<class UObject*>& NewSe
 	}
 }
 
-FActionMenuContent FSoundClassEditor::OnCreateGraphActionMenu(UEdGraph* InGraph, const FVector2D& InNodePosition, const TArray<UEdGraphPin*>& InDraggedPins, bool bAutoExpand, SGraphEditor::FActionMenuClosed InOnMenuClosed)
+FActionMenuContent FSoundClassEditor::OnCreateGraphActionMenu(UEdGraph* InGraph, const FVector2f& InNodePosition, const TArray<UEdGraphPin*>& InDraggedPins, bool bAutoExpand, SGraphEditor::FActionMenuClosed InOnMenuClosed)
 {
 	TSharedRef<SSoundClassActionMenu> ActionMenu = 
 		SNew(SSoundClassActionMenu)
 		.GraphObj(InGraph)
-		.NewNodePosition(InNodePosition)
+		.NewNodePosition(FDeprecateSlateVector2D(InNodePosition))
 		.DraggedFromPins(InDraggedPins)
 		.AutoExpandActionMenu(bAutoExpand)
 		.OnClosedCallback(InOnMenuClosed);
diff --git a/Engine/Source/Editor/AudioEditor/Private/SoundClassEditor.h b/Engine/Source/Editor/AudioEditor/Private/SoundClassEditor.h
--- a/Engine/Source/Editor/AudioEditor/Private/SoundClassEditor.h
+++ b/Engine/Source/Editor/AudioEditor/Private/SoundClassEditor.h
@@ -84,7 +84,7 @@ private:
 	void OnSelectedNodesChanged(const TSet<class UObject*>& NewSelection);
 
 	/** Called to create context menu when right-clicking on graph */
-	FActionMenuContent OnCreateGraphActionMenu(UEdGraph* InGraph, const FVector2D& InNodePosition, const TArray<UEdGraphPin*>& InDraggedPins, bool bAutoExpand, SGraphEditor::FActionMenuClosed InOnMenuClosed);
+	FActionMenuContent OnCreateGraphActionMenu(UEdGraph* InGraph, const FVector2f& InNodePosition, const TArray<UEdGraphPin*>& InDraggedPins, bool bAutoExpand, SGraphEditor::FActionMenuClosed InOnMenuClosed);
 
 	/** Select every node in the graph */
 	void SelectAllNodes();
@@ -106,7 +106,7 @@ private:
 	bool CanExcuteToggleSolo() const;
 	bool IsSoloToggled() const;
 
-	/** Toggle mute (Muting means this will be explicitly silienced)*/
+	/** Toggle mute (Muting means this will be explicitly silenced)*/
 	void ToggleMute();
 	bool CanExcuteToggleMute() const;
 	bool IsMuteToggled() const;
diff --git a/Engine/Source/Editor/AudioEditor/Private/SoundSubmixEditor.cpp b/Engine/Source/Editor/AudioEditor/Private/SoundSubmixEditor.cpp
--- a/Engine/Source/Editor/AudioEditor/Private/SoundSubmixEditor.cpp
+++ b/Engine/Source/Editor/AudioEditor/Private/SoundSubmixEditor.cpp
@@ -313,7 +313,7 @@ TSharedRef<SGraphEditor> FSoundSubmixEditor::CreateGraphEditorWidget(USoundSubmi
 
 	SGraphEditor::FGraphEditorEvents InEvents;
 	InEvents.OnSelectionChanged = SGraphEditor::FOnSelectionChanged::CreateSP(this, &FSoundSubmixEditor::OnSelectedNodesChanged);
-	InEvents.OnCreateActionMenu = SGraphEditor::FOnCreateActionMenu::CreateSP(this, &FSoundSubmixEditor::OnCreateGraphActionMenu);
+	InEvents.OnCreateActionMenuAtLocation = SGraphEditor::FOnCreateActionMenuAtLocation::CreateSP(this, &FSoundSubmixEditor::OnCreateGraphActionMenu);
 
 	return SNew(SSoundSubmixGraphEditor, SharedThis(this))
 		.AdditionalCommands(GraphEditorCommands)
@@ -342,12 +342,12 @@ void FSoundSubmixEditor::OnSelectedNodesChanged(const TSet<UObject*>& NewSelecti
 	}
 }
 
-FActionMenuContent FSoundSubmixEditor::OnCreateGraphActionMenu(UEdGraph* InGraph, const FVector2D& InNodePosition, const TArray<UEdGraphPin*>& InDraggedPins, bool bAutoExpand, SGraphEditor::FActionMenuClosed InOnMenuClosed)
+FActionMenuContent FSoundSubmixEditor::OnCreateGraphActionMenu(UEdGraph* InGraph, const FVector2f& InNodePosition, const TArray<UEdGraphPin*>& InDraggedPins, bool bAutoExpand, SGraphEditor::FActionMenuClosed InOnMenuClosed)
 {
 	TSharedRef<SSoundSubmixActionMenu> ActionMenu = 
 		SNew(SSoundSubmixActionMenu)
 		.GraphObj(InGraph)
-		.NewNodePosition(InNodePosition)
+		.NewNodePosition(FDeprecateSlateVector2D(InNodePosition))
 		.DraggedFromPins(InDraggedPins)
 		.AutoExpandActionMenu(bAutoExpand)
 		.OnClosedCallback(InOnMenuClosed);
diff --git a/Engine/Source/Editor/AudioEditor/Private/SoundSubmixEditor.h b/Engine/Source/Editor/AudioEditor/Private/SoundSubmixEditor.h
--- a/Engine/Source/Editor/AudioEditor/Private/SoundSubmixEditor.h
+++ b/Engine/Source/Editor/AudioEditor/Private/SoundSubmixEditor.h
@@ -93,7 +93,7 @@ private:
 	void OnSelectedNodesChanged(const TSet<UObject*>& NewSelection);
 
 	/** Called to create context menu when right-clicking on graph */
-	FActionMenuContent OnCreateGraphActionMenu(UEdGraph* InGraph, const FVector2D& InNodePosition, const TArray<UEdGraphPin*>& InDraggedPins, bool bAutoExpand, SGraphEditor::FActionMenuClosed InOnMenuClosed);
+	FActionMenuContent OnCreateGraphActionMenu(UEdGraph* InGraph, const FVector2f& InNodePosition, const TArray<UEdGraphPin*>& InDraggedPins, bool bAutoExpand, SGraphEditor::FActionMenuClosed InOnMenuClosed);
 
 	/** Adds all children of provided root submix as editable */
 	void AddEditableSubmixChildren(USoundSubmixBase* RootSubmix);
diff --git a/Engine/Source/Editor/GraphEditor/Private/SGraphEditorImpl.cpp b/Engine/Source/Editor/GraphEditor/Private/SGraphEditorImpl.cpp
--- a/Engine/Source/Editor/GraphEditor/Private/SGraphEditorImpl.cpp
+++ b/Engine/Source/Editor/GraphEditor/Private/SGraphEditorImpl.cpp
@@ -677,7 +677,7 @@ void SGraphEditorImpl::Construct( const FArguments& InArgs )
 	EdGraphObj = InArgs._GraphToEdit;
 
 	OnFocused = InArgs._GraphEvents.OnFocused;
-	OnCreateActionMenu = InArgs._GraphEvents.OnCreateActionMenu;
+	OnCreateActionMenuAtLocation = InArgs._GraphEvents.OnCreateActionMenuAtLocation;
 	OnCreateNodeOrPinMenu = InArgs._GraphEvents.OnCreateNodeOrPinMenu;
 	
 	struct Local
@@ -746,11 +746,11 @@ void SGraphEditorImpl::Construct( const FArguments& InArgs )
 			.OnNodeDoubleClicked( InArgs._GraphEvents.OnNodeDoubleClicked )
 			.IsEditable( this, &SGraphEditorImpl::IsGraphEditable )
 			.DisplayAsReadOnly( this, &SGraphEditorImpl::DisplayGraphAsReadOnly )
-			.OnDropActor( InArgs._GraphEvents.OnDropActor )
-			.OnDropStreamingLevel( InArgs._GraphEvents.OnDropStreamingLevel )
+			.OnDropActors(InArgs._GraphEvents.OnDropActors)
+			.OnDropStreamingLevels( InArgs._GraphEvents.OnDropStreamingLevels )
 			.OnVerifyTextCommit( InArgs._GraphEvents.OnVerifyTextCommit )
 			.OnTextCommitted( InArgs._GraphEvents.OnTextCommitted )
-			.OnSpawnNodeByShortcut( InArgs._GraphEvents.OnSpawnNodeByShortcut )
+			.OnSpawnNodeByShortcutAtLocation( InArgs._GraphEvents.OnSpawnNodeByShortcutAtLocation)
 			//.OnUpdateGraphPanel( this, &SGraphEditorImpl::GraphEd_OnPanelUpdated )
 			.OnDisallowedPinConnection( InArgs._GraphEvents.OnDisallowedPinConnection )
 			.ShowGraphStateOverlay(InArgs._ShowGraphStateOverlay)
@@ -1305,9 +1305,9 @@ FActionMenuContent SGraphEditorImpl::GraphEd_OnGetContextMenuFor(const FGraphCon
 		{
 			if (EdGraphObj->GetSchema() != NULL )
 			{
-				if(OnCreateActionMenu.IsBound())
+				if(OnCreateActionMenuAtLocation.IsBound())
 				{
-					Result = OnCreateActionMenu.Execute(
+					Result = OnCreateActionMenuAtLocation.Execute(
 						EdGraphObj, 
 						SpawnInfo.NodeAddPosition,
 						SpawnInfo.DragFromPins,
diff --git a/Engine/Source/Editor/GraphEditor/Private/SGraphEditorImpl.h b/Engine/Source/Editor/GraphEditor/Private/SGraphEditorImpl.h
--- a/Engine/Source/Editor/GraphEditor/Private/SGraphEditorImpl.h
+++ b/Engine/Source/Editor/GraphEditor/Private/SGraphEditorImpl.h
@@ -246,7 +246,7 @@ private:
 	TAttribute<FGraphAppearanceInfo> Appearance;
 
 	SGraphEditor::FOnFocused OnFocused;
-	SGraphEditor::FOnCreateActionMenu OnCreateActionMenu;
+	SGraphEditor::FOnCreateActionMenuAtLocation OnCreateActionMenuAtLocation;
 	SGraphEditor::FOnCreateNodeOrPinMenu OnCreateNodeOrPinMenu;
 
 	TAttribute<bool> IsEditable;
diff --git a/Engine/Source/Editor/GraphEditor/Private/SGraphPanel.cpp b/Engine/Source/Editor/GraphEditor/Private/SGraphPanel.cpp
--- a/Engine/Source/Editor/GraphEditor/Private/SGraphPanel.cpp
+++ b/Engine/Source/Editor/GraphEditor/Private/SGraphPanel.cpp
@@ -99,11 +99,11 @@ void SGraphPanel::Construct( const SGraphPanel::FArguments& InArgs )
 	this->IsEditable = InArgs._IsEditable;
 	this->DisplayAsReadOnly = InArgs._DisplayAsReadOnly;
 	this->OnNodeDoubleClicked = InArgs._OnNodeDoubleClicked;
-	this->OnDropActor = InArgs._OnDropActor;
-	this->OnDropStreamingLevel = InArgs._OnDropStreamingLevel;
+	this->OnDropActors = InArgs._OnDropActors;
+	this->OnDropStreamingLevels = InArgs._OnDropStreamingLevels;
 	this->OnVerifyTextCommit = InArgs._OnVerifyTextCommit;
 	this->OnTextCommitted = InArgs._OnTextCommitted;
-	this->OnSpawnNodeByShortcut = InArgs._OnSpawnNodeByShortcut;
+	this->OnSpawnNodeByShortcutAtLocation = InArgs._OnSpawnNodeByShortcutAtLocation;
 	this->OnUpdateGraphPanel = InArgs._OnUpdateGraphPanel;
 	this->OnDisallowedPinConnection = InArgs._OnDisallowedPinConnection;
 	this->OnDoubleClicked = InArgs._OnDoubleClicked;
@@ -834,7 +834,7 @@ FReply SGraphPanel::OnKeyDown( const FGeometry& MyGeometry, const FKeyEvent& InK
 		// If you're dragging multiple connections then it's probably a Ctrl + Drag to move wires around, which
 		// wouldn't make as much sense to do node creation during
 		UEdGraphPin* PreviewConnectionPin = PreviewConnectorFromPins.Num() == 1 ? PreviewConnectorFromPins[0].GetPinObj(*this) : nullptr;
-		if (OnSpawnNodeByShortcut.IsBound() && PreviewConnectionPin)
+		if (OnSpawnNodeByShortcutAtLocation.IsBound() && PreviewConnectionPin)
 		{
 			// Note: We can't use SavedMousePosForOnPaintEventLocalSpace since it isn't updated while dragging,
 			// and so would just be the mouse position of the connection origin. So instead we'll just use the current cursor pos
@@ -842,7 +842,7 @@ FReply SGraphPanel::OnKeyDown( const FGeometry& MyGeometry, const FKeyEvent& InK
 			FInputChord KeyChord = FInputChord(InKeyEvent.GetKey(), EModifierKey::FromBools(InKeyEvent.IsControlDown(), InKeyEvent.IsAltDown(), InKeyEvent.IsShiftDown(), InKeyEvent.IsCommandDown()));
 
 			int32 NodeCountBefore = GraphObj->Nodes.Num();
-			FReply SpawnNodeReply = OnSpawnNodeByShortcut.Execute(KeyChord, FVector2D(NewNodePosition));
+			FReply SpawnNodeReply = OnSpawnNodeByShortcutAtLocation.Execute(KeyChord, NewNodePosition);
 			int32 NodeCountAfter = GraphObj->Nodes.Num();
 
 			// If we spawned a node then we won't call down into super and instead do some extra handling
@@ -1355,14 +1355,14 @@ FReply SGraphPanel::OnDrop( const FGeometry& MyGeometry, const FDragDropEvent& D
 	else if (Operation->IsOfType<FActorDragDropGraphEdOp>())
 	{
 		TSharedPtr<FActorDragDropGraphEdOp> ActorOp = StaticCastSharedPtr<FActorDragDropGraphEdOp>(Operation);
-		OnDropActor.ExecuteIfBound(ActorOp->Actors, GraphObj, FVector2D(NodeAddPosition));
+		OnDropActors.ExecuteIfBound(ActorOp->Actors, GraphObj, NodeAddPosition);
 		return FReply::Handled();
 	}
 
 	else if (Operation->IsOfType<FLevelDragDropOp>())
 	{
 		TSharedPtr<FLevelDragDropOp> LevelOp = StaticCastSharedPtr<FLevelDragDropOp>(Operation);
-		OnDropStreamingLevel.ExecuteIfBound(LevelOp->StreamingLevelsToDrop, GraphObj, FVector2D(NodeAddPosition));
+		OnDropStreamingLevels.ExecuteIfBound(LevelOp->StreamingLevelsToDrop, GraphObj, NodeAddPosition);
 		return FReply::Handled();
 	}
 
diff --git a/Engine/Source/Editor/GraphEditor/Private/SNodePanel.cpp b/Engine/Source/Editor/GraphEditor/Private/SNodePanel.cpp
--- a/Engine/Source/Editor/GraphEditor/Private/SNodePanel.cpp
+++ b/Engine/Source/Editor/GraphEditor/Private/SNodePanel.cpp
@@ -1046,9 +1046,9 @@ FReply SNodePanel::OnMouseButtonUp( const FGeometry& MyGeometry, const FPointerE
 				// We clicked on the panel background
 				this->SelectionManager.ClearSelectionSet();
 
-				if(OnSpawnNodeByShortcut.IsBound())
+				if(OnSpawnNodeByShortcutAtLocation.IsBound())
 				{
-					OnSpawnNodeByShortcut.Execute(LastKeyChordDetected, PanelCoordToGraphCoord(  MyGeometry.AbsoluteToLocal( MouseEvent.GetScreenSpacePosition() ) ));
+					OnSpawnNodeByShortcutAtLocation.Execute(LastKeyChordDetected, PanelCoordToGraphCoord(  MyGeometry.AbsoluteToLocal( MouseEvent.GetScreenSpacePosition() ) ));
 				}
 
 				LastKeyChordDetected = FInputChord();
diff --git a/Engine/Source/Editor/GraphEditor/Public/DragAndDrop/GraphNodeDragDropOp.h b/Engine/Source/Editor/GraphEditor/Public/DragAndDrop/GraphNodeDragDropOp.h
--- a/Engine/Source/Editor/GraphEditor/Public/DragAndDrop/GraphNodeDragDropOp.h
+++ b/Engine/Source/Editor/GraphEditor/Public/DragAndDrop/GraphNodeDragDropOp.h
@@ -4,12 +4,26 @@
 #include "CoreMinimal.h"
 #include "DragAndDrop/DecoratedDragDropOp.h"
 
+class UEdGraph;
+
 class FGraphNodeDragDropOp : public FDecoratedDragDropOp
 {
 public:
 	DRAG_DROP_OPERATOR_TYPE(FGraphNodeDragDropOp, FDecoratedDragDropOp)
 
-	DECLARE_DELEGATE_FourParams( FOnPerformDropToGraph, TSharedPtr<FDragDropOperation>, class UEdGraph*, const FVector2D&, const FVector2D&);
+	UE_DEPRECATED(5.6, "Please use the delegate accepting FVector2f.")
+	DECLARE_DELEGATE_FourParams( FOnPerformDropToGraph, TSharedPtr<FDragDropOperation>, UEdGraph*, const FVector2D&, const FVector2D&);
+	DECLARE_DELEGATE_FourParams(FOnPerformDropToGraphAtLocation, TSharedPtr<FDragDropOperation>, UEdGraph*, const FVector2f&, const FVector2f&);
 
+	PRAGMA_DISABLE_DEPRECATION_WARNINGS
 	FOnPerformDropToGraph OnPerformDropToGraph;
+	PRAGMA_ENABLE_DEPRECATION_WARNINGS
+
+	FOnPerformDropToGraphAtLocation OnPerformDropToGraphAtLocation = FOnPerformDropToGraphAtLocation::CreateLambda(
+	[this](TSharedPtr<FDragDropOperation> InOperation, UEdGraph* InGraph, const FVector2f& InNodePos, const FVector2f& InScreenPos)
+		{
+		PRAGMA_DISABLE_DEPRECATION_WARNINGS
+			OnPerformDropToGraph.ExecuteIfBound(InOperation, InGraph, FVector2D(InNodePos), FVector2D(InScreenPos));
+		PRAGMA_ENABLE_DEPRECATION_WARNINGS
+		});
 };
diff --git a/Engine/Source/Editor/GraphEditor/Public/SGraphEditorActionMenu.h b/Engine/Source/Editor/GraphEditor/Public/SGraphEditorActionMenu.h
--- a/Engine/Source/Editor/GraphEditor/Public/SGraphEditorActionMenu.h
+++ b/Engine/Source/Editor/GraphEditor/Public/SGraphEditorActionMenu.h
@@ -30,12 +30,12 @@ class GRAPHEDITOR_API SGraphEditorActionMenu : public SBorder
 public:
 	SLATE_BEGIN_ARGS( SGraphEditorActionMenu )
 		: _GraphObj( static_cast<UEdGraph*>(NULL) )
-		, _NewNodePosition( FVector2D::ZeroVector )
+		, _NewNodePosition( FVector2f::ZeroVector )
 		, _AutoExpandActionMenu( false )
 		{}
 
 		SLATE_ARGUMENT( UEdGraph*, GraphObj )
-		SLATE_ARGUMENT( FVector2D, NewNodePosition )
+		SLATE_ARGUMENT( FDeprecateSlateVector2D, NewNodePosition )
 		SLATE_ARGUMENT( TArray<UEdGraphPin*>, DraggedFromPins )
 		SLATE_ARGUMENT( SGraphEditor::FActionMenuClosed, OnClosedCallback )
 		SLATE_ARGUMENT( bool, AutoExpandActionMenu )
diff --git a/Engine/Source/Editor/GraphEditor/Public/SGraphPanel.h b/Engine/Source/Editor/GraphEditor/Public/SGraphPanel.h
--- a/Engine/Source/Editor/GraphEditor/Public/SGraphPanel.h
+++ b/Engine/Source/Editor/GraphEditor/Public/SGraphPanel.h
@@ -93,9 +93,43 @@ public:
 
 		SLATE_EVENT( FOnGetContextMenuFor, OnGetContextMenuFor )
 		SLATE_EVENT( SGraphEditor::FOnSelectionChanged, OnSelectionChanged )
-		SLATE_EVENT( FSingleNodeEvent, OnNodeDoubleClicked )
-		SLATE_EVENT( SGraphEditor::FOnDropActor, OnDropActor )
-		SLATE_EVENT( SGraphEditor::FOnDropStreamingLevel, OnDropStreamingLevel )
+		SLATE_EVENT(FSingleNodeEvent, OnNodeDoubleClicked)
+		SLATE_EVENT(SGraphEditor::FOnDropActors, OnDropActors)
+
+PRAGMA_DISABLE_DEPRECATION_WARNINGS
+		static SGraphEditor::FOnDropActors ConvertOnDropActorDelegate(const SGraphEditor::FOnDropActor& LegacyDelegate)
+		{
+			if (LegacyDelegate.IsBound())
+			{
+				return SGraphEditor::FOnDropActors::CreateLambda([LegacyDelegate](const TArray< TWeakObjectPtr<AActor> >& Actors, UEdGraph* InGraph, const FVector2f& InDropLocation)
+				{
+					LegacyDelegate.Execute(Actors, InGraph, FVector2D(InDropLocation));
+				});
+			}
+
+			return SGraphEditor::FOnDropActors();
+		}
+		SLATE_EVENT_DEPRECATED(5.6, "Use SGraphEditor::FOnDropActors.",
+			SGraphEditor::FOnDropActor,
+			OnDropActor, OnDropActors, ConvertOnDropActorDelegate)
+		static SGraphEditor::FOnDropStreamingLevels ConvertOnDropStreamingLevelDelegate(const SGraphEditor::FOnDropStreamingLevel& LegacyDelegate)
+		{
+			if (LegacyDelegate.IsBound())
+			{
+				return SGraphEditor::FOnDropStreamingLevels::CreateLambda([LegacyDelegate](const TArray< TWeakObjectPtr<ULevelStreaming> >& Levels, UEdGraph* InGraph, const FVector2f& InDropLocation)
+				{
+					LegacyDelegate.Execute(Levels, InGraph, FVector2D(InDropLocation));
+				});
+			}
+
+			return SGraphEditor::FOnDropStreamingLevels();
+		}
+		SLATE_EVENT_DEPRECATED(5.6, "Use SGraphEditor::FOnDropStreamingLevels.",
+			SGraphEditor::FOnDropStreamingLevel,
+			OnDropStreamingLevel, OnDropStreamingLevels, ConvertOnDropStreamingLevelDelegate)
+PRAGMA_ENABLE_DEPRECATION_WARNINGS
+	
+		SLATE_EVENT( SGraphEditor::FOnDropStreamingLevels, OnDropStreamingLevels )
 		SLATE_ARGUMENT( class UEdGraph*, GraphObj )
 		SLATE_ARGUMENT( TSharedPtr<TArray<FDiffSingleResult>>, DiffResults )
 		SLATE_ATTRIBUTE( int32, FocusedDiffResult )
@@ -106,7 +140,25 @@ public:
 		SLATE_ATTRIBUTE(bool, ShowGraphStateOverlay)
 		SLATE_EVENT( FOnNodeVerifyTextCommit, OnVerifyTextCommit )
 		SLATE_EVENT( FOnNodeTextCommitted, OnTextCommitted )
-		SLATE_EVENT( SGraphEditor::FOnSpawnNodeByShortcut, OnSpawnNodeByShortcut )
+		
+		PRAGMA_DISABLE_DEPRECATION_WARNINGS
+		static SGraphEditor::FOnSpawnNodeByShortcutAtLocation ConvertOnSpawnNodeByShortcutDelegate(const SGraphEditor::FOnSpawnNodeByShortcut& LegacyDelegate)
+		{
+			if (LegacyDelegate.IsBound())
+			{
+				return SGraphEditor::FOnSpawnNodeByShortcutAtLocation::CreateLambda([LegacyDelegate](FInputChord InInputChord, const FVector2f& InLocation)
+				{
+					return LegacyDelegate.Execute(InInputChord, FVector2D(InLocation));
+				});
+			}
+
+			return SGraphEditor::FOnSpawnNodeByShortcutAtLocation();
+		}
+		PRAGMA_ENABLE_DEPRECATION_WARNINGS
+		SLATE_EVENT_DEPRECATED(5.6, "Use SGraphEditor::FOnSpawnNodeByShortcutAtLocation.",
+			PRAGMA_DISABLE_DEPRECATION_WARNINGS SGraphEditor::FOnSpawnNodeByShortcut PRAGMA_ENABLE_DEPRECATION_WARNINGS,
+			OnSpawnNodeByShortcut, OnSpawnNodeByShortcutAtLocation, ConvertOnSpawnNodeByShortcutDelegate)
+		SLATE_EVENT( SGraphEditor::FOnSpawnNodeByShortcutAtLocation, OnSpawnNodeByShortcutAtLocation )
 		SLATE_EVENT( FOnUpdateGraphPanel, OnUpdateGraphPanel )
 		SLATE_EVENT( SGraphEditor::FOnDisallowedPinConnection, OnDisallowedPinConnection )
 		SLATE_EVENT( SGraphEditor::FOnDoubleClicked, OnDoubleClicked )
@@ -310,10 +362,10 @@ protected:
 	FOnGetContextMenuFor OnGetContextMenuFor;
 
 	/** Invoked when an actor is dropped onto the panel */
-	SGraphEditor::FOnDropActor OnDropActor;
+	SGraphEditor::FOnDropActors OnDropActors;
 
 	/** Invoked when a streaming level is dropped onto the panel */
-	SGraphEditor::FOnDropStreamingLevel OnDropStreamingLevel;
+	SGraphEditor::FOnDropStreamingLevels OnDropStreamingLevels;
 
 	/** What to do when a node is double-clicked */
 	FSingleNodeEvent OnNodeDoubleClicked;
diff --git a/Engine/Source/Editor/GraphEditor/Public/SNodePanel.h b/Engine/Source/Editor/GraphEditor/Public/SNodePanel.h
--- a/Engine/Source/Editor/GraphEditor/Public/SNodePanel.h
+++ b/Engine/Source/Editor/GraphEditor/Public/SNodePanel.h
@@ -1164,7 +1164,7 @@ protected:
 	EGraphRenderingLOD::Type CurrentLOD;
 
 	/** Invoked when the user may be attempting to spawn a node using a shortcut */
-	SGraphEditor::FOnSpawnNodeByShortcut OnSpawnNodeByShortcut;
+	SGraphEditor::FOnSpawnNodeByShortcutAtLocation OnSpawnNodeByShortcutAtLocation;
 
 	/** The last key chord detected in this graph panel */
 	FInputChord LastKeyChordDetected;
diff --git a/Engine/Source/Editor/Kismet/Private/BlueprintEditor.cpp b/Engine/Source/Editor/Kismet/Private/BlueprintEditor.cpp
--- a/Engine/Source/Editor/Kismet/Private/BlueprintEditor.cpp
+++ b/Engine/Source/Editor/Kismet/Private/BlueprintEditor.cpp
@@ -1775,12 +1775,12 @@ TSharedRef<SGraphEditor> FBlueprintEditor::CreateGraphEditorWidget(TSharedRef<FT
 void FBlueprintEditor::SetupGraphEditorEvents(UEdGraph* InGraph, SGraphEditor::FGraphEditorEvents& InEvents)
 {
 	InEvents.OnSelectionChanged = SGraphEditor::FOnSelectionChanged::CreateSP( this, &FBlueprintEditor::OnSelectedNodesChanged );
-	InEvents.OnDropActor = SGraphEditor::FOnDropActor::CreateSP( this, &FBlueprintEditor::OnGraphEditorDropActor );
-	InEvents.OnDropStreamingLevel = SGraphEditor::FOnDropStreamingLevel::CreateSP( this, &FBlueprintEditor::OnGraphEditorDropStreamingLevel );
+	InEvents.OnDropActors = SGraphEditor::FOnDropActors::CreateSP( this, &FBlueprintEditor::OnGraphEditorDropActor );
+	InEvents.OnDropStreamingLevels = SGraphEditor::FOnDropStreamingLevels::CreateSP( this, &FBlueprintEditor::OnGraphEditorDropStreamingLevel );
 	InEvents.OnNodeDoubleClicked = FSingleNodeEvent::CreateSP(this, &FBlueprintEditor::OnNodeDoubleClicked);
 	InEvents.OnVerifyTextCommit = FOnNodeVerifyTextCommit::CreateSP(this, &FBlueprintEditor::OnNodeVerifyTitleCommit);
 	InEvents.OnTextCommitted = FOnNodeTextCommitted::CreateSP(this, &FBlueprintEditor::OnNodeTitleCommitted);
-	InEvents.OnSpawnNodeByShortcut = SGraphEditor::FOnSpawnNodeByShortcut::CreateSP(this, &FBlueprintEditor::OnSpawnGraphNodeByShortcut, InGraph);
+	InEvents.OnSpawnNodeByShortcutAtLocation = SGraphEditor::FOnSpawnNodeByShortcutAtLocation::CreateSP(this, &FBlueprintEditor::OnSpawnGraphNodeByShortcut, InGraph);
 	InEvents.OnNodeSpawnedByKeymap = SGraphEditor::FOnNodeSpawnedByKeymap::CreateSP(this, &FBlueprintEditor::OnNodeSpawnedByKeymap );
 	InEvents.OnDisallowedPinConnection = SGraphEditor::FOnDisallowedPinConnection::CreateSP(this, &FBlueprintEditor::OnDisallowedPinConnection);
 	InEvents.OnDoubleClicked = SGraphEditor::FOnDoubleClicked::CreateSP(this, &FBlueprintEditor::NavigateToParentGraphByDoubleClick);
@@ -1788,7 +1788,7 @@ void FBlueprintEditor::SetupGraphEditorEvents(UEdGraph* InGraph, SGraphEditor::F
 	// Custom menu for K2 schemas
 	if(InGraph->Schema != nullptr && InGraph->Schema->IsChildOf(UEdGraphSchema_K2::StaticClass()))
 	{
-		InEvents.OnCreateActionMenu = SGraphEditor::FOnCreateActionMenu::CreateSP(this, &FBlueprintEditor::OnCreateGraphActionMenu);
+		InEvents.OnCreateActionMenuAtLocation = SGraphEditor::FOnCreateActionMenuAtLocation::CreateSP(this, &FBlueprintEditor::OnCreateGraphActionMenu);
 	}
 }
 
@@ -3943,14 +3943,14 @@ void FBlueprintEditor::OnGraphEditorBackgrounded(const TSharedRef<SGraphEditor>&
 	}
 }
 
-void FBlueprintEditor::OnGraphEditorDropActor(const TArray< TWeakObjectPtr<AActor> >& Actors, UEdGraph* Graph, const FVector2D& DropLocation)
+void FBlueprintEditor::OnGraphEditorDropActor(const TArray< TWeakObjectPtr<AActor> >& Actors, UEdGraph* Graph, const FVector2f& DropLocation)
 {
 	// We need to check that the dropped actor is in the right sublevel for the reference
 	ULevel* BlueprintLevel = FBlueprintEditorUtils::GetLevelFromBlueprint(GetBlueprintObj());
 
 	if (BlueprintLevel && FBlueprintEditorUtils::IsLevelScriptBlueprint(GetBlueprintObj()))
 	{
-		FVector2D NodeLocation = DropLocation;
+		FDeprecateSlateVector2D NodeLocation = DropLocation;
 		for (int32 i = 0; i < Actors.Num(); i++)
 		{
 			AActor* DroppedActor = Actors[i].Get();
@@ -3971,7 +3971,7 @@ void FBlueprintEditor::OnGraphEditorDropActor(const TArray< TWeakObjectPtr<AActo
 	}
 }
 
-void FBlueprintEditor::OnGraphEditorDropStreamingLevel(const TArray< TWeakObjectPtr<ULevelStreaming> >& Levels, UEdGraph* Graph, const FVector2D& DropLocation)
+void FBlueprintEditor::OnGraphEditorDropStreamingLevel(const TArray< TWeakObjectPtr<ULevelStreaming> >& Levels, UEdGraph* Graph, const FVector2f& DropLocation)
 {
 	UFunction* TargetFunc = UGameplayStatics::StaticClass()->FindFunctionByName(GET_FUNCTION_NAME_CHECKED(UGameplayStatics, GetStreamingLevel));
 	check(TargetFunc);
@@ -3983,7 +3983,7 @@ void FBlueprintEditor::OnGraphEditorDropStreamingLevel(const TArray< TWeakObject
 		{
 			UK2Node_CallFunction* Node = FEdGraphSchemaAction_K2NewNode::SpawnNode<UK2Node_CallFunction>(
 				Graph,
-				DropLocation + (i * FVector2D(0, 80)),
+				FDeprecateSlateVector2D(DropLocation + (i * FVector2f(0.0f, 80.0f))),
 				EK2NewNodeFlags::SelectNewNode,
 				[TargetFunc](UK2Node_CallFunction* NewInstance)
 				{
@@ -3999,6 +3999,11 @@ void FBlueprintEditor::OnGraphEditorDropStreamingLevel(const TArray< TWeakObject
 }
 
 FActionMenuContent FBlueprintEditor::OnCreateGraphActionMenu(UEdGraph* InGraph, const FVector2D& InNodePosition, const TArray<UEdGraphPin*>& InDraggedPins, bool bAutoExpand, SGraphEditor::FActionMenuClosed InOnMenuClosed)
+{
+	return OnCreateGraphActionMenu(InGraph, UE::Slate::CastToVector2f(InNodePosition), InDraggedPins, bAutoExpand, InOnMenuClosed);
+}
+
+FActionMenuContent FBlueprintEditor::OnCreateGraphActionMenu(UEdGraph* InGraph, const FVector2f& InNodePosition, const TArray<UEdGraphPin*>& InDraggedPins, bool bAutoExpand, SGraphEditor::FActionMenuClosed InOnMenuClosed)
 {
 	HasOpenActionMenu = InGraph;
 	if (!BlueprintEditorImpl::GraphHasUserPlacedNodes(InGraph))
@@ -4009,7 +4014,7 @@ FActionMenuContent FBlueprintEditor::OnCreateGraphActionMenu(UEdGraph* InGraph,
 	TSharedRef<SBlueprintActionMenu> ActionMenu = 
 		SNew(SBlueprintActionMenu, SharedThis(this))
 		.GraphObj(InGraph)
-		.NewNodePosition(InNodePosition)
+		.NewNodePosition(FDeprecateSlateVector2D(InNodePosition))
 		.DraggedFromPins(InDraggedPins)
 		.AutoExpandActionMenu(bAutoExpand)
 		.OnClosedCallback(InOnMenuClosed)
@@ -10436,6 +10441,13 @@ void FBlueprintEditor::DestroyPreview()
 }
 
 
+FReply FBlueprintEditor::OnSpawnGraphNodeByShortcut(FInputChord InChord, const FVector2f& InPosition, UEdGraph* InGraph)
+{
+PRAGMA_DISABLE_DEPRECATION_WARNINGS
+	return OnSpawnGraphNodeByShortcut(InChord, FVector2D(InPosition), InGraph);
+PRAGMA_ENABLE_DEPRECATION_WARNINGS
+}
+
 FReply FBlueprintEditor::OnSpawnGraphNodeByShortcut(FInputChord InChord, const FVector2D& InPosition, UEdGraph* InGraph)
 {
 	UEdGraph* Graph = InGraph;
diff --git a/Engine/Source/Editor/Kismet/Public/BlueprintEditor.h b/Engine/Source/Editor/Kismet/Public/BlueprintEditor.h
--- a/Engine/Source/Editor/Kismet/Public/BlueprintEditor.h
+++ b/Engine/Source/Editor/Kismet/Public/BlueprintEditor.h
@@ -656,7 +656,9 @@ public:
 	virtual void ClearSelectionStateFor(FName SelectionOwner);
 
 	/** Handles spawning a graph node in the current graph using the passed in chord */
-	virtual FReply OnSpawnGraphNodeByShortcut(FInputChord InChord, const FVector2D& InPosition, UEdGraph* InGraph);
+	UE_DEPRECATED(5.6, "Please use the version of the function accepting FVector2f.")
+	virtual FReply OnSpawnGraphNodeByShortcut(FInputChord InChord, const FVector2D& InPosition, UEdGraph* InGraph) UE_SLATE_DEPRECATED_VECTOR_VIRTUAL_FUNCTION;
+	virtual FReply OnSpawnGraphNodeByShortcut(FInputChord InChord, const FVector2f& InPosition, UEdGraph* InGraph);
 
 	/**
 	 * Perform the actual promote to variable action on the given pin in the given blueprint.
@@ -882,13 +884,15 @@ protected:
 	void OnSelectedNodesChanged(const TSet<class UObject*>& NewSelection) { OnSelectedNodesChangedImpl(NewSelection); }
 
 	/** Called when an actor is dropped onto the graph editor */
-	void OnGraphEditorDropActor(const TArray< TWeakObjectPtr<AActor> >& Actors, UEdGraph* Graph, const FVector2D& DropLocation);
+	void OnGraphEditorDropActor(const TArray< TWeakObjectPtr<AActor> >& Actors, UEdGraph* Graph, const FVector2f& DropLocation);
 
 	/** Called when a streaming level is dropped onto the graph editor */
-	void OnGraphEditorDropStreamingLevel(const TArray< TWeakObjectPtr<ULevelStreaming> >& Levels, UEdGraph* Graph, const FVector2D& DropLocation);
+	void OnGraphEditorDropStreamingLevel(const TArray< TWeakObjectPtr<ULevelStreaming> >& Levels, UEdGraph* Graph, const FVector2f& DropLocation);
 
 	/** Called to create context menu when right-clicking on graph */
+	UE_DEPRECATED(5.6, "Please use the version of the function taking FVector2f.")
 	FActionMenuContent OnCreateGraphActionMenu(UEdGraph* InGraph, const FVector2D& InNodePosition, const TArray<UEdGraphPin*>& InDraggedPins, bool bAutoExpand, SGraphEditor::FActionMenuClosed InOnMenuClosed);
+	FActionMenuContent OnCreateGraphActionMenu(UEdGraph* InGraph, const FVector2f& InNodePosition, const TArray<UEdGraphPin*>& InDraggedPins, bool bAutoExpand, SGraphEditor::FActionMenuClosed InOnMenuClosed);
 
 	/** Called from graph context menus when they close to tell the editor why they closed */
 	void OnGraphActionMenuClosed(bool bActionExecuted, bool bContextSensitiveChecked, bool bGraphPinContext);
diff --git a/Engine/Source/Editor/MaterialEditor/Private/MaterialEditor.cpp b/Engine/Source/Editor/MaterialEditor/Private/MaterialEditor.cpp
--- a/Engine/Source/Editor/MaterialEditor/Private/MaterialEditor.cpp
+++ b/Engine/Source/Editor/MaterialEditor/Private/MaterialEditor.cpp
@@ -7948,7 +7948,7 @@ TSharedRef<SGraphEditor> FMaterialEditor::CreateGraphEditorWidget(TSharedRef<cla
 	InEvents.OnNodeDoubleClicked = FSingleNodeEvent::CreateSP(this, &FMaterialEditor::OnNodeDoubleClicked);
 	InEvents.OnTextCommitted = FOnNodeTextCommitted::CreateSP(this, &FMaterialEditor::OnNodeTitleCommitted);
 	InEvents.OnVerifyTextCommit = FOnNodeVerifyTextCommit::CreateSP(this, &FMaterialEditor::OnVerifyNodeTextCommit);
-	InEvents.OnSpawnNodeByShortcut = SGraphEditor::FOnSpawnNodeByShortcut::CreateSP(this, &FMaterialEditor::OnSpawnGraphNodeByShortcut, static_cast<UEdGraph*>(InGraph));
+	InEvents.OnSpawnNodeByShortcutAtLocation = SGraphEditor::FOnSpawnNodeByShortcutAtLocation::CreateSP(this, &FMaterialEditor::OnSpawnGraphNodeByShortcut, static_cast<UEdGraph*>(InGraph));
 
 	// Create the title bar widget
 	TSharedPtr<SWidget> TitleBarWidget = SNew(SMaterialEditorTitleBar)
@@ -8579,7 +8579,7 @@ bool FMaterialEditor::OnVerifyNodeTextCommit(const FText& NewText, UEdGraphNode*
 	return bValid;
 }
 
-FReply FMaterialEditor::OnSpawnGraphNodeByShortcut(FInputChord InChord, const FVector2D& InPosition, UEdGraph* InGraph)
+FReply FMaterialEditor::OnSpawnGraphNodeByShortcut(FInputChord InChord, const FVector2f& InPosition, UEdGraph* InGraph)
 {
 	UEdGraph* Graph = InGraph;
 	if (FMaterialEditorSpawnNodeCommands::IsRegistered())
@@ -8589,7 +8589,7 @@ FReply FMaterialEditor::OnSpawnGraphNodeByShortcut(FInputChord InChord, const FV
 		if (Action.IsValid())
 		{
 			TArray<UEdGraphPin*> DummyPins;
-			Action->PerformAction(Graph, DummyPins, UE::Slate::FDeprecateVector2DParameter(InPosition));
+			Action->PerformAction(Graph, DummyPins, InPosition);
 			return FReply::Handled();
 		}
 	}
diff --git a/Engine/Source/Editor/MaterialEditor/Private/MaterialEditor.h b/Engine/Source/Editor/MaterialEditor/Private/MaterialEditor.h
--- a/Engine/Source/Editor/MaterialEditor/Private/MaterialEditor.h
+++ b/Engine/Source/Editor/MaterialEditor/Private/MaterialEditor.h
@@ -536,7 +536,7 @@ protected:
 	 *
 	 * @return	FReply	Whether chord was handled
 	 */
-	FReply OnSpawnGraphNodeByShortcut(FInputChord InChord, const FVector2D& InPosition, UEdGraph* InGraph);
+	FReply OnSpawnGraphNodeByShortcut(FInputChord InChord, const FVector2f& InPosition, UEdGraph* InGraph);
 
 	/** Calls Modify() on the Material and its editor only data, as well as MaterialFunction (and its editor only data) if present */
 	void ModifyMaterial();
diff --git a/Engine/Source/Editor/PhysicsAssetEditor/Private/PhysicsAssetGraph/SPhysicsAssetGraph.cpp b/Engine/Source/Editor/PhysicsAssetEditor/Private/PhysicsAssetGraph/SPhysicsAssetGraph.cpp
--- a/Engine/Source/Editor/PhysicsAssetEditor/Private/PhysicsAssetGraph/SPhysicsAssetGraph.cpp
+++ b/Engine/Source/Editor/PhysicsAssetEditor/Private/PhysicsAssetGraph/SPhysicsAssetGraph.cpp
@@ -59,7 +59,7 @@ void SPhysicsAssetGraph::Construct(const FArguments& InArgs, const TSharedRef<FP
 	AppearanceInfo.CornerText = LOCTEXT("AppearanceCornerText_Physics", "PHYSICS");
 
 	SGraphEditor::FGraphEditorEvents GraphEvents;
-	GraphEvents.OnCreateActionMenu = SGraphEditor::FOnCreateActionMenu::CreateSP(this, &SPhysicsAssetGraph::OnCreateGraphActionMenu);
+	GraphEvents.OnCreateActionMenuAtLocation = SGraphEditor::FOnCreateActionMenuAtLocation::CreateSP(this, &SPhysicsAssetGraph::OnCreateGraphActionMenu);
 	GraphEvents.OnSelectionChanged = SGraphEditor::FOnSelectionChanged::CreateSP(this, &SPhysicsAssetGraph::HandleSelectionChanged);
 	GraphEvents.OnNodeDoubleClicked = FSingleNodeEvent::CreateSP(this, &SPhysicsAssetGraph::HandleNodeDoubleClicked);
 
@@ -87,7 +87,7 @@ void SPhysicsAssetGraph::Construct(const FArguments& InArgs, const TSharedRef<FP
 	GEditor->RegisterForUndo(this);
 }
 
-FActionMenuContent SPhysicsAssetGraph::OnCreateGraphActionMenu(UEdGraph* InGraph, const FVector2D& InNodePosition, const TArray<UEdGraphPin*>& InDraggedPins, bool bAutoExpand, SGraphEditor::FActionMenuClosed InOnMenuClosed)
+FActionMenuContent SPhysicsAssetGraph::OnCreateGraphActionMenu(UEdGraph* InGraph, const FVector2f& InNodePosition, const TArray<UEdGraphPin*>& InDraggedPins, bool bAutoExpand, SGraphEditor::FActionMenuClosed InOnMenuClosed)
 {
 	if(InDraggedPins.Num() > 0 && InDraggedPins[0]->Direction == EGPD_Output)
 	{
diff --git a/Engine/Source/Editor/PhysicsAssetEditor/Private/PhysicsAssetGraph/SPhysicsAssetGraph.h b/Engine/Source/Editor/PhysicsAssetEditor/Private/PhysicsAssetGraph/SPhysicsAssetGraph.h
--- a/Engine/Source/Editor/PhysicsAssetEditor/Private/PhysicsAssetGraph/SPhysicsAssetGraph.h
+++ b/Engine/Source/Editor/PhysicsAssetEditor/Private/PhysicsAssetGraph/SPhysicsAssetGraph.h
@@ -48,7 +48,7 @@ public:
 private:
 
 	/** Called to create context menu when right-clicking on graph */
-	FActionMenuContent OnCreateGraphActionMenu(UEdGraph* InGraph, const FVector2D& InNodePosition, const TArray<UEdGraphPin*>& InDraggedPins, bool bAutoExpand, SGraphEditor::FActionMenuClosed InOnMenuClosed);
+	FActionMenuContent OnCreateGraphActionMenu(UEdGraph* InGraph, const FVector2f& InNodePosition, const TArray<UEdGraphPin*>& InDraggedPins, bool bAutoExpand, SGraphEditor::FActionMenuClosed InOnMenuClosed);
 
 	/** Called when the path is being edited */
 	void OnSearchBarTextChanged(const FText& NewText);
diff --git a/Engine/Source/Editor/UnrealEd/Private/GraphEditor.cpp b/Engine/Source/Editor/UnrealEd/Private/GraphEditor.cpp
--- a/Engine/Source/Editor/UnrealEd/Private/GraphEditor.cpp
+++ b/Engine/Source/Editor/UnrealEd/Private/GraphEditor.cpp
@@ -26,7 +26,9 @@ TArray< TWeakPtr<SGraphEditor> > SGraphEditor::AllInstances;
 void SGraphEditor::ConstructImplementation( const FArguments& InArgs )
 {
 	FGraphEditorModule& GraphEdModule = FModuleManager::LoadModuleChecked<FGraphEditorModule>(TEXT("GraphEditor"));
-	
+
+	// Upgrade any deprecated delegates from the events argument
+	UpgradeDeprecatedDelegates(const_cast<FArguments&>(InArgs)._GraphEvents);
 
 	// Construct the implementation and make it the contents of this widget.
 	Implementation = GraphEdModule.PRIVATE_MakeGraphEditor( InArgs._AdditionalCommands, 
@@ -57,6 +59,40 @@ void SGraphEditor::ConstructImplementation( const FArguments& InArgs )
 	];
 }
 
+void SGraphEditor::UpgradeDeprecatedDelegates(FGraphEditorEvents& EventsToUpdate)
+{
+	PRAGMA_DISABLE_DEPRECATION_WARNINGS
+	if (EventsToUpdate.OnDropActor.IsBound() && !EventsToUpdate.OnDropActors.IsBound())
+	{
+		EventsToUpdate.OnDropActors = FOnDropActors::CreateLambda([&EventsToUpdate](const TArray< TWeakObjectPtr<AActor> >& Actors, UEdGraph* InGraph, const FVector2f& InDropLocation)
+		{
+			EventsToUpdate.OnDropActor.Execute(Actors, InGraph, FVector2D(InDropLocation));
+		});
+	}
+	if (EventsToUpdate.OnDropStreamingLevel.IsBound() && !EventsToUpdate.OnDropStreamingLevels.IsBound())
+	{
+		EventsToUpdate.OnDropStreamingLevels = FOnDropStreamingLevels::CreateLambda([&EventsToUpdate](const TArray< TWeakObjectPtr<class ULevelStreaming> >& Levels, UEdGraph* InGraph, const FVector2f& InDropLocation)
+		{
+			EventsToUpdate.OnDropStreamingLevel.Execute(Levels, InGraph, FVector2D(InDropLocation));
+		});
+	}
+	if (EventsToUpdate.OnCreateActionMenu.IsBound() && !EventsToUpdate.OnCreateActionMenuAtLocation.IsBound())
+	{
+		EventsToUpdate.OnCreateActionMenuAtLocation = FOnCreateActionMenuAtLocation::CreateLambda([&EventsToUpdate](UEdGraph* InGraph, const FVector2f& InLocation, const TArray<UEdGraphPin*>& InPins, bool bAutoExpand, FActionMenuClosed InMenuClosed)
+		{
+			return EventsToUpdate.OnCreateActionMenu.Execute(InGraph, FVector2D(InLocation), InPins, bAutoExpand, InMenuClosed);
+		});
+	}
+	if (EventsToUpdate.OnSpawnNodeByShortcut.IsBound() && !EventsToUpdate.OnSpawnNodeByShortcutAtLocation.IsBound())
+	{
+		EventsToUpdate.OnSpawnNodeByShortcutAtLocation = FOnSpawnNodeByShortcutAtLocation::CreateLambda([&EventsToUpdate](FInputChord InInputChord, const FVector2f& InLocation)
+		{
+			return EventsToUpdate.OnSpawnNodeByShortcut.Execute(InInputChord, FVector2D(InLocation));
+		});
+	}
+	PRAGMA_ENABLE_DEPRECATION_WARNINGS
+}
+
 
 /**
  * Loads the GraphEditorModule and constructs a GraphEditor as a child of this widget.
diff --git a/Engine/Source/Editor/UnrealEd/Public/GraphEditor.h b/Engine/Source/Editor/UnrealEd/Public/GraphEditor.h
--- a/Engine/Source/Editor/UnrealEd/Public/GraphEditor.h
+++ b/Engine/Source/Editor/UnrealEd/Public/GraphEditor.h
@@ -105,17 +105,25 @@ public:
 
 	DECLARE_DELEGATE_OneParam( FOnFocused, const TSharedRef<SGraphEditor>& );
 
+	UE_DEPRECATED(5.6, "Please use the delegate accepting FVector2f.")
 	DECLARE_DELEGATE_ThreeParams( FOnDropActor, const TArray< TWeakObjectPtr<class AActor> >&, class UEdGraph*, const FVector2D& );
-	
+	DECLARE_DELEGATE_ThreeParams( FOnDropActors, const TArray< TWeakObjectPtr<class AActor> >&, class UEdGraph*, const FVector2f& );
+
+	UE_DEPRECATED(5.6, "Please use the delegate accepting FVector2f.")
 	DECLARE_DELEGATE_ThreeParams( FOnDropStreamingLevel, const TArray< TWeakObjectPtr<class ULevelStreaming> >&, class UEdGraph*, const FVector2D& );
+	DECLARE_DELEGATE_ThreeParams( FOnDropStreamingLevels, const TArray< TWeakObjectPtr<class ULevelStreaming> >&, class UEdGraph*, const FVector2f& );
 
 	DECLARE_DELEGATE( FActionMenuClosed );
 
+	UE_DEPRECATED(5.6, "Please use the delegate accepting FVector2f.")
 	DECLARE_DELEGATE_RetVal_FiveParams( FActionMenuContent, FOnCreateActionMenu, UEdGraph*, const FVector2D&, const TArray<UEdGraphPin*>&, bool, FActionMenuClosed );
+	DECLARE_DELEGATE_RetVal_FiveParams( FActionMenuContent, FOnCreateActionMenuAtLocation, UEdGraph*, const FVector2f&, const TArray<UEdGraphPin*>&, bool, FActionMenuClosed );
 
 	DECLARE_DELEGATE_RetVal_FiveParams( FActionMenuContent, FOnCreateNodeOrPinMenu, UEdGraph*, const UEdGraphNode*, const UEdGraphPin*, FMenuBuilder*, bool);
 
+	UE_DEPRECATED(5.6, "Please use the delegate accepting FVector2f.")
 	DECLARE_DELEGATE_RetVal_TwoParams( FReply, FOnSpawnNodeByShortcut, FInputChord, const FVector2D& );
+	DECLARE_DELEGATE_RetVal_TwoParams( FReply, FOnSpawnNodeByShortcutAtLocation, FInputChord, const FVector2f& );
 
 	DECLARE_DELEGATE( FOnNodeSpawnedByKeymap );
 
@@ -130,6 +138,12 @@ public:
 	/** Info about events occurring in/on the graph */
 	struct FGraphEditorEvents
 	{
+PRAGMA_DISABLE_DEPRECATION_WARNINGS
+		FGraphEditorEvents() = default;
+		FGraphEditorEvents(const FGraphEditorEvents& InEvents) = default;
+		~FGraphEditorEvents() = default;
+PRAGMA_ENABLE_DEPRECATION_WARNINGS
+
 		/** Called when selection changes */
 		FOnSelectionChanged OnSelectionChanged;
 		/** Called when a node is double clicked */
@@ -137,19 +151,23 @@ public:
 		/* Called when focus moves to graph */
 		FOnFocused OnFocused;
 		/* Called when an actor is dropped on graph */
-		FOnDropActor OnDropActor;
+		PRAGMA_DISABLE_DEPRECATION_WARNINGS FOnDropActor OnDropActor; PRAGMA_ENABLE_DEPRECATION_WARNINGS
+		FOnDropActors OnDropActors;
 		/* Called when a streaming level is dropped on graph */
-		FOnDropStreamingLevel OnDropStreamingLevel;
+		PRAGMA_DISABLE_DEPRECATION_WARNINGS FOnDropStreamingLevel OnDropStreamingLevel; PRAGMA_ENABLE_DEPRECATION_WARNINGS
+		FOnDropStreamingLevels OnDropStreamingLevels;
 		/** Called when text is being committed on the graph to verify */
 		FOnNodeVerifyTextCommit OnVerifyTextCommit;
 		/** Called when text is committed on the graph */
 		FOnNodeTextCommitted OnTextCommitted;
 		/** Called to create context menu for right clicking in empty area */
-		FOnCreateActionMenu OnCreateActionMenu;
+		PRAGMA_DISABLE_DEPRECATION_WARNINGS FOnCreateActionMenu OnCreateActionMenu; PRAGMA_ENABLE_DEPRECATION_WARNINGS
+		FOnCreateActionMenuAtLocation OnCreateActionMenuAtLocation;
 		/** Called to create context menu for right clicking a node or pin, same parameters as GetContextMenuActions on schema */
 		FOnCreateNodeOrPinMenu OnCreateNodeOrPinMenu;
 		/** Called to spawn a node in the graph using a shortcut */
-		FOnSpawnNodeByShortcut OnSpawnNodeByShortcut;
+		PRAGMA_DISABLE_DEPRECATION_WARNINGS FOnSpawnNodeByShortcut OnSpawnNodeByShortcut; PRAGMA_ENABLE_DEPRECATION_WARNINGS
+		FOnSpawnNodeByShortcutAtLocation OnSpawnNodeByShortcutAtLocation;
 		/** Called when a keymap spawns a node */
 		FOnNodeSpawnedByKeymap OnNodeSpawnedByKeymap;
 		/** Called when the user generates a warning tooltip because a connection was invalid */
@@ -189,8 +207,10 @@ public:
 	
 		SLATE_ARGUMENT( TSharedPtr<TArray<FDiffSingleResult>>, DiffResults )
 		SLATE_ATTRIBUTE( int32, FocusedDiffResult )
-	
+
+		PRAGMA_DISABLE_DEPRECATION_WARNINGS
 		SLATE_ARGUMENT( FGraphEditorEvents, GraphEvents)
+		PRAGMA_ENABLE_DEPRECATION_WARNINGS
 		SLATE_ARGUMENT( bool, AutoExpandActionMenu )
 		SLATE_ARGUMENT( TWeakPtr<FAssetEditorToolkit>, AssetEditorToolkit)
 		SLATE_EVENT(FSimpleDelegate, OnNavigateHistoryBack)
@@ -716,7 +736,8 @@ protected:
 private:
 	static void RegisterGraphEditor(const TSharedRef<SGraphEditor>& InGraphEditor);
 
-	void ConstructImplementation( const FArguments& InArgs );
+	void ConstructImplementation(const FArguments& InArgs);
+	static void UpgradeDeprecatedDelegates(FGraphEditorEvents& EventsToUpdate);
 protected:
 	/** The Graph we are currently editing */
 	UEdGraph* EdGraphObj;
