Fix inconsistent classflags when a DOBP is compiled at the same time as its super (the parent class has its classflags conformed later in the compilation process)
#rnx

diff --git a/Engine/Source/Editor/Kismet/Private/BlueprintCompilationManager.cpp b/Engine/Source/Editor/Kismet/Private/BlueprintCompilationManager.cpp
--- a/Engine/Source/Editor/Kismet/Private/BlueprintCompilationManager.cpp
+++ b/Engine/Source/Editor/Kismet/Private/BlueprintCompilationManager.cpp
@@ -778,14 +778,6 @@ void FBlueprintCompilationManagerImpl::FlushCompilationQueueImpl(bool bSuppressB
 					QueuedBP->bIsRegeneratingOnLoad = false;
 				}
 
-				// No actual compilation work to be done, but try to conform the class and fix up anything that might need to be updated if the native base class has changed in any way
-				FKismetEditorUtilities::ConformBlueprintFlagsAndComponents(QueuedBP);
-
-				if (QueuedBP->GeneratedClass)
-				{
-					FBlueprintEditorUtils::RecreateClassMetaData(QueuedBP, QueuedBP->GeneratedClass, true);
-				}
-
 				QueuedRequests.RemoveAtSwap(I);
 				--I;
 			}
@@ -1522,6 +1514,19 @@ void FBlueprintCompilationManagerImpl::FlushCompilationQueueImpl(bool bSuppressB
 					BPGC->SetDefaultObject(nullptr);
 					BPGC->GetDefaultObject(true);
 				}
+				
+				if(CompilerData.IsSkeletonOnly() && BP->bIsRegeneratingOnLoad)
+				{
+#if DO_GUARD_SLOW
+					ensure(FBlueprintEditorUtils::IsDataOnlyBlueprint(BP));
+#endif
+
+					FKismetEditorUtilities::ConformBlueprintFlagsAndComponents(BP);
+					if (BP->GeneratedClass)
+					{
+						FBlueprintEditorUtils::RecreateClassMetaData(BP, BP->GeneratedClass, true);
+					}
+				}
 			}
 			else
 			{
