Supporting undo after a prefab save(reinstantiation)

[RN] Support undo after a prefab save

diff --git a/Engine/Source/Editor/Kismet/Private/BlueprintCompilationManager.cpp b/Engine/Source/Editor/Kismet/Private/BlueprintCompilationManager.cpp
--- a/Engine/Source/Editor/Kismet/Private/BlueprintCompilationManager.cpp
+++ b/Engine/Source/Editor/Kismet/Private/BlueprintCompilationManager.cpp
@@ -2537,7 +2537,7 @@ void FBlueprintCompilationManagerImpl::BuildDSOMap(UObject* OldObject, UObject*
 		{
 			OldSubobjects.Add(Object);
 		}
-	}, false);
+	}, /*bIncludeNestedObjects*/ false, /*ExclusionFlags*/ RF_MirroredGarbage);
 
 	for(UObject* OldSubobject : OldSubobjects)
 	{
@@ -2552,7 +2552,11 @@ void FBlueprintCompilationManagerImpl::BuildDSOMap(UObject* OldObject, UObject*
 
 		// It may seem aggressive to ensure here, but we should always have a new version of the DSO. If that's
 		// not the case then some testing will need to be done on client of OutOldToNewDSO
-		OutOldToNewDSO.Add(OldSubobject, NewSubobject);
+		// If the new archetype does not exist, do not add it into the map. Otherwise we might reset the value inside a undo/redo transaction
+		if (NewSubobject)
+		{
+			OutOldToNewDSO.Add(OldSubobject, NewSubobject);
+		}
 		BuildDSOMap(OldSubobject, NewSubobject, InOldToNewClassMap, OutOldToNewDSO);
 	}
 }
