Fall back to UBPGC when UBlueprint is not available
#rnx

diff --git a/Engine/Source/Editor/Kismet/Private/BlueprintCompilationManager.cpp b/Engine/Source/Editor/Kismet/Private/BlueprintCompilationManager.cpp
--- a/Engine/Source/Editor/Kismet/Private/BlueprintCompilationManager.cpp
+++ b/Engine/Source/Editor/Kismet/Private/BlueprintCompilationManager.cpp
@@ -2911,9 +2911,20 @@ void FBlueprintCompilationManagerImpl::ReinstanceBatch(TArray<FReinstancingJob>&
 			// patch new class and cdo into linker table:
 			if (ReinstancingJob.Reinstancer.IsValid())
 			{
-				UBlueprint* CurrentBP = CastChecked<UBlueprint>(BPGClass->ClassGeneratedBy);
-				UObject* NewCDO = CurrentBP->GeneratedClass->GetDefaultObject(false);
-				FLinkerLoad* CurrentLinker = CurrentBP->GetLinker();
+				UBlueprint* CurrentBP = Cast<UBlueprint>(BPGClass->ClassGeneratedBy);
+				UObject* NewCDO = nullptr;
+				FLinkerLoad* CurrentLinker = nullptr;
+				if(CurrentBP)
+				{
+					NewCDO = CurrentBP->GeneratedClass->GetDefaultObject(false);
+					CurrentLinker = CurrentBP->GetLinker();
+				}
+				else if(ReinstancingJob.Reinstancer->ClassToReinstance)
+				{
+					NewCDO = ReinstancingJob.Reinstancer->ClassToReinstance->GetDefaultObject(false);
+					CurrentLinker = ReinstancingJob.Reinstancer->ClassToReinstance->GetLinker();
+				}
+
 				if (CurrentLinker)
 				{
 					int32 OldCDOIndex = INDEX_NONE;
