[Editor UI] Centered the "Defined Property Flags" text in Blueprint Details Customization

#rnx

diff --git a/Engine/Source/Editor/Kismet/Private/BlueprintDetailsCustomization.cpp b/Engine/Source/Editor/Kismet/Private/BlueprintDetailsCustomization.cpp
--- a/Engine/Source/Editor/Kismet/Private/BlueprintDetailsCustomization.cpp
+++ b/Engine/Source/Editor/Kismet/Private/BlueprintDetailsCustomization.cpp
@@ -1063,6 +1063,8 @@ void FBlueprintVarActionDetails::CustomizeDetails( IDetailLayoutBuilder& DetailL
 		if (!ErrorMessage.IsEmpty())
 		{
 			DefaultValueCategory.AddCustomRow( ErrorMessage )
+			.WholeRowWidget
+			.VAlign(VAlign_Center)
 			[
 				SNew(STextBlock)
 				.ToolTipText(ErrorMessage)
@@ -1329,6 +1331,7 @@ void FBlueprintVarActionDetails::CustomizeDetails( IDetailLayoutBuilder& DetailL
 
 		Category.AddCustomRow(LOCTEXT("DefinedPropertyFlags", "Defined Property Flags"), true)
 		.WholeRowWidget
+		.VAlign(VAlign_Center)
 		[
 			SNew(STextBlock)
 			.ToolTip(PropertyFlagsTooltip)
