Make sure we capture class dependencies for any function objects in our object graph

#rnx

diff --git a/Engine/Source/Editor/Kismet/Private/Blueprints/BlueprintDependencies.cpp b/Engine/Source/Editor/Kismet/Private/Blueprints/BlueprintDependencies.cpp
--- a/Engine/Source/Editor/Kismet/Private/Blueprints/BlueprintDependencies.cpp
+++ b/Engine/Source/Editor/Kismet/Private/Blueprints/BlueprintDependencies.cpp
@@ -280,6 +280,10 @@ static void UE::Private::GetAllDependencies(
 		{
 			AddStructDependencyImpl(AsStruct, ReferencedStructs);
 		}
+		else if(const UFunction* AsFunction = Cast<UFunction>(Object))
+		{
+			AddClassDependencyImpl(AsFunction->GetOwnerClass(), ReferencedClasses);
+		}
 		else if (!Package->HasAnyPackageFlags(PKG_CompiledIn))
 		{
 			ReferencedAssets.Add(Package);
