IncrementalCook: More aggressive capture of soft dependencies for blueprints

#rnx

diff --git a/Engine/Source/Editor/Kismet/Private/Blueprints/BlueprintDependencies.cpp b/Engine/Source/Editor/Kismet/Private/Blueprints/BlueprintDependencies.cpp
--- a/Engine/Source/Editor/Kismet/Private/Blueprints/BlueprintDependencies.cpp
+++ b/Engine/Source/Editor/Kismet/Private/Blueprints/BlueprintDependencies.cpp
@@ -46,22 +46,21 @@ namespace UE::Private
 constexpr int32 BlueprintCookDependenciesVersion = 3;
 
 // gathering helpers:
-static TArray<const UObject*> GetAllImportedObjects(TConstArrayView<const UObject*> Roots);
-static TArray<const UObject*> GetAllImportedObjects(const UBlueprint* ForBP);
+static void GetAllImportedObjects(TConstArrayView<const UObject*> Roots, TArray<const UObject*>& OutImports, TSet<FName>& OutSoftImports);
+static void GetAllImportedObjects(const UBlueprint* ForBP, TArray<const UObject*>& OutImports, TSet<FName>& OutSoftImports);
 // helpers for handling transitiveness of dependencies (e.g. if i use a struct, i also depend on its member structs):
 static void AddStructDependencyImpl(const UScriptStruct* Struct, TSet<const UScriptStruct*>& ReferencedStructs);
 static void AddClassDependencyImpl(const UClass* Class, TSet<const UClass*>& ReferencedClasses);
 // returns all the dependencies of a bp, for recording/generating a snapshot of dependencies:
 static void GetAllDependencies(
 	const UBlueprint* ForBP, 
-	TArray<const UPackage*>& AssetDependencies, 
+	TArray<FName>& AssetDependencies, 
 	TArray<const UClass*>& NativeClassDependencies,
 	TArray<const UScriptStruct*>& NativeStructDependencies);
 // Native struct and class canonization:
 static FString HashNativeStruct(const UScriptStruct* Struct);
 static FString HashNativeClass( const UClass* NativeClass);
 // Asset canonization:
-static FString HashPackage( const UPackage* Package );
 static FString HashPackageFile(FName PackageName);
 
 // cache of dependency hashes, so they repeatedly recalculated:
@@ -78,14 +77,14 @@ class FBPDependencyCache
 public:
 	static FBPDependencyCacheEntry* LookupCache(const UClass* Class);
 	static FBPDependencyCacheEntry* LookupCache(const UScriptStruct* Struct);
-	static FBPDependencyCacheEntry* LookupCache(const UPackage* Package);
+	static FBPDependencyCacheEntry* LookupCache(FName PackageName);
 private:
 	static FBPDependencyCache& GetCache();
 	static TUniquePtr<FBPDependencyCache> GBPDependencyCache;
 
 	TMap<const UClass*, FBPDependencyCacheEntry> ClassToHash;
 	TMap<const UScriptStruct*, FBPDependencyCacheEntry> StructToHash;
-	TMap<const UPackage*, FBPDependencyCacheEntry> PackageToHash;
+	TMap<FName, FBPDependencyCacheEntry> PackageToHash;
 };
 
 // Filesystem helper:
@@ -99,17 +98,18 @@ static bool LoadBPCookDependenciesImpl(FCbFieldViewIterator Args, UE::Cook::FCoo
 static void SaveBPCookDepenenciesImpl(FCbWriter& Writer, const FBlueprintDependencies& Dependencies);
 }
 
-static TArray<const UObject*> UE::Private::GetAllImportedObjects(TConstArrayView<const UObject*> Roots)
+static void UE::Private::GetAllImportedObjects(TConstArrayView<const UObject*> Roots, TArray<const UObject*>& OutImports, TSet<FName>& OutSoftImports)
 {
 	// This function just gathers all of the UObjects the Root objects (or their inners)
 	// depend upon. Inners are not currently tested for reachability. Tautological self reference
 	// (e.g. I depend on myself or I depend on my inner) are not reported - only external objects:
 	struct FExternalReferenceFinder : public FArchiveUObject
 	{
-		FExternalReferenceFinder(const UObject* Obj, TSet<const UObject*>& InReferences, const TSet<const UObject*>& InRoots)
+		FExternalReferenceFinder(const UObject* Obj, TSet<const UObject*>& InReferences, TSet<FName>& InSoftReferences, const TSet<const UObject*>& InRoots)
 			: FArchiveUObject()
 			, SearchRoots(InRoots)
 			, References(InReferences)
+			, SoftReferences(InSoftReferences)
 		{
 			SetIsSaving(true);
 			SetShouldSkipCompilingAssets(true);
@@ -133,7 +133,7 @@ static TArray<const UObject*> UE::Private::GetAllImportedObjects(TConstArrayView
 		bool ShouldTraverseProperty()
 		{
 			FProperty* SerializedProperty = GetSerializedProperty();
-			if (!SerializedProperty || !SerializedProperty->HasAnyFlags(RF_Transient))
+			if (!SerializedProperty || !SerializedProperty->HasAnyPropertyFlags(CPF_Transient))
 			{
 				return true;
 			}
@@ -142,8 +142,23 @@ static TArray<const UObject*> UE::Private::GetAllImportedObjects(TConstArrayView
 
 		virtual FArchive& operator<<(FLazyObjectPtr& Value) override { return ShouldTraverseProperty() ? FArchiveUObject::operator<<(Value) : *this; }
 		virtual FArchive& operator<<(FObjectPtr& Value) override { return ShouldTraverseProperty() ? FArchiveUObject::operator<<(Value) : *this; }
-		virtual FArchive& operator<<(FSoftObjectPtr& Value) override { return ShouldTraverseProperty() ? FArchiveUObject::operator<<(Value) : *this; }
-		virtual FArchive& operator<<(FSoftObjectPath& Value) override { return ShouldTraverseProperty() ? FArchiveUObject::operator<<(Value) : *this; }
+		virtual FArchive& operator<<(FSoftObjectPtr& Value) override 
+		{ 
+			const FSoftObjectPath& Path = Value.ToSoftObjectPath();
+			if( !Path.IsNull() && ShouldTraverseProperty() )
+			{
+				SoftReferences.Add(Path.GetLongPackageFName());
+			}
+			return *this; 
+		}
+		virtual FArchive& operator<<(FSoftObjectPath& Value) override 
+		{ 
+			if( !Value.IsNull() && ShouldTraverseProperty() )
+			{
+				SoftReferences.Add(Value.GetLongPackageFName());
+			}
+			return *this; 
+		}
 		virtual FArchive& operator<<(FWeakObjectPtr& Value) override { return ShouldTraverseProperty() ? FArchiveUObject::operator<<(Value) : *this; }
 
 		virtual FArchive& operator<<(UObject*& ObjRef) override
@@ -178,6 +193,7 @@ static TArray<const UObject*> UE::Private::GetAllImportedObjects(TConstArrayView
 
 		const TSet<const UObject*>& SearchRoots;
 		TSet<const UObject*>& References;
+		TSet<FName>& SoftReferences;
 	};
 
 	TSet<const UObject*> RootSet(Roots);
@@ -204,15 +220,15 @@ static TArray<const UObject*> UE::Private::GetAllImportedObjects(TConstArrayView
 
 	for(const UObject* Obj: Objects)
 	{
-		FExternalReferenceFinder References(Obj, Result, RootSet);
+		FExternalReferenceFinder References(Obj, Result, OutSoftImports, RootSet);
 	}
 	
-	return Result.Array();
+	OutImports = Result.Array();
 }
 
-static TArray<const UObject*> UE::Private::GetAllImportedObjects( const UBlueprint* ForBP)
+static void UE::Private::GetAllImportedObjects( const UBlueprint* ForBP, TArray<const UObject*>& OutImports, TSet<FName>& OutSoftImports)
 {
-	return GetAllImportedObjects({ForBP, ForBP->GeneratedClass, ForBP->GeneratedClass->GetDefaultObject(false) });
+	return GetAllImportedObjects({ForBP, ForBP->GeneratedClass, ForBP->GeneratedClass->GetDefaultObject(false) }, OutImports, OutSoftImports);
 }
 
 static void UE::Private::AddStructDependencyImpl(const UScriptStruct* Struct, TSet<const UScriptStruct*>& ReferencedStructs)
@@ -255,20 +271,21 @@ static void UE::Private::AddClassDependencyImpl(const UClass* Class, TSet<const
 
 static void UE::Private::GetAllDependencies(
 	const UBlueprint* ForBP,
-	TArray<const UPackage*>& AssetDependencies, 
+	TArray<FName>& AssetDependencies, 
 	TArray<const UClass*>& NativeClassDependencies,
 	TArray<const UScriptStruct*>& NativeStructDependencies)
 {
 	TRACE_CPUPROFILER_EVENT_SCOPE(BlueprintDependencies::GetAllDependencies);
 
 	// gather all references in the BP, its class, and the CDO, along with any subobjects:
-	TArray<const UObject*> ImportedObjects = GetAllImportedObjects(ForBP);
+	TArray<const UObject*> ImportedObjects;
+	TSet<FName> ReferencedAssets;
+	GetAllImportedObjects(ForBP, ImportedObjects, ReferencedAssets);
 
 	// honor transitive dependencies - e.g. inherited classes, inherited structs, aggregate structs
 	// for now all class and struct references are fully transitive:
 	TSet<const UScriptStruct*> ReferencedStructs;
 	TSet<const UClass*> ReferencedClasses;
-	TSet<const UPackage*> ReferencedAssets;
 	for(const UObject* Object : ImportedObjects)
 	{
 		const UPackage* Package = Object->GetPackage();
@@ -286,7 +303,7 @@ static void UE::Private::GetAllDependencies(
 		}
 		else if (!Package->HasAnyPackageFlags(PKG_CompiledIn))
 		{
-			ReferencedAssets.Add(Package);
+			ReferencedAssets.Add(Package->GetFName());
 		}
 	}
 
@@ -299,7 +316,7 @@ static void UE::Private::GetAllDependencies(
 		}
 		else
 		{
-			ReferencedAssets.Add(Class->GetPackage());
+			ReferencedAssets.Add(Class->GetPackage()->GetFName());
 		}
 	}
 
@@ -311,7 +328,7 @@ static void UE::Private::GetAllDependencies(
 		}
 		else
 		{
-			ReferencedAssets.Add(Struct->GetPackage());
+			ReferencedAssets.Add(Struct->GetPackage()->GetFName());
 		}
 	}
 	
@@ -337,9 +354,9 @@ static void UE::Private::GetAllDependencies(
 	};
 	NativeStructDependencies.Sort(SortTopLevelObject);
 	NativeClassDependencies.Sort(SortTopLevelObject);
-	AssetDependencies.Sort([](const UPackage& A, const UPackage& B)
+	AssetDependencies.Sort([](FName A, FName B)
 	{
-		return A.GetFName().LexicalLess(B.GetFName());
+		return A.LexicalLess(B);
 	});
 }
 
@@ -404,26 +421,15 @@ static FString UE::Private::HashNativeClass( const UClass* NativeClass)
 	return CachedValue->Hash;
 }
 
-static FString UE::Private::HashPackage(const UPackage* Package)
+static FString UE::Private::HashPackageFile(FName PackageName)
 {
 	TRACE_CPUPROFILER_EVENT_SCOPE(BlueprintDependencies::HashPackage);
-	FBPDependencyCacheEntry* CachedValue = FBPDependencyCache::LookupCache(Package);
+	FBPDependencyCacheEntry* CachedValue = FBPDependencyCache::LookupCache(PackageName);
 	if(CachedValue->Hash.Len() > 0)
 	{
 		return CachedValue->Hash;
 	}
 
-	CachedValue->Hash = UE::Private::HashPackageFile(Package->GetFName());
-	
-#if UE_STORE_DEPENDENCY_SNAPSHOT
-	CachedValue->Source = "File hashed on disk";
-#endif
-
-	return CachedValue->Hash;
-}
-
-static FString UE::Private::HashPackageFile(FName PackageName)
-{
 	IAssetRegistry& AssetRegistry = FModuleManager::LoadModuleChecked<FAssetRegistryModule>("AssetRegistry").Get();
 	TOptional<FAssetPackageData> PackageData = AssetRegistry.GetAssetPackageDataCopy(PackageName);
 	if (!PackageData)
@@ -431,7 +437,13 @@ static FString UE::Private::HashPackageFile(FName PackageName)
 		return FString();
 	}
 	FIoHash PackageHash = PackageData->GetPackageSavedHash();
-	return LexToString(PackageHash);
+	CachedValue->Hash = LexToString(PackageHash);
+	
+#if UE_STORE_DEPENDENCY_SNAPSHOT
+	CachedValue->Source = "File hashed on disk";
+#endif
+
+	return CachedValue->Hash;
 }
 
 UE::Private::FBPDependencyCacheEntry* UE::Private::FBPDependencyCache::LookupCache(const UClass* Class)
@@ -448,10 +460,10 @@ UE::Private::FBPDependencyCacheEntry* UE::Private::FBPDependencyCache::LookupCac
 	return &Value;
 }
 
-UE::Private::FBPDependencyCacheEntry* UE::Private::FBPDependencyCache::LookupCache(const UPackage* Package)
+UE::Private::FBPDependencyCacheEntry* UE::Private::FBPDependencyCache::LookupCache(FName PackageName)
 {
 	FBPDependencyCache& Cache = GetCache();
-	FBPDependencyCacheEntry& Value = Cache.PackageToHash.FindOrAdd(Package);
+	FBPDependencyCacheEntry& Value = Cache.PackageToHash.FindOrAdd(PackageName);
 	return &Value;
 }
 
@@ -573,7 +585,7 @@ UE_COOK_DEPENDENCY_FUNCTION(ValidateBPCookDependenciesImpl, UE::Private::Validat
 
 void UE::Private::GenerateBlueprintDependencies(const UBlueprint* ForBP, FBlueprintDependencies& OutDependencies)
 {
-	TArray<const UPackage*> AssetDependencies;
+	TArray<FName> AssetDependencies;
 	TArray<const UClass*> NativeClassDependencies;
 	TArray<const UScriptStruct*> NativeStructDependencies;
 	UE::Private::GetAllDependencies(ForBP, AssetDependencies, NativeClassDependencies, NativeStructDependencies);
@@ -583,13 +595,13 @@ void UE::Private::GenerateBlueprintDependencies(const UBlueprint* ForBP, FBluepr
 	// gather dependencies on functions/structs/classes/objects:
 	// treat class as transitive, in case we're using them for type comparisons:
 	// treat struct dependencies as aggregate, in case inner structs change:
-	for(const UPackage* Package : AssetDependencies)
+	for(FName PackageName : AssetDependencies)
 	{
 		FBlueprintDependency Dependency = {
 			.DependencyType = EBPDependencyType::Asset,
-			.PackageName = Package->GetFName(),
+			.PackageName = PackageName,
 			.NativeObjectName = FName(),
-			.Hash = UE::Private::HashPackage(Package)
+			.Hash = UE::Private::HashPackageFile(PackageName)
 		};
 		
 		OutDependencies.BlueprintDependencies.Emplace(
