Find-in-Blueprints: When a BP member variable is a struct, its members' default values are now searched and returned as results.

diff --git a/Engine/Source/Editor/Kismet/Private/ImaginaryBlueprintData.cpp b/Engine/Source/Editor/Kismet/Private/ImaginaryBlueprintData.cpp
--- a/Engine/Source/Editor/Kismet/Private/ImaginaryBlueprintData.cpp
+++ b/Engine/Source/Editor/Kismet/Private/ImaginaryBlueprintData.cpp
@@ -294,7 +294,14 @@ void FImaginaryFiBData::ParseJsonValue(FText InKey, FText InDisplayKey, TSharedP
 	}
 	else if (InJsonValue->Type == EJson::Object)
 	{
-		TSharedPtr< FCategorySectionHelper, ESPMode::ThreadSafe > SubObjectCategory = MakeShareable(new FCategorySectionHelper(AsShared(), InJsonValue->AsObject(), LookupTablePtr, InDisplayKey, bIsInArray));
+		// Create a FCategorySectionHelper for Json objects, which are Json representations of structs, instanced objects, array elements.
+		// bIsTagAndValueCategory represents whether a match with the category's contents should return its parent as a search result.
+		// Let the following category return this as search result when:
+		// - The following category represents an array element, return this (the array)
+		// - The following category represents the default value for a struct or object, return this (the default value)
+		// - The following category represents a nested element for one of the above
+		const bool bIsTagAndValueCategory = bIsInArray || IsTagAndValueCategory() || InDisplayKey.EqualTo(FFindInBlueprintSearchTags::FiB_DefaultValue);
+		TSharedPtr< FCategorySectionHelper, ESPMode::ThreadSafe > SubObjectCategory = MakeShareable(new FCategorySectionHelper(AsShared(), InJsonValue->AsObject(), LookupTablePtr, InDisplayKey, bIsTagAndValueCategory));
 		SubObjectCategory->ParseAllChildData(SearchabilityStatus);
 		ParsedChildData.Add(SubObjectCategory);
 	}
@@ -489,6 +496,25 @@ bool FCategorySectionHelper::CanCallFilter(ESearchQueryFilter InSearchQueryFilte
 	return true;
 }
 
+FString FCategorySectionHelper::GetCategoryPath() const
+{
+	const FCategorySectionHelper* OuterCategory = this;
+	FString Path;
+	do
+	{
+		// Prepend category
+		Path = FString::Printf(TEXT("%s.%s"), *OuterCategory->GetCategoryFunctionName(), *Path);
+
+		// Resolve outer category, or null if the outer isn't a category
+		OuterCategory = (OuterCategory->GetOuter().Pin().IsValid() && OuterCategory->GetOuter().Pin()->IsCategory())
+			? (FCategorySectionHelper*)(OuterCategory->GetOuter().Pin().Get())
+			: nullptr;
+	} while (OuterCategory);
+
+	// Remove redundant . at the end
+	return Path.LeftChop(1);
+}
+
 FSearchResult FCategorySectionHelper::CreateSearchResult_Internal(FSearchResult InTemplate) const
 {
 	if (InTemplate.IsValid())
diff --git a/Engine/Source/Editor/Kismet/Public/ImaginaryBlueprintData.h b/Engine/Source/Editor/Kismet/Public/ImaginaryBlueprintData.h
--- a/Engine/Source/Editor/Kismet/Public/ImaginaryBlueprintData.h
+++ b/Engine/Source/Editor/Kismet/Public/ImaginaryBlueprintData.h
@@ -317,6 +317,10 @@ public:
 		return CategoryName.BuildSourceString();
 	}
 
+	/** Recursively visits outer categories and constructs a path name: Outermost.Outer.ThisCategoryName. 
+	    Intended for debugging. Halts as soon as a non-category is visited. */
+	FString GetCategoryPath() const;
+
 protected:
 	/** FImaginaryFiBData Interface */
 	virtual FSearchResult CreateSearchResult_Internal(FSearchResult InTemplate) const override;
