[UE] Fix editor crash after clicking on a deleted graph node hyperlink in the Blueprint editor's compiler results tab.

diff --git a/Engine/Source/Editor/Kismet/Private/BlueprintEditor.cpp b/Engine/Source/Editor/Kismet/Private/BlueprintEditor.cpp
--- a/Engine/Source/Editor/Kismet/Private/BlueprintEditor.cpp
+++ b/Engine/Source/Editor/Kismet/Private/BlueprintEditor.cpp
@@ -4624,8 +4624,13 @@ void FBlueprintEditor::JumpToNode(const UEdGraphNode* Node, bool bRequestRename)
 	}
 	else
 	{
-		// Open a graph editor and jump to the node
-		GraphEditor = OpenGraphAndBringToFront(Node->GetGraph());
+		// Ensure that the node is owned by this Blueprint before attempting to open its outer graph.
+		const UBlueprint* Blueprint = FBlueprintEditorUtils::FindBlueprintForNode(Node);
+		if (Blueprint && Blueprint == GetBlueprintObj())
+		{
+			// Open a graph editor and jump to the node
+			GraphEditor = OpenGraphAndBringToFront(Node->GetGraph());
+		}
 	}
 
 	if (GraphEditor.IsValid())
