[UE] Fix editor crash after clicking on a deleted graph node hyperlink in the Blueprint editor's compiler results tab.

diff --git a/Engine/Source/Editor/Kismet/Private/BlueprintEditor.cpp b/Engine/Source/Editor/Kismet/Private/BlueprintEditor.cpp
--- a/Engine/Source/Editor/Kismet/Private/BlueprintEditor.cpp
+++ b/Engine/Source/Editor/Kismet/Private/BlueprintEditor.cpp
@@ -4624,9 +4624,8 @@ void FBlueprintEditor::JumpToNode(const UEdGraphNode* Node, bool bRequestRename)
 	}
 	else
 	{
-		// Ensure that the node is owned by this Blueprint before attempting to open its outer graph.
-		const UBlueprint* Blueprint = FBlueprintEditorUtils::FindBlueprintForNode(Node);
-		if (Blueprint && Blueprint == GetBlueprintObj())
+		// Ensure that the node is owned by a Blueprint before attempting to open its outer graph, which requires this.
+		if (FBlueprintEditorUtils::FindBlueprintForNode(Node) != nullptr)
 		{
 			// Open a graph editor and jump to the node
 			GraphEditor = OpenGraphAndBringToFront(Node->GetGraph());
