[Bookmark Icon] Add back the functionality that the bookmark icon will change if a bookmark is active on that graph

diff --git a/Engine/Source/Editor/EditorStyle/Private/StarshipStyle.cpp b/Engine/Source/Editor/EditorStyle/Private/StarshipStyle.cpp
--- a/Engine/Source/Editor/EditorStyle/Private/StarshipStyle.cpp
+++ b/Engine/Source/Editor/EditorStyle/Private/StarshipStyle.cpp
@@ -4557,6 +4557,7 @@ void FStarshipEditorStyle::FStyle::SetupGraphEditorStyles()
 		
 		Set( "GraphEditor.ToggleHideUnrelatedNodes", new IMAGE_BRUSH_SVG( "Starship/GraphEditors/HideUnrelated", Icon20x20 ) );
 		Set( "GraphEditor.Bookmark", new IMAGE_BRUSH_SVG("Starship/Common/Bookmarks", Icon16x16) );
+		Set( "GraphEditor.Bookmark.Outline", new IMAGE_BRUSH_SVG("Starship/Common/BookmarksOutline", Icon16x16));
 
 		// Graph editor widgets
 		{
diff --git a/Engine/Source/Editor/Kismet/Private/SGraphTitleBarAddNewBookmark.cpp b/Engine/Source/Editor/Kismet/Private/SGraphTitleBarAddNewBookmark.cpp
--- a/Engine/Source/Editor/Kismet/Private/SGraphTitleBarAddNewBookmark.cpp
+++ b/Engine/Source/Editor/Kismet/Private/SGraphTitleBarAddNewBookmark.cpp
@@ -43,7 +43,7 @@ void SGraphTitleBarAddNewBookmark::Construct(const FArguments& InArgs)
 	Args.ButtonContent()
 	[
 		SNew(SImage)
-		.Image(FAppStyle::Get().GetBrush("GraphEditor.Bookmark"))
+		.Image(this, &SGraphTitleBarAddNewBookmark::GetAddButtonIcon)
 		.ColorAndOpacity(FSlateColor::UseForeground())
 	]
 	.MenuContent()
@@ -229,7 +229,7 @@ FText SGraphTitleBarAddNewBookmark::GetDefaultNameText() const
 	return ResultText;
 }
 
-FText SGraphTitleBarAddNewBookmark::GetAddButtonGlyph() const
+const FSlateBrush* SGraphTitleBarAddNewBookmark::GetAddButtonIcon() const
 {
 	TSharedPtr<FBlueprintEditor> EditorContext = EditorContextPtr.Pin();
 	if (EditorContext.IsValid())
@@ -238,11 +238,10 @@ FText SGraphTitleBarAddNewBookmark::GetAddButtonGlyph() const
 		EditorContext->GetViewBookmark(CurrentBookmarkId);
 		if (CurrentBookmarkId.IsValid())
 		{
-			return FText::FromString(FString(TEXT("\xf005"))); /*fa-star*/
+			return FAppStyle::Get().GetBrush("GraphEditor.Bookmark"); // Filled bookmark icon
 		}
 	}
-
-	return FText::FromString(FString(TEXT("\xf006"))); /*fa-star-o*/
+	return FAppStyle::Get().GetBrush("GraphEditor.Bookmark.Outline"); // Hollow bookmark icon
 }
 
 FText SGraphTitleBarAddNewBookmark::GetAddButtonLabel() const
diff --git a/Engine/Source/Editor/Kismet/Private/SGraphTitleBarAddNewBookmark.h b/Engine/Source/Editor/Kismet/Private/SGraphTitleBarAddNewBookmark.h
--- a/Engine/Source/Editor/Kismet/Private/SGraphTitleBarAddNewBookmark.h
+++ b/Engine/Source/Editor/Kismet/Private/SGraphTitleBarAddNewBookmark.h
@@ -23,9 +23,9 @@ public:
 	void Construct(const FArguments& InArgs);
 
 protected:
-	FText GetAddButtonGlyph() const;
 	FText GetPopupTitleText() const;
 	FText GetDefaultNameText() const;
+	const FSlateBrush* GetAddButtonIcon() const;
 	FText GetAddButtonLabel() const;
 
 	void OnComboBoxOpened();
