Updated headers using UnrealCodeFixup to make sure dllstorage is on methods/staticvars instead of types. Part 5/n

diff --git a/Engine/Source/Editor/Kismet/Private/UserDefinedEnumEditor.h b/Engine/Source/Editor/Kismet/Private/UserDefinedEnumEditor.h
--- a/Engine/Source/Editor/Kismet/Private/UserDefinedEnumEditor.h
+++ b/Engine/Source/Editor/Kismet/Private/UserDefinedEnumEditor.h
@@ -23,6 +23,8 @@
 #include "UObject/WeakObjectPtr.h"
 #include "UObject/WeakObjectPtrTemplates.h"
 
+#define UE_API KISMET_API
+
 class FDetailWidgetRow;
 class FEditableTextUserDefinedEnum;
 class FEditableTextUserDefinedEnumTooltip;
@@ -32,13 +34,13 @@ class IDetailChildrenBuilder;
 class IDetailLayoutBuilder;
 class SDockTab;
 
-class KISMET_API FUserDefinedEnumEditor : public IUserDefinedEnumEditor
+class FUserDefinedEnumEditor : public IUserDefinedEnumEditor
 {
 	/** App Identifier.*/
-	static const FName UserDefinedEnumEditorAppIdentifier;
+	static UE_API const FName UserDefinedEnumEditorAppIdentifier;
 
 	/**	The tab ids for all the tabs used */
-	static const FName EnumeratorsTabId;
+	static UE_API const FName EnumeratorsTabId;
 	
 	/** Property viewing widget */
 	TSharedPtr<class IDetailsView> PropertyView;
@@ -51,30 +53,30 @@ public:
 	 * @param	InitToolkitHost			When Mode is WorldCentric, this is the level editor instance to spawn this editor within
 	 * @param	EnumToEdit				The user defined enum to edit
 	 */
-	void InitEditor(const EToolkitMode::Type Mode, const TSharedPtr< class IToolkitHost >& InitToolkitHost, UUserDefinedEnum* EnumToEdit);
+	UE_API void InitEditor(const EToolkitMode::Type Mode, const TSharedPtr< class IToolkitHost >& InitToolkitHost, UUserDefinedEnum* EnumToEdit);
 
 	/** Destructor */
-	virtual ~FUserDefinedEnumEditor();
+	UE_API virtual ~FUserDefinedEnumEditor();
 
 	/** IToolkit interface */
-	virtual FName GetToolkitFName() const override;
-	virtual FText GetBaseToolkitName() const override;
-	virtual FText GetToolkitName() const override;
-	virtual FText GetToolkitToolTipText() const override;
-	virtual FString GetWorldCentricTabPrefix() const override;
-	virtual FLinearColor GetWorldCentricTabColorScale() const override;
+	UE_API virtual FName GetToolkitFName() const override;
+	UE_API virtual FText GetBaseToolkitName() const override;
+	UE_API virtual FText GetToolkitName() const override;
+	UE_API virtual FText GetToolkitToolTipText() const override;
+	UE_API virtual FString GetWorldCentricTabPrefix() const override;
+	UE_API virtual FLinearColor GetWorldCentricTabColorScale() const override;
 
-	virtual void RegisterTabSpawners(const TSharedRef<class FTabManager>& TabManager) override;
-	virtual void UnregisterTabSpawners(const TSharedRef<class FTabManager>& TabManager) override;
+	UE_API virtual void RegisterTabSpawners(const TSharedRef<class FTabManager>& TabManager) override;
+	UE_API virtual void UnregisterTabSpawners(const TSharedRef<class FTabManager>& TabManager) override;
 
 protected:
-	TSharedRef<SDockTab> SpawnEnumeratorsTab(const FSpawnTabArgs& Args);
+	UE_API TSharedRef<SDockTab> SpawnEnumeratorsTab(const FSpawnTabArgs& Args);
 
 private:
-	void FillToolbar(FToolBarBuilder& ToolbarBuilder);
+	UE_API void FillToolbar(FToolBarBuilder& ToolbarBuilder);
 
 	/** Handles new enum element request */
-	FReply OnAddNewEnumerator();
+	UE_API FReply OnAddNewEnumerator();
 
 private:
 	TWeakObjectPtr<UUserDefinedEnum> TargetEnum;
@@ -189,3 +191,5 @@ private:
 	/** The editable text interface for the tooltip data */
 	TSharedPtr<FEditableTextUserDefinedEnumTooltip> TooltipEditor;
 };
+
+#undef UE_API
diff --git a/Engine/Source/Editor/Kismet/Public/BPVariableDragDropAction.h b/Engine/Source/Editor/Kismet/Public/BPVariableDragDropAction.h
--- a/Engine/Source/Editor/Kismet/Public/BPVariableDragDropAction.h
+++ b/Engine/Source/Editor/Kismet/Public/BPVariableDragDropAction.h
@@ -20,6 +20,8 @@
 #include "UObject/WeakObjectPtr.h"
 #include "UObject/WeakObjectPtrTemplates.h"
 
+#define UE_API KISMET_API
+
 class UBlueprint;
 class UEdGraph;
 class UEdGraphNode;
@@ -28,16 +30,16 @@ struct FSlateBrush;
 struct FSlateColor;
 
 /** DragDropAction class for dropping a Variable onto a graph */
-class KISMET_API FKismetVariableDragDropAction : public FMyBlueprintItemDragDropAction
+class FKismetVariableDragDropAction : public FMyBlueprintItemDragDropAction
 {
 public:
 	DRAG_DROP_OPERATOR_TYPE(FKismetVariableDragDropAction, FMyBlueprintItemDragDropAction)
 
 	// FGraphEditorDragDropAction interface
-	virtual void HoverTargetChanged() override;
-	virtual FReply DroppedOnPin(const FVector2f& ScreenPosition, const FVector2f& GraphPosition) override;
-	virtual FReply DroppedOnNode(const FVector2f& ScreenPosition, const FVector2f& GraphPosition) override;
-	virtual FReply DroppedOnPanel(const TSharedRef< class SWidget >& Panel, const FVector2f& ScreenPosition, const FVector2f& GraphPosition, UEdGraph& Graph) override;
+	UE_API virtual void HoverTargetChanged() override;
+	UE_API virtual FReply DroppedOnPin(const FVector2f& ScreenPosition, const FVector2f& GraphPosition) override;
+	UE_API virtual FReply DroppedOnNode(const FVector2f& ScreenPosition, const FVector2f& GraphPosition) override;
+	UE_API virtual FReply DroppedOnPanel(const TSharedRef< class SWidget >& Panel, const FVector2f& ScreenPosition, const FVector2f& GraphPosition, UEdGraph& Graph) override;
 	// End of FGraphEditorDragDropAction
 
 	static TSharedRef<FKismetVariableDragDropAction> New(TSharedPtr<FEdGraphSchemaAction> InAction, FName InVariableName, UStruct* InVariableSource, FNodeCreationAnalytic AnalyticCallback)
@@ -62,7 +64,7 @@ public:
 
 protected:
 	 /** Construct a FKismetVariableDragDropAction */
-	FKismetVariableDragDropAction();
+	UE_API FKismetVariableDragDropAction();
 
 	/** Structure for required node construction parameters */
 	struct FNodeConstructionParams
@@ -74,19 +76,19 @@ protected:
 	};
 
 	// FGraphSchemaActionDragDropAction interface
-	virtual void GetDefaultStatusSymbol(const FSlateBrush*& PrimaryBrushOut, FSlateColor& IconColorOut, FSlateBrush const*& SecondaryBrushOut, FSlateColor& SecondaryColorOut) const override;
+	UE_API virtual void GetDefaultStatusSymbol(const FSlateBrush*& PrimaryBrushOut, FSlateColor& IconColorOut, FSlateBrush const*& SecondaryBrushOut, FSlateColor& SecondaryColorOut) const override;
 	// End of FGraphSchemaActionDragDropAction interface
 
 	// FMyBlueprintItemDragDropAction interface
-	virtual UBlueprint* GetSourceBlueprint() const override;
+	UE_API virtual UBlueprint* GetSourceBlueprint() const override;
 	// End of FMyBlueprintItemDragDropAction interface
 
 	/** Called when user selects to create a Getter for the variable */
-	static void MakeGetter(FNodeConstructionParams InParams);
+	static UE_API void MakeGetter(FNodeConstructionParams InParams);
 	/** Called when user selects to create a Setter for the variable */
-	static void MakeSetter(FNodeConstructionParams InParams);
+	static UE_API void MakeSetter(FNodeConstructionParams InParams);
 	/** Called too check if we can execute a setter on a given property */
-	static bool CanExecuteMakeSetter(FNodeConstructionParams InParams, FProperty* InVariableProperty);
+	static UE_API bool CanExecuteMakeSetter(FNodeConstructionParams InParams, FProperty* InVariableProperty);
 
 	/**
 	 * Test new variable type against existing links for node and get any links that will break
@@ -95,7 +97,7 @@ protected:
 	 * @param	NewVariableProperty			the property for the new variable type 
 	 * @param	OutBroken						All of the links which are NOT compatible with the new type
 	 */
-	void GetLinksThatWillBreak(UEdGraphNode* Node, FProperty* NewVariableProperty, TArray<class UEdGraphPin*>& OutBroken);
+	UE_API void GetLinksThatWillBreak(UEdGraphNode* Node, FProperty* NewVariableProperty, TArray<class UEdGraphPin*>& OutBroken);
 
 	/** Indicates if replacing the variable node, with the new property will require any links to be broken*/
 	bool WillBreakLinks( UEdGraphNode* Node, FProperty* NewVariableProperty ) 
@@ -111,11 +113,11 @@ protected:
 	 * @param InVariableProperty		The variable property to check with
 	 * @param InGraph					The graph to check against placing the variable
 	 */
-	bool CanVariableBeDropped(const FProperty* InVariableProperty, const UEdGraph& InGraph) const;
+	UE_API bool CanVariableBeDropped(const FProperty* InVariableProperty, const UEdGraph& InGraph) const;
 
 
 	/** Returns the local variable's scope, if any */
-	UStruct* GetLocalVariableScope() const;
+	UE_API UStruct* GetLocalVariableScope() const;
 
 protected:
 	/** Name of variable being dragged */
@@ -123,3 +125,5 @@ protected:
 	/** Scope this variable belongs to */
 	TWeakObjectPtr<UStruct> VariableSource;
 };
+
+#undef UE_API
diff --git a/Engine/Source/Editor/Kismet/Public/BlueprintActionMenuBuilder.h b/Engine/Source/Editor/Kismet/Public/BlueprintActionMenuBuilder.h
--- a/Engine/Source/Editor/Kismet/Public/BlueprintActionMenuBuilder.h
+++ b/Engine/Source/Editor/Kismet/Public/BlueprintActionMenuBuilder.h
@@ -9,6 +9,8 @@
 #include "Internationalization/Text.h"
 #include "Templates/SharedPointer.h"
 
+#define UE_API KISMET_API
+
 class FBlueprintActionFilter;
 class FBlueprintEditor;
 
@@ -29,7 +31,7 @@ struct FBlueprintActionInfo;
  * FBlueprintActionMenuItems for actions that pass. Takes care of generating the 
  * each menu item's category/name/etc.
  */
-struct KISMET_API FBlueprintActionMenuBuilder : public FGraphActionListBuilderBase
+struct FBlueprintActionMenuBuilder : public FGraphActionListBuilderBase
 {
 public:
 	/** Flags used to configure the builder. */
@@ -65,13 +67,13 @@ public:
 	 * 
 	 * @param Flags	- Optional configuration flags.
 	 */
-	FBlueprintActionMenuBuilder(EConfigFlags Flags = EConfigFlags::DefaultConfig);
+	UE_API FBlueprintActionMenuBuilder(EConfigFlags Flags = EConfigFlags::DefaultConfig);
 
 	UE_DEPRECATED(5.2, "The action filter now stores a reference to the authoritative editor context. Please use the default constructor instead.")
-	FBlueprintActionMenuBuilder(TWeakPtr<FBlueprintEditor> BlueprintEditorPtr);
+	UE_API FBlueprintActionMenuBuilder(TWeakPtr<FBlueprintEditor> BlueprintEditorPtr);
 	
 	// FGraphActionListBuilderBase interface
-	virtual void Empty() override;
+	UE_API virtual void Empty() override;
 	// End FGraphActionListBuilderBase interface
 	
 	/**
@@ -86,30 +88,30 @@ public:
 	 * @param  MenuOrder The sort order to assign this section of the menu (higher numbers come first).
 	 * @param  Flags	 Set of ESectionFlags to customize this menu section.
 	 */
-	void AddMenuSection(FBlueprintActionFilter const& Filter, FText const& Heading = FText::GetEmpty(), int32 MenuOrder = 0, uint32 const Flags = 0);
+	UE_API void AddMenuSection(FBlueprintActionFilter const& Filter, FText const& Heading = FText::GetEmpty(), int32 MenuOrder = 0, uint32 const Flags = 0);
 	
 	/**
 	 * Regenerates the entire menu list from the cached menu sections. Filters 
 	 * and adds action items from the blueprint action database (as defined by 
 	 * the MenuSections list).
 	 */
-	void RebuildActionList();
+	UE_API void RebuildActionList();
 
 	/** Returns the current number of actions that are still pending */
-	int32 GetNumPendingActions() const;
+	UE_API int32 GetNumPendingActions() const;
 
 	/** Processes any actions that may be added asynchronously or across multiple frames. Returns true if one or more actions were added into the list. */
-	bool ProcessPendingActions();
+	UE_API bool ProcessPendingActions();
 
 	/** Returns the normalized completion state when processing pending actions (e.g. for a status indicator) */
-	float GetPendingActionsProgress() const;
+	UE_API float GetPendingActionsProgress() const;
 
 protected:
 	// Adds menu items for the given database action.
-	void MakeMenuItems(FBlueprintActionInfo& InAction);
+	UE_API void MakeMenuItems(FBlueprintActionInfo& InAction);
 
 	// Called when all menu items have been constructed.
-	void BuildCompleted();
+	UE_API void BuildCompleted();
 
 private:
 	/** 
@@ -129,4 +131,6 @@ private:
 	bool bUsePendingActionList;
 };
 
-ENUM_CLASS_FLAGS(FBlueprintActionMenuBuilder::EConfigFlags);
\ No newline at end of file
+ENUM_CLASS_FLAGS(FBlueprintActionMenuBuilder::EConfigFlags);
+
+#undef UE_API
diff --git a/Engine/Source/Editor/Kismet/Public/BlueprintActionMenuItem.h b/Engine/Source/Editor/Kismet/Public/BlueprintActionMenuItem.h
--- a/Engine/Source/Editor/Kismet/Public/BlueprintActionMenuItem.h
+++ b/Engine/Source/Editor/Kismet/Public/BlueprintActionMenuItem.h
@@ -18,6 +18,8 @@
 
 #include "BlueprintActionMenuItem.generated.h"
 
+#define UE_API KISMET_API
+
 class FReferenceCollector;
 class UBlueprintNodeSpawner;
 class UEdGraph;
@@ -34,7 +36,7 @@ struct FSlateBrush;
  * invoked from this class (separated to divide ui and node-spawning).
  */
 USTRUCT()
-struct KISMET_API FBlueprintActionMenuItem : public FEdGraphSchemaAction
+struct FBlueprintActionMenuItem : public FEdGraphSchemaAction
 {
 	GENERATED_USTRUCT_BODY();
 	static FName StaticGetTypeId() { static FName const TypeId("FBlueprintActionMenuItem"); return TypeId; }
@@ -42,20 +44,20 @@ struct KISMET_API FBlueprintActionMenuItem : public FEdGraphSchemaAction
 public:	
 	/** Constructors */
 	FBlueprintActionMenuItem(UBlueprintNodeSpawner const* NodeSpawner = nullptr) : Action(NodeSpawner), IconTint(FLinearColor::White), IconBrush(nullptr) {}
-	FBlueprintActionMenuItem(UBlueprintNodeSpawner const* NodeSpawner, FBlueprintActionUiSpec const& UiDef, IBlueprintNodeBinder::FBindingSet const& Bindings = IBlueprintNodeBinder::FBindingSet(), FText InNodeCategory = FText(), int32 InGrouping = 0);
+	UE_API FBlueprintActionMenuItem(UBlueprintNodeSpawner const* NodeSpawner, FBlueprintActionUiSpec const& UiDef, IBlueprintNodeBinder::FBindingSet const& Bindings = IBlueprintNodeBinder::FBindingSet(), FText InNodeCategory = FText(), int32 InGrouping = 0);
 	
 	// FEdGraphSchemaAction interface
 	virtual FName         GetTypeId() const final { return StaticGetTypeId(); }
 	using FEdGraphSchemaAction::PerformAction; // Prevent hiding of deprecated base class function with FVector2D
-	virtual UEdGraphNode* PerformAction(UEdGraph* ParentGraph, UEdGraphPin* FromPin, const FVector2f& Location, bool bSelectNewNode = true) final;
-	virtual UEdGraphNode* PerformAction(UEdGraph* ParentGraph, TArray<UEdGraphPin*>& FromPins, const FVector2f& Location, bool bSelectNewNode = true) final;
-	virtual void          AddReferencedObjects(FReferenceCollector& Collector) final;
+	UE_API virtual UEdGraphNode* PerformAction(UEdGraph* ParentGraph, UEdGraphPin* FromPin, const FVector2f& Location, bool bSelectNewNode = true) final;
+	UE_API virtual UEdGraphNode* PerformAction(UEdGraph* ParentGraph, TArray<UEdGraphPin*>& FromPins, const FVector2f& Location, bool bSelectNewNode = true) final;
+	UE_API virtual void          AddReferencedObjects(FReferenceCollector& Collector) final;
 	// End FEdGraphSchemaAction interface
 
 	/** @return */
-	UBlueprintNodeSpawner const* GetRawAction() const;
+	UE_API UBlueprintNodeSpawner const* GetRawAction() const;
 
-	void AppendBindings(const FBlueprintActionContext& Context, IBlueprintNodeBinder::FBindingSet const& BindingSet);
+	UE_API void AppendBindings(const FBlueprintActionContext& Context, IBlueprintNodeBinder::FBindingSet const& BindingSet);
 
 	/**
 	 * Retrieves the icon brush for this menu entry (to be displayed alongside
@@ -64,18 +66,18 @@ public:
 	 * @param  ColorOut	An output parameter that's filled in with the color to tint the brush with.
 	 * @return An slate brush to be used for this menu item in the action menu.
 	 */
-	FSlateBrush const* GetMenuIcon(FSlateColor& ColorOut);
+	UE_API FSlateBrush const* GetMenuIcon(FSlateColor& ColorOut);
 
 	/** Utility struct for pairing documentation page names with excerpt names */
-	struct KISMET_API FDocExcerptRef
+	struct FDocExcerptRef
 	{
-		bool IsValid() const;
+		UE_API bool IsValid() const;
 
 		FString DocLink;
 		FString DocExcerptName;		
 	};
 	/**  */
-	const FDocExcerptRef& GetDocumentationExcerpt() const;
+	UE_API const FDocExcerptRef& GetDocumentationExcerpt() const;
 
 private:
 	/** Specialized node-spawner, that comprises the action portion of this menu entry. */
@@ -89,3 +91,5 @@ private:
 	/** References the documentation page/excerpt pertaining to the node this will spawn */
 	FDocExcerptRef DocExcerptRef;
 };
+
+#undef UE_API
diff --git a/Engine/Source/Editor/Kismet/Public/BlueprintCompilationManager.h b/Engine/Source/Editor/Kismet/Public/BlueprintCompilationManager.h
--- a/Engine/Source/Editor/Kismet/Public/BlueprintCompilationManager.h
+++ b/Engine/Source/Editor/Kismet/Public/BlueprintCompilationManager.h
@@ -9,6 +9,8 @@
 #include "Templates/SharedPointer.h"
 #include "Templates/SubclassOf.h"
 
+#define UE_API KISMET_API
+
 class FCompilerResultsLog;
 class FProperty;
 class FString;
@@ -33,10 +35,10 @@ struct FBPCompileRequest
 	FCompilerResultsLog* ClientResultsLog;
 };
 
-struct KISMET_API FBlueprintCompilationManager
+struct FBlueprintCompilationManager
 {
-	static void Initialize();
-	static void Shutdown();
+	static UE_API void Initialize();
+	static UE_API void Shutdown();
 
 	/**
 	 * Compiles all blueprints that have been placed in the compilation queue. 
@@ -45,46 +47,46 @@ struct KISMET_API FBlueprintCompilationManager
 	 * of objects the compilation manager has loaded. The compilation manager
 	 * will not load data while processing the compilation queue)
 	 */
-	static void FlushCompilationQueue(FUObjectSerializeContext* InLoadContext);
+	static UE_API void FlushCompilationQueue(FUObjectSerializeContext* InLoadContext);
 	
 	/**
 	 * Flushes the compilation queue and finishes reinstancing
 	 */
-	static void FlushCompilationQueueAndReinstance(const bool bFindAndReplaceCDOReferences = false);
+	static UE_API void FlushCompilationQueueAndReinstance(const bool bFindAndReplaceCDOReferences = false);
 
 	/**
 	 * Immediately compiles the blueprint, no expectation that related blueprints be subsequently compiled.
 	 * It will be significantly more efficient to queue blueprints and then flush the compilation queue
 	 * if there are several blueprints that require compilation (e.g. typical case on PIE):
 	 */
-	static void CompileSynchronously(const FBPCompileRequest& Request);
+	static UE_API void CompileSynchronously(const FBPCompileRequest& Request);
 
 	/**
 	 * Adds a newly loaded blueprint to the compilation queue
 	 */
-	static void NotifyBlueprintLoaded(UBlueprint* BPLoaded);
+	static UE_API void NotifyBlueprintLoaded(UBlueprint* BPLoaded);
 
 	/**
 	 * Adds a blueprint to the compilation queue - useful for batch compilation
 	 */
-	static void QueueForCompilation(UBlueprint* BP);
+	static UE_API void QueueForCompilation(UBlueprint* BP);
 
 	/** Returns true when UBlueprint::GeneratedClass members are up to date */
-	static bool IsGeneratedClassLayoutReady();
+	static UE_API bool IsGeneratedClassLayoutReady();
 	
 	/** 
 	 * Returns the Default Value associated with ForClass::Property, if ForClass is currently 
 	 * being compiled this function can look at the old version of the CDO and read the default
 	 * value from there
 	 */
-	static bool GetDefaultValue(const UClass* ForClass, const FProperty* Property, FString& OutDefaultValueAsString);
+	static UE_API bool GetDefaultValue(const UClass* ForClass, const FProperty* Property, FString& OutDefaultValueAsString);
 
 	/**
 	 * Safely reparents all child classes of every Key in OldClassToNewClass to the class in 
 	 * the corresponding Value. Typically this means every child type will be reinstanced - although
 	 * reinstancing could be avoided when layouts match.
 	 */
-	static void ReparentHierarchies(const TMap<UClass*, UClass*>& OldClassToNewClass);
+	static UE_API void ReparentHierarchies(const TMap<UClass*, UClass*>& OldClassToNewClass);
 
 	/** 
 	 * Registers a blueprint compiler extension - anytime a blueprint of the provided type is compiled
@@ -92,8 +94,9 @@ struct KISMET_API FBlueprintCompilationManager
 	 * compilation there may be blueprints compiled before the extension is registed unless special
 	 * care has been taken.
 	 */
-	static void RegisterCompilerExtension(TSubclassOf<UBlueprint> BlueprintType, UBlueprintCompilerExtension* Extension);
+	static UE_API void RegisterCompilerExtension(TSubclassOf<UBlueprint> BlueprintType, UBlueprintCompilerExtension* Extension);
 private:
-	FBlueprintCompilationManager();
+	UE_API FBlueprintCompilationManager();
 };
 
+#undef UE_API
diff --git a/Engine/Source/Editor/Kismet/Public/BlueprintCompilerExtension.h b/Engine/Source/Editor/Kismet/Public/BlueprintCompilerExtension.h
--- a/Engine/Source/Editor/Kismet/Public/BlueprintCompilerExtension.h
+++ b/Engine/Source/Editor/Kismet/Public/BlueprintCompilerExtension.h
@@ -10,6 +10,8 @@
 
 #include "BlueprintCompilerExtension.generated.h"
 
+#define UE_API KISMET_API
+
 class FKismetCompilerContext;
 class UEdGraph;
 
@@ -22,14 +24,14 @@ struct FBlueprintCompiledData
 	TArray<TObjectPtr<UEdGraph>> IntermediateGraphs;
 };
 
-UCLASS(Abstract)
-class KISMET_API UBlueprintCompilerExtension : public UObject
+UCLASS(MinimalAPI, Abstract)
+class UBlueprintCompilerExtension : public UObject
 {
 	GENERATED_BODY()
 public:
-	UBlueprintCompilerExtension(const FObjectInitializer& ObjectInitializer);
+	UE_API UBlueprintCompilerExtension(const FObjectInitializer& ObjectInitializer);
 
-	void BlueprintCompiled(const FKismetCompilerContext& CompilationContext, const FBlueprintCompiledData& Data);
+	UE_API void BlueprintCompiled(const FKismetCompilerContext& CompilationContext, const FBlueprintCompiledData& Data);
 
 protected:
 	/** 
@@ -38,3 +40,5 @@ protected:
 	 */
 	virtual void ProcessBlueprintCompiled(const FKismetCompilerContext& CompilationContext, const FBlueprintCompiledData& Data) {}
 };
+
+#undef UE_API
diff --git a/Engine/Source/Editor/Kismet/Public/BlueprintDragDropMenuItem.h b/Engine/Source/Editor/Kismet/Public/BlueprintDragDropMenuItem.h
--- a/Engine/Source/Editor/Kismet/Public/BlueprintDragDropMenuItem.h
+++ b/Engine/Source/Editor/Kismet/Public/BlueprintDragDropMenuItem.h
@@ -17,6 +17,8 @@
 
 #include "BlueprintDragDropMenuItem.generated.h"
 
+#define UE_API KISMET_API
+
 class FDragDropOperation;
 class FReferenceCollector;
 class FText;
@@ -42,7 +44,7 @@ struct FSlateColor;
  * the proper FDragDropOperation.
  */
 USTRUCT()
-struct KISMET_API FBlueprintDragDropMenuItem : public FEdGraphSchemaAction
+struct FBlueprintDragDropMenuItem : public FEdGraphSchemaAction
 {
 	GENERATED_USTRUCT_BODY();
 	
@@ -58,13 +60,13 @@ public:
 	 * @param  SampleAction	One of the node-spawners that this "consolidated" menu item represents (shouldn't matter which one).
 	 * @param  MenuGrouping	Defines the order in which this entry is listed in the menu.
 	 */
-	FBlueprintDragDropMenuItem(FBlueprintActionContext const& Context, UBlueprintNodeSpawner const* SampleAction, int32 MenuGrouping, FText InNodeCategory, FText InMenuDesc, FText InToolTip );
+	UE_API FBlueprintDragDropMenuItem(FBlueprintActionContext const& Context, UBlueprintNodeSpawner const* SampleAction, int32 MenuGrouping, FText InNodeCategory, FText InMenuDesc, FText InToolTip );
 	
 	// FEdGraphSchemaAction interface
 	virtual FName GetTypeId() const final { return StaticGetTypeId(); }
-	virtual UEdGraphNode* PerformAction(UEdGraph* ParentGraph, UEdGraphPin* FromPin, const FVector2f& Location, bool bSelectNewNode = true) override;
-	virtual UEdGraphNode* PerformAction(UEdGraph* ParentGraph, TArray<UEdGraphPin*>& FromPins, const FVector2f& Location, bool bSelectNewNode = true) override;
-	virtual void  AddReferencedObjects(FReferenceCollector& Collector) override;
+	UE_API virtual UEdGraphNode* PerformAction(UEdGraph* ParentGraph, UEdGraphPin* FromPin, const FVector2f& Location, bool bSelectNewNode = true) override;
+	UE_API virtual UEdGraphNode* PerformAction(UEdGraph* ParentGraph, TArray<UEdGraphPin*>& FromPins, const FVector2f& Location, bool bSelectNewNode = true) override;
+	UE_API virtual void  AddReferencedObjects(FReferenceCollector& Collector) override;
 	// End FEdGraphSchemaAction interface
 
 	/**
@@ -74,16 +76,16 @@ public:
 	 * @param  ColorOut	An output parameter that's filled in with the color to tint the brush with.
 	 * @return An slate brush to be used for this menu item in the action menu.
 	 */
-	FSlateBrush const* GetMenuIcon(FSlateColor& ColorOut);
+	UE_API FSlateBrush const* GetMenuIcon(FSlateColor& ColorOut);
 
 	/** */
-	void AppendAction(UBlueprintNodeSpawner const* Action);
+	UE_API void AppendAction(UBlueprintNodeSpawner const* Action);
 
 	/** @return   */
-	UBlueprintNodeSpawner const* GetSampleAction() const;
+	UE_API UBlueprintNodeSpawner const* GetSampleAction() const;
 
 	/** @return   */
-	TSet<UBlueprintNodeSpawner const*> const& GetActionSet() const;
+	UE_API TSet<UBlueprintNodeSpawner const*> const& GetActionSet() const;
 
 	/**
 	 * Attempts to create a certain drag/drop action with respect to the 
@@ -92,9 +94,11 @@ public:
 	 * @param  AnalyticsDelegate	The analytics callback to assign the drag/drop operation.
 	 * @return An empty TSharedPtr if we failed to determine the FDragDropOperation type to use, else a newly instantiated drag/drop op.
 	 */
-	TSharedPtr<FDragDropOperation> OnDragged(FNodeCreationAnalytic AnalyticsDelegate) const;
+	UE_API TSharedPtr<FDragDropOperation> OnDragged(FNodeCreationAnalytic AnalyticsDelegate) const;
 
 private:
 	/** An arbitrary member of the node-spawner subset that this menu entry represents. */
 	TSet<TObjectPtr<const UBlueprintNodeSpawner>> ActionSet;
 };
+
+#undef UE_API
