Updated headers using UnrealCodeFixup to make sure dllstorage is on methods/staticvars instead of types. Part 5/n

diff --git a/Engine/Source/Editor/Kismet/Public/BlueprintEditor.h b/Engine/Source/Editor/Kismet/Public/BlueprintEditor.h
--- a/Engine/Source/Editor/Kismet/Public/BlueprintEditor.h
+++ b/Engine/Source/Editor/Kismet/Public/BlueprintEditor.h
@@ -55,6 +55,8 @@
 #include "WorkflowOrientedApp/WorkflowTabFactory.h"
 #include "WorkflowOrientedApp/WorkflowTabManager.h"
 
+#define UE_API KISMET_API
+
 class AActor;
 class FBlueprintEditorToolbar;
 class FBlueprintNamespaceHelper;
@@ -158,13 +160,13 @@ public:
 
 #define LOCTEXT_NAMESPACE "BlueprintEditor"
 
-struct KISMET_API FSelectionDetailsSummoner : public FWorkflowTabFactory
+struct FSelectionDetailsSummoner : public FWorkflowTabFactory
 {
 public:
-	FSelectionDetailsSummoner(TSharedPtr<class FAssetEditorToolkit> InHostingApp);
+	UE_API FSelectionDetailsSummoner(TSharedPtr<class FAssetEditorToolkit> InHostingApp);
 
-	virtual TSharedRef<SWidget> CreateTabBody(const FWorkflowTabSpawnInfo& Info) const override;
-	virtual TSharedRef<SDockTab> SpawnTab(const FWorkflowTabSpawnInfo& Info) const override;
+	UE_API virtual TSharedRef<SWidget> CreateTabBody(const FWorkflowTabSpawnInfo& Info) const override;
+	UE_API virtual TSharedRef<SDockTab> SpawnTab(const FWorkflowTabSpawnInfo& Info) const override;
 
 	virtual FText GetTabToolTipText(const FWorkflowTabSpawnInfo& Info) const override
 	{
@@ -207,7 +209,7 @@ namespace ENodeCreateAction
 // FBlueprintEditor
 
 /** Main Kismet asset editor */
-class KISMET_API FBlueprintEditor : public IBlueprintEditor, public FGCObject, public FNotifyHook, public 
+class FBlueprintEditor : public IBlueprintEditor, public FGCObject, public FNotifyHook, public 
 FTickableEditorObject, public FEditorUndoClient, public FNoncopyable
 {
 	DECLARE_MULTICAST_DELEGATE_OneParam(FOnSetPinVisibility, SGraphEditor::EPinVisibility);
@@ -229,12 +231,12 @@ FTickableEditorObject, public FEditorUndoClient, public FNoncopyable
 
 public:
 	//~ Begin IToolkit Interface
-	virtual void RegisterTabSpawners(const TSharedRef<class FTabManager>& TabManager) override;
+	UE_API virtual void RegisterTabSpawners(const TSharedRef<class FTabManager>& TabManager) override;
 	//~ End IToolkit Interface
 
 	DECLARE_MULTICAST_DELEGATE_OneParam(FOnModeSet, FName);
 	FOnModeSet& OnModeSet() { return OnModeSetData; }
-	virtual void SetCurrentMode(FName NewMode) override;
+	UE_API virtual void SetCurrentMode(FName NewMode) override;
 
 public:
 	/**
@@ -245,94 +247,94 @@ public:
 	 * @param	InBlueprints			The blueprints to edit
 	 * @param	bShouldOpenInDefaultsMode	If true, the editor will open in defaults editing mode
 	 */
-	void InitBlueprintEditor(const EToolkitMode::Type Mode, const TSharedPtr< class IToolkitHost >& InitToolkitHost, const TArray<class UBlueprint*>& InBlueprints, bool bShouldOpenInDefaultsMode);
+	UE_API void InitBlueprintEditor(const EToolkitMode::Type Mode, const TSharedPtr< class IToolkitHost >& InitToolkitHost, const TArray<class UBlueprint*>& InBlueprints, bool bShouldOpenInDefaultsMode);
 
 public:
 	//~ Begin FAssetEditorToolkit Interface
-	virtual bool OnRequestClose(EAssetEditorCloseReason InCloseReason) override;
-	virtual void OnClose() override;
+	UE_API virtual bool OnRequestClose(EAssetEditorCloseReason InCloseReason) override;
+	UE_API virtual void OnClose() override;
 	// End of FAssetEditorToolkit 
 
 	//~ Begin IToolkit Interface
-	virtual FName GetToolkitContextFName() const override;
-	virtual FName GetToolkitFName() const override;
-	virtual FText GetBaseToolkitName() const override;
-	virtual FText GetToolkitName() const override;
-	virtual FText GetToolkitToolTipText() const override;
-	virtual FString GetWorldCentricTabPrefix() const override;
-	virtual FLinearColor GetWorldCentricTabColorScale() const override;
-	virtual bool IsBlueprintEditor() const override;
+	UE_API virtual FName GetToolkitContextFName() const override;
+	UE_API virtual FName GetToolkitFName() const override;
+	UE_API virtual FText GetBaseToolkitName() const override;
+	UE_API virtual FText GetToolkitName() const override;
+	UE_API virtual FText GetToolkitToolTipText() const override;
+	UE_API virtual FString GetWorldCentricTabPrefix() const override;
+	UE_API virtual FLinearColor GetWorldCentricTabColorScale() const override;
+	UE_API virtual bool IsBlueprintEditor() const override;
 	//~ End IToolkit Interface
 
 	//~ Begin FGCObject Interface
-	virtual void AddReferencedObjects(FReferenceCollector& Collector) override;
-	virtual FString GetReferencerName() const override;
+	UE_API virtual void AddReferencedObjects(FReferenceCollector& Collector) override;
+	UE_API virtual FString GetReferencerName() const override;
 	//~ End FGCObject Interface
 
 	//~ Begin IBlueprintEditor Interface
-	virtual void RefreshEditors(ERefreshBlueprintEditorReason::Type Reason = ERefreshBlueprintEditorReason::UnknownReason) override;
-	virtual void RefreshMyBlueprint();
-	virtual void RefreshInspector();
-	virtual void AddToSelection(UEdGraphNode* InNode) override;
-	virtual void JumpToHyperlink(const UObject* ObjectReference, bool bRequestRename = false) override;
-	virtual void JumpToPin(const class UEdGraphPin* Pin) override;
-	virtual void SummonSearchUI(bool bSetFindWithinBlueprint, FString NewSearchTerms = FString(), bool bSelectFirstResult = false) override;
-	virtual void SummonFindAndReplaceUI() override;
-	virtual UEdGraph* GetFocusedGraph() const override;
-	virtual TSharedPtr<SGraphEditor> OpenGraphAndBringToFront(UEdGraph* Graph, bool bSetFocus = true) override;
+	UE_API virtual void RefreshEditors(ERefreshBlueprintEditorReason::Type Reason = ERefreshBlueprintEditorReason::UnknownReason) override;
+	UE_API virtual void RefreshMyBlueprint();
+	UE_API virtual void RefreshInspector();
+	UE_API virtual void AddToSelection(UEdGraphNode* InNode) override;
+	UE_API virtual void JumpToHyperlink(const UObject* ObjectReference, bool bRequestRename = false) override;
+	UE_API virtual void JumpToPin(const class UEdGraphPin* Pin) override;
+	UE_API virtual void SummonSearchUI(bool bSetFindWithinBlueprint, FString NewSearchTerms = FString(), bool bSelectFirstResult = false) override;
+	UE_API virtual void SummonFindAndReplaceUI() override;
+	UE_API virtual UEdGraph* GetFocusedGraph() const override;
+	UE_API virtual TSharedPtr<SGraphEditor> OpenGraphAndBringToFront(UEdGraph* Graph, bool bSetFocus = true) override;
 
 	UE_DEPRECATED(5.0, "GetSelectedSCSEditorTreeNodes has been deprecated. Use GetSelectedSubobjectEditorTreeNodes instead.")
-	virtual TArray<TSharedPtr<class FSCSEditorTreeNode>> GetSelectedSCSEditorTreeNodes() const override;
+	UE_API virtual TArray<TSharedPtr<class FSCSEditorTreeNode>> GetSelectedSCSEditorTreeNodes() const override;
 	
 	UE_DEPRECATED(5.0, "FindAndSelectSCSEditorTreeNode has been deprecated. Use FindAndSelectSubobjectEditorTreeNode instead.")
 	virtual TSharedPtr<class FSCSEditorTreeNode> FindAndSelectSCSEditorTreeNode(const UActorComponent* InComponent, bool IsCntrlDown) override { return nullptr; }
 
-	virtual TArray<TSharedPtr<FSubobjectEditorTreeNode>> GetSelectedSubobjectEditorTreeNodes() const override;
-	virtual TSharedPtr<FSubobjectEditorTreeNode> FindAndSelectSubobjectEditorTreeNode(const UActorComponent* InComponent, bool IsCntrlDown) override;
-	virtual int32 GetNumberOfSelectedNodes() const override;
-	virtual void AnalyticsTrackNodeEvent(UBlueprint* Blueprint, UEdGraphNode *GraphNode, bool bNodeDelete = false) const override;
-	void AnalyticsTrackCompileEvent(UBlueprint* Blueprint, int32 NumErrors, int32 NumWarnings) const;
+	UE_API virtual TArray<TSharedPtr<FSubobjectEditorTreeNode>> GetSelectedSubobjectEditorTreeNodes() const override;
+	UE_API virtual TSharedPtr<FSubobjectEditorTreeNode> FindAndSelectSubobjectEditorTreeNode(const UActorComponent* InComponent, bool IsCntrlDown) override;
+	UE_API virtual int32 GetNumberOfSelectedNodes() const override;
+	UE_API virtual void AnalyticsTrackNodeEvent(UBlueprint* Blueprint, UEdGraphNode *GraphNode, bool bNodeDelete = false) const override;
+	UE_API void AnalyticsTrackCompileEvent(UBlueprint* Blueprint, int32 NumErrors, int32 NumWarnings) const;
 	virtual TSharedPtr<class IClassViewerFilter> GetImportedClassViewerFilter() const override { return ImportedClassViewerFilter; }
 	UE_DEPRECATED(5.1, "Please use GetPinTypeSelectorFilters")
 	virtual TSharedPtr<class IPinTypeSelectorFilter> GetImportedPinTypeSelectorFilter() const override { return ImportedPinTypeSelectorFilter; }
-	virtual void GetPinTypeSelectorFilters(TArray<TSharedPtr<class IPinTypeSelectorFilter>>& OutFilters) const override;
-	virtual bool IsNonImportedObject(const UObject* InObject) const;
-	virtual bool IsNonImportedObject(const FSoftObjectPath& InObject) const;
+	UE_API virtual void GetPinTypeSelectorFilters(TArray<TSharedPtr<class IPinTypeSelectorFilter>>& OutFilters) const override;
+	UE_API virtual bool IsNonImportedObject(const UObject* InObject) const;
+	UE_API virtual bool IsNonImportedObject(const FSoftObjectPath& InObject) const;
 	//~ End IBlueprintEditor Interface
 
 	//~ Begin FTickableEditorObject Interface
-	virtual void Tick(float DeltaTime) override;
+	UE_API virtual void Tick(float DeltaTime) override;
 	virtual ETickableTickType GetTickableTickType() const override { return ETickableTickType::Always; }
-	virtual TStatId GetStatId() const override;
+	UE_API virtual TStatId GetStatId() const override;
 	//~ End FTickableEditorObject Interface
 
 public:
-	FBlueprintEditor();
+	UE_API FBlueprintEditor();
 
-	virtual ~FBlueprintEditor();
+	UE_API virtual ~FBlueprintEditor();
 
 	/** Add context objects for menus and toolbars */
-	virtual void InitToolMenuContext(FToolMenuContext& MenuContext) override;
+	UE_API virtual void InitToolMenuContext(FToolMenuContext& MenuContext) override;
 
 	/** Check the Node Title is visible */
-	bool IsNodeTitleVisible(const UEdGraphNode* Node, bool bRequestRename);
+	UE_API bool IsNodeTitleVisible(const UEdGraphNode* Node, bool bRequestRename);
 
 	/** Pan the view to center on a particular node */
-	void JumpToNode(const class UEdGraphNode* Node, bool bRequestRename = false);
+	UE_API void JumpToNode(const class UEdGraphNode* Node, bool bRequestRename = false);
 
 	/** Returns a pointer to the Blueprint object we are currently editing, as long as we are editing exactly one */
-	virtual UBlueprint* GetBlueprintObj() const;
+	UE_API virtual UBlueprint* GetBlueprintObj() const;
 
 	/**	Returns whether the editor is currently editing a single blueprint object */
-	bool IsEditingSingleBlueprint() const;
+	UE_API bool IsEditingSingleBlueprint() const;
 
 	/** Ensures the blueprint has keyboard focus. */
-	void SetKeyboardFocus();
+	UE_API void SetKeyboardFocus();
 
 	/** Getters for the various Kismet2 widgets */
 	TSharedRef<SKismetInspector> GetInspector() const { return Inspector.ToSharedRef(); }
 	TSharedRef<SKismetInspector> GetDefaultEditor() const { return DefaultEditor.ToSharedRef(); }
-	TSharedRef<SBlueprintPalette> GetPalette();
+	UE_API TSharedRef<SBlueprintPalette> GetPalette();
 	TSharedRef<SBlueprintBookmarks> GetBookmarksWidget() const { return BookmarksWidget.ToSharedRef(); }
 	TSharedRef<SWidget> GetCompilerResults() const { return CompilerResults.ToSharedRef(); }
 	TSharedRef<SFindInBlueprints> GetFindResults() const { return FindResults.ToSharedRef(); }
@@ -349,7 +351,7 @@ public:
 	/**
 	 * Provides access to the preview actor.
 	 */
-	AActor* GetPreviewActor() const;
+	UE_API AActor* GetPreviewActor() const;
 
 	/**
 	 * Provides access to the preview scene.
@@ -365,103 +367,103 @@ public:
 	* @param InBlueprint			The Blueprint to create or update the preview for.
 	* @param bInForceFullUpdate	Force a full update to respawn actors.
 	*/
-	void UpdatePreviewActor(UBlueprint* InBlueprint, bool bInForceFullUpdate = false);
+	UE_API void UpdatePreviewActor(UBlueprint* InBlueprint, bool bInForceFullUpdate = false);
 
 	/**
 	* Destroy the Blueprint preview.
 	*/
-	void DestroyPreview();
+	UE_API void DestroyPreview();
 
 	TSharedPtr<class FBlueprintEditorToolbar> GetToolbarBuilder() { return Toolbar; }
 
 	/** @return the documentation location for this editor */
-	virtual FString GetDocumentationLink() const override;
+	UE_API virtual FString GetDocumentationLink() const override;
 
 	/**	Returns whether the edited blueprint has components */
-	bool CanAccessComponentsMode() const;
+	UE_API bool CanAccessComponentsMode() const;
 
 	/** Returns if we are currently closing the editor */
-	bool IsEditorClosing() const;
+	UE_API bool IsEditorClosing() const;
 
 	// @todo This is a hack for now until we reconcile the default toolbar with application modes
-	void RegisterToolbarTab(const TSharedRef<class FTabManager>& TabManager);
+	UE_API void RegisterToolbarTab(const TSharedRef<class FTabManager>& TabManager);
 	/** Throw a simple message into the log */
-	void LogSimpleMessage(const FText& MessageText);
+	UE_API void LogSimpleMessage(const FText& MessageText);
 
 	/** Dumps messages to the compiler log, with an option to force it to display/come to front */
-	void DumpMessagesToCompilerLog(const TArray<TSharedRef<class FTokenizedMessage>>& Messages, bool bForceMessageDisplay);
+	UE_API void DumpMessagesToCompilerLog(const TArray<TSharedRef<class FTokenizedMessage>>& Messages, bool bForceMessageDisplay);
 
 	/** Returns true if PIE is active */
-	bool IsPlayInEditorActive() const;
+	UE_API bool IsPlayInEditorActive() const;
 
 	/** Get the currently selected set of nodes */
-	FGraphPanelSelectionSet GetSelectedNodes() const;
+	UE_API FGraphPanelSelectionSet GetSelectedNodes() const;
 
 	/** Returns the currently selected node if there is a single node selected (if there are multiple nodes selected or none selected, it will return nullptr) */
-	UEdGraphNode* GetSingleSelectedNode() const;
+	UE_API UEdGraphNode* GetSingleSelectedNode() const;
 
 	/** Save the current set of edited objects in the LastEditedObjects array so it will be opened next time we open K2 */
-	void SaveEditedObjectState();
+	UE_API void SaveEditedObjectState();
 
 	/** Create new tab for each element of LastEditedObjects array */
-	void RestoreEditedObjectState();
+	UE_API void RestoreEditedObjectState();
 
 	// Request a save of the edited object state
 	// This is used to delay it by one frame when triggered by a tab being closed, so it can finish closing before remembering the new state
-	void RequestSaveEditedObjectState();
+	UE_API void RequestSaveEditedObjectState();
 
 	/** Get the visible bounds of the given graph node */
-	void GetBoundsForNode(const UEdGraphNode* InNode, class FSlateRect& OutRect, float InPadding) const;
+	UE_API void GetBoundsForNode(const UEdGraphNode* InNode, class FSlateRect& OutRect, float InPadding) const;
 
 	/** Gets the focused graph current view bookmark ID */
-	void GetViewBookmark(FGuid& BookmarkId);
+	UE_API void GetViewBookmark(FGuid& BookmarkId);
 
 	/** Gets the focused graph view location/zoom amount */
-	void GetViewLocation(FVector2D& Location, float& ZoomAmount);
+	UE_API void GetViewLocation(FVector2D& Location, float& ZoomAmount);
 
 	/** Sets the focused graph view location/zoom amount */
-	void SetViewLocation(const FVector2D& Location, float ZoomAmount, const FGuid& BookmarkId = FGuid());
+	UE_API void SetViewLocation(const FVector2D& Location, float ZoomAmount, const FGuid& BookmarkId = FGuid());
 
 	/** Returns whether a graph is editable or not */
-	virtual bool IsEditable(UEdGraph* InGraph) const;
+	UE_API virtual bool IsEditable(UEdGraph* InGraph) const;
 	/** Determines if the graph is ReadOnly, this differs from editable in that it is never expected to be edited and is in a read-only state */
-	bool IsGraphReadOnly(UEdGraph* InGraph) const;
+	UE_API bool IsGraphReadOnly(UEdGraph* InGraph) const;
 
 	/** Used to determine the visibility of the graph's instruction text. */
-	float GetInstructionTextOpacity(UEdGraph* InGraph) const;
+	UE_API float GetInstructionTextOpacity(UEdGraph* InGraph) const;
 
 	/** Returns true if in editing mode */
-	virtual bool InEditingMode() const;
+	UE_API virtual bool InEditingMode() const;
 
 	/** Returns true if able to compile */
-	virtual bool IsCompilingEnabled() const;
+	UE_API virtual bool IsCompilingEnabled() const;
 
 	/** Returns true if the parent class of the Blueprint being edited is also a Blueprint */
-	bool IsParentClassABlueprint() const;
+	UE_API bool IsParentClassABlueprint() const;
 
 	/** Returns true if the parent class of the Blueprint being edited is an editable Blueprint */
-	bool IsParentClassAnEditableBlueprint() const;
+	UE_API bool IsParentClassAnEditableBlueprint() const;
 
 	/** Returns true if the parent class of the Blueprint being edited is native */
-	bool IsParentClassNative() const;
+	UE_API bool IsParentClassNative() const;
 
 	/** Returns true if the parent class is native and the link to it's header can be shown*/
-	virtual bool IsNativeParentClassCodeLinkEnabled() const;
+	UE_API virtual bool IsNativeParentClassCodeLinkEnabled() const;
 
 	/** Handles opening the header file of native parent class */
-	void OnEditParentClassNativeCodeClicked();
+	UE_API void OnEditParentClassNativeCodeClicked();
 
 	/** Returns: "(<NativeParentClass>.h)" */
-	FText GetTextForNativeParentClassHeaderLink() const;
+	UE_API FText GetTextForNativeParentClassHeaderLink() const;
 
 	/** Determines visibility of the native parent class manipulation buttons on the menu bar overlay */
-	EVisibility GetNativeParentClassButtonsVisibility() const;
+	UE_API EVisibility GetNativeParentClassButtonsVisibility() const;
 
 	/** Determines visibility of the standard parent class label on the menu bar overlay */
-	EVisibility GetParentClassNameVisibility() const;
+	UE_API EVisibility GetParentClassNameVisibility() const;
 
 	/** Returns our PIE Status - SIMULATING / SERVER / CLIENT */
-	FText GetPIEStatus() const;
+	UE_API FText GetPIEStatus() const;
 
 	/**
 	 * Util for finding a glyph for a graph
@@ -470,7 +472,7 @@ public:
 	 * @param bInLargeIcon - if true the icon returned is 22x22 pixels, else it is 16x16
 	 * @return An appropriate brush to use to represent the graph, if the graph is an unknown type the function will return the default "function" glyph
 	 */
-	static const FSlateBrush* GetGlyphForGraph(const UEdGraph* Graph, bool bInLargeIcon = false);
+	static UE_API const FSlateBrush* GetGlyphForGraph(const UEdGraph* Graph, bool bInLargeIcon = false);
 
 	/**
 	 * Util for finding a glyph and color for a variable.
@@ -482,7 +484,7 @@ public:
 	 * @param SecondaryColorOut The resulting secondary color for the glyph (used for Map types)
 	 * @return					The resulting glyph brush
 	 */
-	static FSlateBrush const* GetVarIconAndColor(const UStruct* VarScope, FName VarName, FSlateColor& IconColorOut, FSlateBrush const*& SecondaryBrushOut, FSlateColor& SecondaryColorOut);
+	static UE_API FSlateBrush const* GetVarIconAndColor(const UStruct* VarScope, FName VarName, FSlateColor& IconColorOut, FSlateBrush const*& SecondaryBrushOut, FSlateColor& SecondaryColorOut);
 
 	/**
 	 * Util for finding a glyph and color for a variable.
@@ -493,7 +495,7 @@ public:
 	 * @param SecondaryColorOut The resulting secondary color for the glyph (used for Map types)
 	 * @return					The resulting glyph brush
 	 */
-	static FSlateBrush const* GetVarIconAndColorFromProperty(const FProperty* Property, FSlateColor& IconColorOut, FSlateBrush const*& SecondaryBrushOut, FSlateColor& SecondaryColorOut);
+	static UE_API FSlateBrush const* GetVarIconAndColorFromProperty(const FProperty* Property, FSlateColor& IconColorOut, FSlateBrush const*& SecondaryBrushOut, FSlateColor& SecondaryColorOut);
 
 	/**
 	* Util for finding a glyph and color for a variable.
@@ -504,78 +506,78 @@ public:
 	* @param SecondaryColorOut The resulting secondary color for the glyph (used for Map types)
 	* @return					The resulting glyph brush
 	*/
-	static FSlateBrush const* GetVarIconAndColorFromPinType(const FEdGraphPinType& PinType, FSlateColor& IconColorOut, FSlateBrush const*& SecondaryBrushOut, FSlateColor& SecondaryColorOut);
+	static UE_API FSlateBrush const* GetVarIconAndColorFromPinType(const FEdGraphPinType& PinType, FSlateColor& IconColorOut, FSlateBrush const*& SecondaryBrushOut, FSlateColor& SecondaryColorOut);
 
 	/** Overridable function for determining if the current mode can script */
-	virtual bool IsInAScriptingMode() const;
+	UE_API virtual bool IsInAScriptingMode() const;
 
 	/** Called when Compile button is clicked */
-	virtual void Compile();
+	UE_API virtual void Compile();
 
 	/** Helper functions used to construct/operate-on the "Save on Compile" command */
-	bool IsSaveOnCompileEnabled() const;
+	UE_API bool IsSaveOnCompileEnabled() const;
 
 	/** Calls the above function, but returns an FReply::Handled(). Used in SButtons */
-	virtual FReply Compile_OnClickWithReply();
+	UE_API virtual FReply Compile_OnClickWithReply();
 
 	/** Called when the refresh all nodes button is clicked */
-	void RefreshAllNodes_OnClicked();
+	UE_API void RefreshAllNodes_OnClicked();
 
-	EVisibility IsDebuggerVisible() const;
+	UE_API EVisibility IsDebuggerVisible() const;
 
-	virtual void OnActiveTabChanged(TSharedPtr<SDockTab> PreviouslyActive, TSharedPtr<SDockTab> NewlyActivated);
+	UE_API virtual void OnActiveTabChanged(TSharedPtr<SDockTab> PreviouslyActive, TSharedPtr<SDockTab> NewlyActivated);
 
 	/** Called when a token in a log message is clicked */
-	void LogToken_OnClicked(const class IMessageToken& Token);
+	UE_API void LogToken_OnClicked(const class IMessageToken& Token);
 
-	virtual void FocusInspectorOnGraphSelection(const TSet<class UObject*>& NewSelection, bool bForceRefresh = false);
+	UE_API virtual void FocusInspectorOnGraphSelection(const TSet<class UObject*>& NewSelection, bool bForceRefresh = false);
 
 	/** Variable list window calls this after it is updated */
-	void VariableListWasUpdated();
+	UE_API void VariableListWasUpdated();
 
 	/** Virtual override point for editing defaults; allowing more derived editors to edit something else */
-	virtual void StartEditingDefaults(bool bAutoFocus = true, bool bForceRefresh = false);
+	UE_API virtual void StartEditingDefaults(bool bAutoFocus = true, bool bForceRefresh = false);
 
 	// Called by the blueprint editing app mode to focus the appropriate tabs, etc...
-	void SetupViewForBlueprintEditingMode();
+	UE_API void SetupViewForBlueprintEditingMode();
 
 	// Ensures the blueprint is up to date
-	void EnsureBlueprintIsUpToDate(UBlueprint* BlueprintObj);
+	UE_API void EnsureBlueprintIsUpToDate(UBlueprint* BlueprintObj);
 
 	// Should be called when initializing any editor built off this foundation
-	void CommonInitialization(const TArray<UBlueprint*>& InitBlueprints, bool bShouldOpenInDefaultsMode);
+	UE_API void CommonInitialization(const TArray<UBlueprint*>& InitBlueprints, bool bShouldOpenInDefaultsMode);
 
 	// Other types of BP editors that don't use BP VM/function libs can choose to not load BP function libraries, for example: Control Rig Blueprints
 	virtual bool ShouldLoadBPLibrariesFromAssetRegistry() { return true; }
 
 	// Should be called when initializing an editor that has a blueprint, after layout (tab spawning) is done
-	void PostLayoutBlueprintEditorInitialization();
+	UE_API void PostLayoutBlueprintEditorInitialization();
 
 	/** Called when graph editor focus is changed */
-	virtual void OnGraphEditorFocused(const TSharedRef<class SGraphEditor>& InGraphEditor);
+	UE_API virtual void OnGraphEditorFocused(const TSharedRef<class SGraphEditor>& InGraphEditor);
 
 	/** Called when the graph editor tab is backgrounded */
-	virtual void OnGraphEditorBackgrounded(const TSharedRef<SGraphEditor>& InGraphEditor);
+	UE_API virtual void OnGraphEditorBackgrounded(const TSharedRef<SGraphEditor>& InGraphEditor);
 
 	/** Enable/disable the subobject editor preview viewport */
-	void EnableSubobjectPreview(bool bEnable);
+	UE_API void EnableSubobjectPreview(bool bEnable);
 
 	/** Refresh the preview viewport to reflect changes in the subobject */
-	void UpdateSubobjectPreview(bool bUpdateNow = false);
+	UE_API void UpdateSubobjectPreview(bool bUpdateNow = false);
 
 	///////////////////////////////////////////////////////
 	// Subobject Editor Callbacks
 	/** Delegate invoked when the Subobject editor needs to obtain the object context for editing */
-	UObject* GetSubobjectEditorObjectContext() const;
+	UE_API UObject* GetSubobjectEditorObjectContext() const;
 
 	/** Delegate invoked when the selection is changed in the subobject editor widget */
-	virtual void OnSelectionUpdated(const TArray<TSharedPtr<class FSubobjectEditorTreeNode>>& SelectedNodes);
+	UE_API virtual void OnSelectionUpdated(const TArray<TSharedPtr<class FSubobjectEditorTreeNode>>& SelectedNodes);
 
 	/** Delegate invoked when an item is double clicked in the subobject editor widget */
-	virtual void OnComponentDoubleClicked(TSharedPtr<class FSubobjectEditorTreeNode> Node);
+	UE_API virtual void OnComponentDoubleClicked(TSharedPtr<class FSubobjectEditorTreeNode> Node);
 
 	/** Delegate invoked when a new subobject is added in the subobject editor widget */
-	virtual void OnComponentAddedToBlueprint(const FSubobjectData& NewSubobjectData);
+	UE_API virtual void OnComponentAddedToBlueprint(const FSubobjectData& NewSubobjectData);
 	
 	/** Delegate invoked when the selection is changed in the subobject editor widget */
 	UE_DEPRECATED(5.0, "OnSelectionUpdated(const TArray<TSharedPtr<class FSCSEditorTreeNode>>&) has been deprecated. Use OnSelectionUpdated(const TArray<TSharedPtr<class FSubobjectEditorTreeNode>>&) instead.")
@@ -586,29 +588,29 @@ public:
 	virtual void OnComponentDoubleClicked(TSharedPtr<class FSCSEditorTreeNode> Node) {}
 
 	/** Pin visibility accessors */
-	void SetPinVisibility(SGraphEditor::EPinVisibility Visibility);
+	UE_API void SetPinVisibility(SGraphEditor::EPinVisibility Visibility);
 	bool GetPinVisibility(SGraphEditor::EPinVisibility Visibility) const { return PinVisibility == Visibility; }
 
 	/** Reparent the current blueprint */
-	void ReparentBlueprint_Clicked();
-	virtual bool ReparentBlueprint_IsVisible() const;
-	void ReparentBlueprint_NewParentChosen(UClass* ChosenClass);
+	UE_API void ReparentBlueprint_Clicked();
+	UE_API virtual bool ReparentBlueprint_IsVisible() const;
+	UE_API void ReparentBlueprint_NewParentChosen(UClass* ChosenClass);
 
 	/** Utility function to handle all steps required to rename a newly added action */
-	void RenameNewlyAddedAction(FName InActionName);
+	UE_API void RenameNewlyAddedAction(FName InActionName);
 
 	/** Adds a new variable to this blueprint */
-	void OnAddNewVariable();
+	UE_API void OnAddNewVariable();
 	FReply OnAddNewVariable_OnClick() { OnAddNewVariable(); return FReply::Handled(); }
 
 	/** Checks if adding a local variable is allowed in the focused graph */
-	virtual bool CanAddNewLocalVariable() const;
+	UE_API virtual bool CanAddNewLocalVariable() const;
 
 	/** Adds a new local variable to the focused function graph */
-	virtual void OnAddNewLocalVariable();
+	UE_API virtual void OnAddNewLocalVariable();
 
 	/** Pastes a new local variable to the focused function graph */
-	virtual void OnPasteNewLocalVariable(const FBPVariableDescription& VariableDescription);
+	UE_API virtual void OnPasteNewLocalVariable(const FBPVariableDescription& VariableDescription);
 
 	/** Create a variant for a function */
 	virtual void AddNewFunctionVariant(const UEdGraph* InOriginalFunction) {}
@@ -625,40 +627,40 @@ public:
 	};
 
 	/** Called when New Function button is clicked */
-	virtual void NewDocument_OnClicked(ECreatedDocumentType GraphType);
+	UE_API virtual void NewDocument_OnClicked(ECreatedDocumentType GraphType);
 	FReply NewDocument_OnClick(ECreatedDocumentType GraphType) { NewDocument_OnClicked(GraphType); return FReply::Handled(); }
 
 	/** Called when New Delegate button is clicked */
-	void OnAddNewDelegate();
-	bool AddNewDelegateIsVisible() const;
+	UE_API void OnAddNewDelegate();
+	UE_API bool AddNewDelegateIsVisible() const;
 
 	// Called to see if the new document menu items is visible for this type
 	virtual bool IsSectionVisible(NodeSectionID::Type InSectionID) const { return true; }
-	virtual bool NewDocument_IsVisibleForType(ECreatedDocumentType GraphType) const;
+	UE_API virtual bool NewDocument_IsVisibleForType(ECreatedDocumentType GraphType) const;
 	EVisibility NewDocument_GetVisibilityForType(ECreatedDocumentType GraphType) const
 	{
 		return NewDocument_IsVisibleForType(GraphType) ? EVisibility::Visible : EVisibility::Collapsed;
 	}
 
-	static FName SelectionState_MyBlueprint;
-	static FName SelectionState_Components;
-	static FName SelectionState_Graph;
-	static FName SelectionState_ClassSettings;
-	static FName SelectionState_ClassDefaults;
+	static UE_API FName SelectionState_MyBlueprint;
+	static UE_API FName SelectionState_Components;
+	static UE_API FName SelectionState_Graph;
+	static UE_API FName SelectionState_ClassSettings;
+	static UE_API FName SelectionState_ClassDefaults;
 
 	/** Gets or sets the flag for context sensitivity in the graph action menu */
 	bool& GetIsContextSensitive() { return bIsActionMenuContextSensitive; }
 
 	/** Gets the UI selection state of this editor */
 	FName GetUISelectionState() const { return CurrentUISelection; }
-	void SetUISelectionState(FName SelectionOwner);
+	UE_API void SetUISelectionState(FName SelectionOwner);
 
-	virtual void ClearSelectionStateFor(FName SelectionOwner);
+	UE_API virtual void ClearSelectionStateFor(FName SelectionOwner);
 
 	/** Handles spawning a graph node in the current graph using the passed in chord */
 	UE_DEPRECATED(5.6, "Please use the version of the function accepting FVector2f.")
-	virtual FReply OnSpawnGraphNodeByShortcut(FInputChord InChord, const FVector2D& InPosition, UEdGraph* InGraph) UE_SLATE_DEPRECATED_VECTOR_VIRTUAL_FUNCTION;
-	virtual FReply OnSpawnGraphNodeByShortcut(FInputChord InChord, const FVector2f& InPosition, UEdGraph* InGraph);
+	UE_API virtual FReply OnSpawnGraphNodeByShortcut(FInputChord InChord, const FVector2D& InPosition, UEdGraph* InGraph) UE_SLATE_DEPRECATED_VECTOR_VIRTUAL_FUNCTION;
+	UE_API virtual FReply OnSpawnGraphNodeByShortcut(FInputChord InChord, const FVector2f& InPosition, UEdGraph* InGraph);
 
 	/**
 	 * Perform the actual promote to variable action on the given pin in the given blueprint.
@@ -669,20 +671,20 @@ public:
 	 * @param   InOptionalLocation		Where the new node should be placed. If null, a fixed offset from the parent node will be used.
 	 */
 	UE_DEPRECATED(5.6, "Use the version of this function accepting a TOptional")
-	void DoPromoteToVariable(UBlueprint* InBlueprint, UEdGraphPin* InTargetPin, bool bInToMemberVariable, const FVector2D* InOptionalLocation = nullptr);
-	void DoPromoteToVariable2f(UBlueprint* InBlueprint, UEdGraphPin* InTargetPin, bool bInToMemberVariable, TOptional<FVector2f> InOptionalLocation = TOptional<FVector2f>());
+	UE_API void DoPromoteToVariable(UBlueprint* InBlueprint, UEdGraphPin* InTargetPin, bool bInToMemberVariable, const FVector2D* InOptionalLocation = nullptr);
+	UE_API void DoPromoteToVariable2f(UBlueprint* InBlueprint, UEdGraphPin* InTargetPin, bool bInToMemberVariable, TOptional<FVector2f> InOptionalLocation = TOptional<FVector2f>());
 
 	/** Called when node is spawned by keymap */
-	void OnNodeSpawnedByKeymap();
+	UE_API void OnNodeSpawnedByKeymap();
 
 	/** Update Node Creation mechanisms for analytics */
-	void UpdateNodeCreationStats(const ENodeCreateAction::Type CreateAction);
+	UE_API void UpdateNodeCreationStats(const ENodeCreateAction::Type CreateAction);
 
 	/** Sets customizations for the BP editor details panel. */
-	void SetDetailsCustomization(TSharedPtr<class FDetailsViewObjectFilter> DetailsObjectFilter, TSharedPtr<class IDetailRootObjectCustomization> DetailsRootCustomization);
+	UE_API void SetDetailsCustomization(TSharedPtr<class FDetailsViewObjectFilter> DetailsObjectFilter, TSharedPtr<class IDetailRootObjectCustomization> DetailsRootCustomization);
 
 	/** Sets subobject editor UI customization */
-	void SetSubobjectEditorUICustomization(TSharedPtr<class ISCSEditorUICustomization> SCSEditorUICustomization);
+	UE_API void SetSubobjectEditorUICustomization(TSharedPtr<class ISCSEditorUICustomization> SCSEditorUICustomization);
 
 	UE_DEPRECATED(5.0, "SetSCSEditorUICustomization has been deprecated. Use SetSubobjectEditorUICustomization instead.")
 	void SetSCSEditorUICustomization(TSharedPtr<class ISCSEditorUICustomization> SCSEditorUICustomization)
@@ -695,33 +697,33 @@ public:
 	 * @param	InComponentName			The name of the component to customize behavior for
 	 * @param	InCustomization			The customization instance to use
 	 */
-	void RegisterSCSEditorCustomization(const FName& InComponentName, TSharedPtr<class ISCSEditorCustomization> InCustomization);
+	UE_API void RegisterSCSEditorCustomization(const FName& InComponentName, TSharedPtr<class ISCSEditorCustomization> InCustomization);
 
 	/**
 	 * Unregister a previously registered customization for interacting with the subobject editor
 	 * @param	InComponentName			The name of the component to customize behavior for
 	 */
-	void UnregisterSCSEditorCustomization(const FName& InComponentName);
+	UE_API void UnregisterSCSEditorCustomization(const FName& InComponentName);
 
 	/** Forces the merge tool to be shown */
-	void CreateMergeToolTab();
-	void CreateMergeToolTab(const UBlueprint* BaseBlueprint, const UBlueprint* RemoteBlueprint, const FOnMergeResolved& ResolutionCallback);
+	UE_API void CreateMergeToolTab();
+	UE_API void CreateMergeToolTab(const UBlueprint* BaseBlueprint, const UBlueprint* RemoteBlueprint, const FOnMergeResolved& ResolutionCallback);
 
 	/** Closes the merge tool, rather than simply hiding it */
-	void CloseMergeTool();
+	UE_API void CloseMergeTool();
 
 	/** Dumps the current blueprint search index to a JSON file for debugging purposes */
-	void OnGenerateSearchIndexForDebugging();
+	UE_API void OnGenerateSearchIndexForDebugging();
 
 	/** Dumps the currently-cached index data for the blueprint to a file for debugging */
-	void OnDumpCachedIndexDataForBlueprint();
+	UE_API void OnDumpCachedIndexDataForBlueprint();
 
 	/**
 	 * Check to see if we can customize the subobject editor for the passed-in scene component
 	 * @param	InComponentToCustomize	The component to check to see if a customization exists
 	 * @return an subobject editor customization instance, if one exists.
 	 */
-	TSharedPtr<class ISCSEditorCustomization> CustomizeSubobjectEditor(const USceneComponent* InComponentToCustomize) const;
+	UE_API TSharedPtr<class ISCSEditorCustomization> CustomizeSubobjectEditor(const USceneComponent* InComponentToCustomize) const;
 
 	UE_DEPRECATED(5.0, "CustomizeSCSEditor has been deprecated. Use CustomizeSubobjectEditor instead.")
 	TSharedPtr<class ISCSEditorCustomization> CustomizeSCSEditor(USceneComponent* InComponentToCustomize) const
@@ -739,40 +741,40 @@ public:
 	virtual FString GetCustomDebugObjectLabel(UObject* ObjectBeingDebugged) const { return FString(); }
 
 	/** Called when a node's title is committed for a rename */
-	void OnNodeTitleCommitted(const FText& NewText, ETextCommit::Type CommitInfo, UEdGraphNode* NodeBeingChanged);
+	UE_API void OnNodeTitleCommitted(const FText& NewText, ETextCommit::Type CommitInfo, UEdGraphNode* NodeBeingChanged);
 
 	/** Called by a graph title bar to get any extra information the editor would like to display */
-	virtual FText GetGraphDecorationString(UEdGraph* InGraph) const;
+	UE_API virtual FText GetGraphDecorationString(UEdGraph* InGraph) const;
 
 	/** Gets the display name of a graph */
-	static FText GetGraphDisplayName(const UEdGraph* Graph);
+	static UE_API FText GetGraphDisplayName(const UEdGraph* Graph);
 
 	/** Checks to see if the provided graph is contained within the current blueprint */
-	bool IsGraphInCurrentBlueprint(const UEdGraph* InGraph) const;
+	UE_API bool IsGraphInCurrentBlueprint(const UEdGraph* InGraph) const;
 
 	/** Get the context to use from the Blueprint type */
-	static FName GetContextFromBlueprintType(EBlueprintType InType);
+	static UE_API FName GetContextFromBlueprintType(EBlueprintType InType);
 
 	/* Selects an item in "My Blueprint" by name. */
-	void SelectGraphActionItemByName(const FName& ItemName, ESelectInfo::Type SelectInfo = ESelectInfo::Direct, int32 SectionId = INDEX_NONE, bool bIsCategory = false);
+	UE_API void SelectGraphActionItemByName(const FName& ItemName, ESelectInfo::Type SelectInfo = ESelectInfo::Direct, int32 SectionId = INDEX_NONE, bool bIsCategory = false);
 
 	/** Handle when the debug object is changed in the UI */
 	virtual void HandleSetObjectBeingDebugged(UObject* InObject) {}
 
 	/** Adds a new bookmark node to the Blueprint that's being edited. */
-	FBPEditorBookmarkNode* AddBookmark(const FText& DisplayName, const FEditedDocumentInfo& BookmarkInfo, bool bSharedBookmark = false);
+	UE_API FBPEditorBookmarkNode* AddBookmark(const FText& DisplayName, const FEditedDocumentInfo& BookmarkInfo, bool bSharedBookmark = false);
 
 	/** Renames the bookmark node with the given ID. */
-	void RenameBookmark(const FGuid& BookmarkNodeId, const FText& NewName);
+	UE_API void RenameBookmark(const FGuid& BookmarkNodeId, const FText& NewName);
 
 	/** Removes the bookmark node with the given ID. */
-	void RemoveBookmark(const FGuid& BookmarkNodeId, bool bRefreshUI = true);
+	UE_API void RemoveBookmark(const FGuid& BookmarkNodeId, bool bRefreshUI = true);
 
 	/** Gets the default schema for this editor */
 	TSubclassOf<UEdGraphSchema> GetDefaultSchema() const { return GetDefaultSchemaClass(); }
 	
 	/** Imports the given namespace into the editor. This may trigger a load event for additional macro and/or function library assets if not already loaded. */
-	void ImportNamespace(const FString& InNamespace);
+	UE_API void ImportNamespace(const FString& InNamespace);
 
 	/** Parameters for the extended ImportNamespaceEx() method */
 	struct FImportNamespaceExParameters
@@ -795,234 +797,234 @@ public:
 	};
 
 	/** Imports a set of namespace(s) into the editor. This may trigger a load event for additional macro and/or function library assets if not already loaded. */
-	void ImportNamespaceEx(const FImportNamespaceExParameters& InParams);
+	UE_API void ImportNamespaceEx(const FImportNamespaceExParameters& InParams);
 
 	/** Removes the given namespace from the editor's current import context. However, this will NOT unload any associated macro and/or function library assets. */
-	void RemoveNamespace(const FString& InNamespace);
+	UE_API void RemoveNamespace(const FString& InNamespace);
 
 	/** Selects a local variable to load in the details panel. */
 	virtual bool SelectLocalVariable(const UEdGraph* Graph, const FName& VariableName) { return false; }
 
 	/** Duplicates an EdGraphNode using selection and copy / paste functionality */
-	void SelectAndDuplicateNode(UEdGraphNode* InNode);
+	UE_API void SelectAndDuplicateNode(UEdGraphNode* InNode);
 
 protected:
 
 	/** Called during initialization of the blueprint editor to register commands and extenders. */
-	virtual void InitalizeExtenders();
+	UE_API virtual void InitalizeExtenders();
 
 	/** Called during initialization of the blueprint editor to register any application modes. */
-	virtual void RegisterApplicationModes(const TArray<UBlueprint*>& InBlueprints, bool bShouldOpenInDefaultsMode, bool bNewlyCreated = false);
+	UE_API virtual void RegisterApplicationModes(const TArray<UBlueprint*>& InBlueprints, bool bShouldOpenInDefaultsMode, bool bNewlyCreated = false);
 
 	// Updates the selected object used by the stand lone defaults editor widget.
-	void RefreshStandAloneDefaultsEditor();
+	UE_API void RefreshStandAloneDefaultsEditor();
 
 	// Zooming to fit the entire graph
-	void ZoomToWindow_Clicked();
-	bool CanZoomToWindow() const;
+	UE_API void ZoomToWindow_Clicked();
+	UE_API bool CanZoomToWindow() const;
 
 	// Zooming to fit the current selection
-	void ZoomToSelection_Clicked();
-	bool CanZoomToSelection() const;
+	UE_API void ZoomToSelection_Clicked();
+	UE_API bool CanZoomToSelection() const;
 
 	// Navigating into/out of graphs
-	void NavigateToParentGraph();
-	void NavigateToParentGraphByDoubleClick();
-	bool CanNavigateToParentGraph() const;
-	void NavigateToChildGraph();
-	bool CanNavigateToChildGraph() const;
+	UE_API void NavigateToParentGraph();
+	UE_API void NavigateToParentGraphByDoubleClick();
+	UE_API bool CanNavigateToParentGraph() const;
+	UE_API void NavigateToChildGraph();
+	UE_API bool CanNavigateToChildGraph() const;
 
 	/** Determines visibility of the find parent class in content browser button on the menu bar overlay */
-	EVisibility GetFindParentClassVisibility() const;
+	UE_API EVisibility GetFindParentClassVisibility() const;
 
 	/** Determines visibility of the edit parent class button on the menu bar overlay */
-	EVisibility GetEditParentClassVisibility() const;
+	UE_API EVisibility GetEditParentClassVisibility() const;
 
 	/** Recreates the overlay on the menu bar */
-	virtual void PostRegenerateMenusAndToolbars() override;
+	UE_API virtual void PostRegenerateMenusAndToolbars() override;
 
 	/** Returns the name of the Blueprint's parent class */
-	FText GetParentClassNameText() const;
+	UE_API FText GetParentClassNameText() const;
 
 	/** Handler for "Find parent class in CB" button */
-	FReply OnFindParentClassInContentBrowserClicked();
+	UE_API FReply OnFindParentClassInContentBrowserClicked();
 
 	/** Handler for "Edit parent class" button */
-	FReply OnEditParentClassClicked();
+	UE_API FReply OnEditParentClassClicked();
 
 	/** Called to start a quick find (focus the search box in the explorer tab) */
-	void FindInBlueprint_Clicked();
+	UE_API void FindInBlueprint_Clicked();
 
 	// Is the main details panel currently showing 'Global options' (e.g., class metadata)?
-	bool IsDetailsPanelEditingGlobalOptions() const;
+	UE_API bool IsDetailsPanelEditingGlobalOptions() const;
 
 	/** Edit the class settings aka Blueprint global options */
-	void EditGlobalOptions_Clicked();
+	UE_API void EditGlobalOptions_Clicked();
 
 	// Is the main details panel currently showing 'Class defaults' (Note: Has nothing to do with the standalone class defaults panel)?
-	bool IsDetailsPanelEditingClassDefaults() const;
+	UE_API bool IsDetailsPanelEditingClassDefaults() const;
 
 	/** Edit the class defaults */
-	void EditClassDefaults_Clicked();
+	UE_API void EditClassDefaults_Clicked();
 
 	/** Called to undo the last action */
-	void UndoGraphAction();
+	UE_API void UndoGraphAction();
 
 	/** Whether or not we can perform an undo of the last transacted action */
-	bool CanUndoGraphAction() const;
+	UE_API bool CanUndoGraphAction() const;
 
 	/** Called to redo the last undone action */
-	void RedoGraphAction();
+	UE_API void RedoGraphAction();
 
 	/** Whether or not we can redo an undone action */
-	bool CanRedoGraphAction() const;
+	UE_API bool CanRedoGraphAction() const;
 	
 	/** Called when the selection changes in the GraphEditor */
-	virtual void OnSelectedNodesChangedImpl(const TSet<class UObject*>& NewSelection);
+	UE_API virtual void OnSelectedNodesChangedImpl(const TSet<class UObject*>& NewSelection);
 
 	/** Called when the selection changes in the GraphEditor */
 	void OnSelectedNodesChanged(const TSet<class UObject*>& NewSelection) { OnSelectedNodesChangedImpl(NewSelection); }
 
 	/** Called when an actor is dropped onto the graph editor */
-	void OnGraphEditorDropActor(const TArray< TWeakObjectPtr<AActor> >& Actors, UEdGraph* Graph, const FVector2f& DropLocation);
+	UE_API void OnGraphEditorDropActor(const TArray< TWeakObjectPtr<AActor> >& Actors, UEdGraph* Graph, const FVector2f& DropLocation);
 
 	/** Called when a streaming level is dropped onto the graph editor */
-	void OnGraphEditorDropStreamingLevel(const TArray< TWeakObjectPtr<ULevelStreaming> >& Levels, UEdGraph* Graph, const FVector2f& DropLocation);
+	UE_API void OnGraphEditorDropStreamingLevel(const TArray< TWeakObjectPtr<ULevelStreaming> >& Levels, UEdGraph* Graph, const FVector2f& DropLocation);
 
 	/** Called to create context menu when right-clicking on graph */
 	UE_DEPRECATED(5.6, "Please use the version of the function taking FVector2f.")
-	FActionMenuContent OnCreateGraphActionMenu(UEdGraph* InGraph, const FVector2D& InNodePosition, const TArray<UEdGraphPin*>& InDraggedPins, bool bAutoExpand, SGraphEditor::FActionMenuClosed InOnMenuClosed);
-	FActionMenuContent OnCreateGraphActionMenu(UEdGraph* InGraph, const FVector2f& InNodePosition, const TArray<UEdGraphPin*>& InDraggedPins, bool bAutoExpand, SGraphEditor::FActionMenuClosed InOnMenuClosed);
+	UE_API FActionMenuContent OnCreateGraphActionMenu(UEdGraph* InGraph, const FVector2D& InNodePosition, const TArray<UEdGraphPin*>& InDraggedPins, bool bAutoExpand, SGraphEditor::FActionMenuClosed InOnMenuClosed);
+	UE_API FActionMenuContent OnCreateGraphActionMenu(UEdGraph* InGraph, const FVector2f& InNodePosition, const TArray<UEdGraphPin*>& InDraggedPins, bool bAutoExpand, SGraphEditor::FActionMenuClosed InOnMenuClosed);
 
 	/** Called from graph context menus when they close to tell the editor why they closed */
-	void OnGraphActionMenuClosed(bool bActionExecuted, bool bContextSensitiveChecked, bool bGraphPinContext);
+	UE_API void OnGraphActionMenuClosed(bool bActionExecuted, bool bContextSensitiveChecked, bool bGraphPinContext);
 
 	/** Called when the Blueprint we are editing has changed */
-	virtual void OnBlueprintChangedImpl(UBlueprint* InBlueprint, bool bIsJustBeingCompiled = false);
+	UE_API virtual void OnBlueprintChangedImpl(UBlueprint* InBlueprint, bool bIsJustBeingCompiled = false);
 
 	/** Called when the Blueprint we are editing has changed, forwards to impl */
 	void OnBlueprintChanged(UBlueprint* InBlueprint) { return OnBlueprintChangedImpl(InBlueprint); }
 
-	void OnBlueprintCompiled(UBlueprint* InBlueprint);
+	UE_API void OnBlueprintCompiled(UBlueprint* InBlueprint);
 
 	/** Handles the unloading of Blueprints (by closing the editor, if it operating on the Blueprint being unloaded)*/
-	void OnBlueprintUnloaded(UBlueprint* InBlueprint);
+	UE_API void OnBlueprintUnloaded(UBlueprint* InBlueprint);
 
 	/** Called when a property change is about to be propagated to instances of the Blueprint */
-	void OnPreObjectPropertyChanged(UObject* InObject, const FEditPropertyChain& EditPropertyChain);
+	UE_API void OnPreObjectPropertyChanged(UObject* InObject, const FEditPropertyChain& EditPropertyChain);
 
 	/** Called after a property change has been propagated to instances of the Blueprint */
-	void OnPostObjectPropertyChanged(UObject* InObject, FPropertyChangedEvent& PropertyChangedEvent);
+	UE_API void OnPostObjectPropertyChanged(UObject* InObject, FPropertyChangedEvent& PropertyChangedEvent);
 
 	//@TODO: Should the breakpoint/watch modification operations be whole-blueprint, or current-graph?
 
 	/** Deletes all breakpoints for the blueprint being edited */
-	void ClearAllBreakpoints();
+	UE_API void ClearAllBreakpoints();
 
 	/** Disables all breakpoints for the blueprint being edited */
-	void DisableAllBreakpoints();
+	UE_API void DisableAllBreakpoints();
 
 	/** Enables all breakpoints for the blueprint being edited */
-	void EnableAllBreakpoints();
+	UE_API void EnableAllBreakpoints();
 
 	/** Clears all watches associated with the blueprint being edited */
-	void ClearAllWatches();
+	UE_API void ClearAllWatches();
 
-	bool HasAnyBreakpoints() const;
-	bool HasAnyEnabledBreakpoints() const;
-	bool HasAnyDisabledBreakpoints() const;
-	bool HasAnyWatches() const;
+	UE_API bool HasAnyBreakpoints() const;
+	UE_API bool HasAnyEnabledBreakpoints() const;
+	UE_API bool HasAnyDisabledBreakpoints() const;
+	UE_API bool HasAnyWatches() const;
 
 	/** Spawns a new blueprint debugger tab */
-	void OpenBlueprintDebugger();
-	bool CanOpenBlueprintDebugger() const;
+	UE_API void OpenBlueprintDebugger();
+	UE_API bool CanOpenBlueprintDebugger() const;
 
 	// Utility helper to get the currently hovered pin in the currently visible graph, or nullptr if there isn't one
-	UEdGraphPin* GetCurrentlySelectedPin() const;
+	UE_API UEdGraphPin* GetCurrentlySelectedPin() const;
 
 	// UI Action functionality
-	void OnPromoteToVariable(bool bInToMemberVariable);
-	bool CanPromoteToVariable(bool bInToMemberVariable) const;
+	UE_API void OnPromoteToVariable(bool bInToMemberVariable);
+	UE_API bool CanPromoteToVariable(bool bInToMemberVariable) const;
 
-	void OnSplitStructPin();
-	bool CanSplitStructPin() const;
+	UE_API void OnSplitStructPin();
+	UE_API bool CanSplitStructPin() const;
 
-	void OnRecombineStructPin();
-	bool CanRecombineStructPin() const;
+	UE_API void OnRecombineStructPin();
+	UE_API bool CanRecombineStructPin() const;
 
-	void OnAddExecutionPin();
-	bool CanAddExecutionPin() const;
+	UE_API void OnAddExecutionPin();
+	UE_API bool CanAddExecutionPin() const;
 
-	void OnInsertExecutionPinBefore();
-	void OnInsertExecutionPinAfter();
-	void OnInsertExecutionPin(EPinInsertPosition Position);
-	bool CanInsertExecutionPin() const;
+	UE_API void OnInsertExecutionPinBefore();
+	UE_API void OnInsertExecutionPinAfter();
+	UE_API void OnInsertExecutionPin(EPinInsertPosition Position);
+	UE_API bool CanInsertExecutionPin() const;
 
-	void OnRemoveExecutionPin();
-	bool CanRemoveExecutionPin() const;
+	UE_API void OnRemoveExecutionPin();
+	UE_API bool CanRemoveExecutionPin() const;
 
-	void OnRemoveThisStructVarPin();
-	bool CanRemoveThisStructVarPin() const;
+	UE_API void OnRemoveThisStructVarPin();
+	UE_API bool CanRemoveThisStructVarPin() const;
 
-	void OnRemoveOtherStructVarPins();
-	bool CanRemoveOtherStructVarPins() const;
+	UE_API void OnRemoveOtherStructVarPins();
+	UE_API bool CanRemoveOtherStructVarPins() const;
 
-	void OnRestoreAllStructVarPins();
-	bool CanRestoreAllStructVarPins() const;
+	UE_API void OnRestoreAllStructVarPins();
+	UE_API bool CanRestoreAllStructVarPins() const;
 
-	void OnResetPinToDefaultValue();
-	bool CanResetPinToDefaultValue() const;
+	UE_API void OnResetPinToDefaultValue();
+	UE_API bool CanResetPinToDefaultValue() const;
 
-	void OnAddOptionPin();
-	bool CanAddOptionPin() const;
+	UE_API void OnAddOptionPin();
+	UE_API bool CanAddOptionPin() const;
 
-	void OnRemoveOptionPin();
-	bool CanRemoveOptionPin() const;
+	UE_API void OnRemoveOptionPin();
+	UE_API bool CanRemoveOptionPin() const;
 
 	/** Functions for handling the changing of the pin's type (PinCategory, PinSubCategory, etc) */
-	FEdGraphPinType OnGetPinType(UEdGraphPin* SelectedPin) const;
-	void OnChangePinType();
-	void OnChangePinTypeFinished(const FEdGraphPinType& PinType, UEdGraphPin* SelectedPin);
-	bool CanChangePinType() const;
+	UE_API FEdGraphPinType OnGetPinType(UEdGraphPin* SelectedPin) const;
+	UE_API void OnChangePinType();
+	UE_API void OnChangePinTypeFinished(const FEdGraphPinType& PinType, UEdGraphPin* SelectedPin);
+	UE_API bool CanChangePinType() const;
 
-	void OnAddParentNode();
-	bool CanAddParentNode() const;
+	UE_API void OnAddParentNode();
+	UE_API bool CanAddParentNode() const;
 	
-	void OnCreateMatchingFunction();
-	bool CanCreateMatchingFunction() const;
+	UE_API void OnCreateMatchingFunction();
+	UE_API bool CanCreateMatchingFunction() const;
 
-	void OnEnableBreakpoint();
-	bool CanEnableBreakpoint() const;
+	UE_API void OnEnableBreakpoint();
+	UE_API bool CanEnableBreakpoint() const;
 
-	void OnToggleBreakpoint();
-	bool CanToggleBreakpoint() const;
+	UE_API void OnToggleBreakpoint();
+	UE_API bool CanToggleBreakpoint() const;
 
-	void OnDisableBreakpoint();
-	bool CanDisableBreakpoint() const;
+	UE_API void OnDisableBreakpoint();
+	UE_API bool CanDisableBreakpoint() const;
 
-	void OnAddBreakpoint();
-	bool CanAddBreakpoint() const;
+	UE_API void OnAddBreakpoint();
+	UE_API bool CanAddBreakpoint() const;
 
-	void OnRemoveBreakpoint();
-	bool CanRemoveBreakpoint() const;
+	UE_API void OnRemoveBreakpoint();
+	UE_API bool CanRemoveBreakpoint() const;
 
-	void OnCollapseNodes();
-	bool CanCollapseNodes() const;
+	UE_API void OnCollapseNodes();
+	UE_API bool CanCollapseNodes() const;
 
-	void OnCollapseSelectionToFunction();
-	bool CanCollapseSelectionToFunction() const;
+	UE_API void OnCollapseSelectionToFunction();
+	UE_API bool CanCollapseSelectionToFunction() const;
 
-	void OnCollapseSelectionToMacro();
-	bool CanCollapseSelectionToMacro() const;
+	UE_API void OnCollapseSelectionToMacro();
+	UE_API bool CanCollapseSelectionToMacro() const;
 
-	void OnPromoteSelectionToFunction();
-	bool CanPromoteSelectionToFunction() const;
+	UE_API void OnPromoteSelectionToFunction();
+	UE_API bool CanPromoteSelectionToFunction() const;
 
-	void OnPromoteSelectionToMacro();
-	bool CanPromoteSelectionToMacro() const;
+	UE_API void OnPromoteSelectionToMacro();
+	UE_API bool CanPromoteSelectionToMacro() const;
 
-	void OnExpandNodes();
-	bool CanExpandNodes() const;
+	UE_API void OnExpandNodes();
+	UE_API bool CanExpandNodes() const;
 
 	/**
 	* Move the given set of nodes to an average spot near the Source position
@@ -1031,23 +1033,23 @@ protected:
 	* @param SourcePos						The source position used to average the nodes around
 	* @param bExpandedNodesNeedUniqueGuid	If true then a new Guid will be generated for each node in the set
 	*/
-	void MoveNodesToAveragePos(TSet<UEdGraphNode*>& AverageNodes, FVector2D SourcePos, bool bExpandedNodesNeedUniqueGuid = false) const;
+	UE_API void MoveNodesToAveragePos(TSet<UEdGraphNode*>& AverageNodes, FVector2D SourcePos, bool bExpandedNodesNeedUniqueGuid = false) const;
 
-	void OnConvertFunctionToEvent();
+	UE_API void OnConvertFunctionToEvent();
 
 public:
 	/** Converts the given function entry node to an event if it passes validation */
-	bool ConvertFunctionIfValid(UK2Node_FunctionEntry* FuncEntryNode);
+	UE_API bool ConvertFunctionIfValid(UK2Node_FunctionEntry* FuncEntryNode);
 
 	/** Converts the given event node to a function graph on this blueprint if it passes validation */
-	bool ConvertEventIfValid(UK2Node_Event* EventToConv);
+	UE_API bool ConvertEventIfValid(UK2Node_Event* EventToConv);
 
 protected:
 	/** 
 	* Callback function for the context menu on a node to determine if a function 
 	* could possibly be converted to an event
 	*/
-	bool CanConvertFunctionToEvent() const;
+	UE_API bool CanConvertFunctionToEvent() const;
 
 	/*
 	* Given a function node, move all nodes from the function out of the function graph,
@@ -1055,11 +1057,11 @@ protected:
 	* 
 	* @param SelectedCallFunctionNode	The function node to convert to an event
 	*/
-	void ConvertFunctionToEvent(UK2Node_FunctionEntry* SelectedCallFunctionNode);
+	UE_API void ConvertFunctionToEvent(UK2Node_FunctionEntry* SelectedCallFunctionNode);
 
-	void OnConvertEventToFunction();
-	void ConvertEventToFunction(UK2Node_Event* SelectedEventNode);
-	bool CanConvertEventToFunction() const;
+	UE_API void OnConvertEventToFunction();
+	UE_API void ConvertEventToFunction(UK2Node_Event* SelectedEventNode);
+	UE_API bool CanConvertEventToFunction() const;
 
 	/**
 	* Get all connected nodes and add them to a set of out nodes
@@ -1067,28 +1069,28 @@ protected:
 	* @param Node		The node to get all connects to 
 	* @return	TArray of all connected nodes
 	*/
-	TArray<UEdGraphNode*> GetAllConnectedNodes(UEdGraphNode* const Node) const;
+	UE_API TArray<UEdGraphNode*> GetAllConnectedNodes(UEdGraphNode* const Node) const;
 
-	void OnAlignTop();
-	void OnAlignMiddle();
-	void OnAlignBottom();
-	void OnAlignLeft();
-	void OnAlignCenter();
-	void OnAlignRight();
+	UE_API void OnAlignTop();
+	UE_API void OnAlignMiddle();
+	UE_API void OnAlignBottom();
+	UE_API void OnAlignLeft();
+	UE_API void OnAlignCenter();
+	UE_API void OnAlignRight();
 
-	void OnStraightenConnections();
+	UE_API void OnStraightenConnections();
 
-	void OnDistributeNodesH();
-	void OnDistributeNodesV();
+	UE_API void OnDistributeNodesH();
+	UE_API void OnDistributeNodesV();
 
-	void OnStackNodesH();
-	void OnStackNodesV();
+	UE_API void OnStackNodesH();
+	UE_API void OnStackNodesV();
 
-	void SelectAllNodes();
-	bool CanSelectAllNodes() const;
+	UE_API void SelectAllNodes();
+	UE_API bool CanSelectAllNodes() const;
 
-	virtual void DeleteSelectedNodes();
-	virtual bool CanDeleteNodes() const;
+	UE_API virtual void DeleteSelectedNodes();
+	UE_API virtual bool CanDeleteNodes() const;
 
 	/**
 	* Given a node, make connections from anything connected to it's input pin to
@@ -1096,53 +1098,53 @@ protected:
 	* 
 	* @param Node	The node to reconnect
 	*/
-	void ReconnectExecPins(class UK2Node* Node);
+	UE_API void ReconnectExecPins(class UK2Node* Node);
 
-	void DeleteSelectedDuplicatableNodes();
+	UE_API void DeleteSelectedDuplicatableNodes();
 
-	virtual void CutSelectedNodes();
-	virtual bool CanCutNodes() const;
+	UE_API virtual void CutSelectedNodes();
+	UE_API virtual bool CanCutNodes() const;
 
-	virtual void CopySelectedNodes();
-	virtual bool CanCopyNodes() const;
+	UE_API virtual void CopySelectedNodes();
+	UE_API virtual bool CanCopyNodes() const;
 
 	/** Paste on graph at specific location */
-	virtual void PasteNodesHere(class UEdGraph* DestinationGraph, const FVector2f& GraphLocation) override;
+	UE_API virtual void PasteNodesHere(class UEdGraph* DestinationGraph, const FVector2f& GraphLocation) override;
 
 	/** Paste Variable Definition or Nodes */
-	virtual void PasteGeneric();
-	virtual bool CanPasteGeneric() const;
+	UE_API virtual void PasteGeneric();
+	UE_API virtual bool CanPasteGeneric() const;
 
-	virtual void PasteNodes();
-	virtual bool CanPasteNodes() const override;
+	UE_API virtual void PasteNodes();
+	UE_API virtual bool CanPasteNodes() const override;
 
-	void DuplicateNodes();
-	bool CanDuplicateNodes() const;
+	UE_API void DuplicateNodes();
+	UE_API bool CanDuplicateNodes() const;
 
-	void OnSelectReferenceInLevel();
-	bool CanSelectReferenceInLevel() const;
+	UE_API void OnSelectReferenceInLevel();
+	UE_API bool CanSelectReferenceInLevel() const;
 
-	void OnAssignReferencedActor();
-	bool CanAssignReferencedActor() const;
+	UE_API void OnAssignReferencedActor();
+	UE_API bool CanAssignReferencedActor() const;
 
-	virtual void OnStartWatchingPin();
-	virtual bool CanStartWatchingPin() const;
+	UE_API virtual void OnStartWatchingPin();
+	UE_API virtual bool CanStartWatchingPin() const;
 
-	virtual void OnStopWatchingPin();
-	virtual bool CanStopWatchingPin() const;
+	UE_API virtual void OnStopWatchingPin();
+	UE_API virtual bool CanStopWatchingPin() const;
 
-	void ToggleSaveIntermediateBuildProducts();
-	bool GetSaveIntermediateBuildProducts() const;
+	UE_API void ToggleSaveIntermediateBuildProducts();
+	UE_API bool GetSaveIntermediateBuildProducts() const;
 
-	void OnListObjectsReferencedByClass();
-	void OnListObjectsReferencedByBlueprint();
-	void OnRepairCorruptedBlueprint();
+	UE_API void OnListObjectsReferencedByClass();
+	UE_API void OnListObjectsReferencedByBlueprint();
+	UE_API void OnRepairCorruptedBlueprint();
 
-	void OnNodeDoubleClicked(UEdGraphNode* Node);
+	UE_API void OnNodeDoubleClicked(UEdGraphNode* Node);
 
-	virtual void OnEditTabClosed(TSharedRef<SDockTab> Tab);
+	UE_API virtual void OnEditTabClosed(TSharedRef<SDockTab> Tab);
 
-	virtual bool GetBoundsForSelectedNodes(class FSlateRect& Rect, float Padding) override;
+	UE_API virtual bool GetBoundsForSelectedNodes(class FSlateRect& Rect, float Padding) override;
 
 	/**
 	 * Pulls out the pins to use as a template when collapsing a selection to a function with a custom event involved.
@@ -1153,7 +1155,7 @@ protected:
 	 * @param InResultNode				The result node in the graph
 	 * @param InCollapsableNodes		The selection of nodes being collapsed
 	 */
-	void ExtractEventTemplateForFunction(class UK2Node_CustomEvent* InCustomEvent, UEdGraphNode* InGatewayNode, class UK2Node_EditablePinBase* InEntryNode, class UK2Node_EditablePinBase* InResultNode, TSet<UEdGraphNode*>& InCollapsableNodes);
+	UE_API void ExtractEventTemplateForFunction(class UK2Node_CustomEvent* InCustomEvent, UEdGraphNode* InGatewayNode, class UK2Node_EditablePinBase* InEntryNode, class UK2Node_EditablePinBase* InResultNode, TSet<UEdGraphNode*>& InCollapsableNodes);
 
 	/**
 	 * Collapses a selection of nodes into a graph for composite, function, or macro nodes.
@@ -1165,16 +1167,16 @@ protected:
 	 * @param InDestinationGraph		The destination graph to move the selected nodes to
 	 * @param InCollapsableNodes		The selection of nodes being collapsed
 	 */
-	void CollapseNodesIntoGraph(UEdGraphNode* InGatewayNode, class UK2Node_EditablePinBase* InEntryNode, class UK2Node_EditablePinBase* InResultNode, UEdGraph* InSourceGraph, UEdGraph* InDestinationGraph, TSet<UEdGraphNode*>& InCollapsableNodes, bool bCanDiscardEmptyReturnNode, bool bCanHaveWeakObjPtrParam);
+	UE_API void CollapseNodesIntoGraph(UEdGraphNode* InGatewayNode, class UK2Node_EditablePinBase* InEntryNode, class UK2Node_EditablePinBase* InResultNode, UEdGraph* InSourceGraph, UEdGraph* InDestinationGraph, TSet<UEdGraphNode*>& InCollapsableNodes, bool bCanDiscardEmptyReturnNode, bool bCanHaveWeakObjPtrParam);
 
 	/** Called when a selection of nodes are being collapsed into a sub-graph */
-	void CollapseNodes(TSet<class UEdGraphNode*>& InCollapsableNodes);
+	UE_API void CollapseNodes(TSet<class UEdGraphNode*>& InCollapsableNodes);
 
 	/** Called when a selection of nodes are being collapsed into a function */
-	UEdGraph* CollapseSelectionToFunction(TSharedPtr<SGraphEditor> InRootGraph, TSet<class UEdGraphNode*>& InCollapsableNodes, UEdGraphNode*& OutFunctionNode);
+	UE_API UEdGraph* CollapseSelectionToFunction(TSharedPtr<SGraphEditor> InRootGraph, TSet<class UEdGraphNode*>& InCollapsableNodes, UEdGraphNode*& OutFunctionNode);
 
 	/** Called when a selection of nodes are being collapsed into a macro */
-	UEdGraph* CollapseSelectionToMacro(TSharedPtr<SGraphEditor> InRootGraph, TSet<class UEdGraphNode*>& InCollapsableNodes, UEdGraphNode*& OutMacroNode);
+	UE_API UEdGraph* CollapseSelectionToMacro(TSharedPtr<SGraphEditor> InRootGraph, TSet<class UEdGraphNode*>& InCollapsableNodes, UEdGraphNode*& OutMacroNode);
 
 	/**
 	 * Called when a selection of nodes is being collapsed into a function
@@ -1183,7 +1185,7 @@ protected:
 	 *
 	 * @return					Returns TRUE if the selection can be promoted to a function
 	 */
-	bool CanCollapseSelectionToFunction(TSet<class UEdGraphNode*>& InSelection) const;
+	UE_API bool CanCollapseSelectionToFunction(TSet<class UEdGraphNode*>& InSelection) const;
 
 	/**
 	 * Called when a selection of nodes is being collapsed into a macro
@@ -1192,11 +1194,11 @@ protected:
 	 *
 	 * @return					Returns TRUE if the selection can be promoted to a macro
 	 */
-	bool CanCollapseSelectionToMacro(TSet<class UEdGraphNode*>& InSelection) const;
+	UE_API bool CanCollapseSelectionToMacro(TSet<class UEdGraphNode*>& InSelection) const;
 
 	/**
 	 * Expands passed in node */
-	static void ExpandNode(UEdGraphNode* InNodeToExpand, UEdGraph* InSourceGraph, TSet<UEdGraphNode*>& OutExpandedNodes);
+	static UE_API void ExpandNode(UEdGraphNode* InNodeToExpand, UEdGraph* InSourceGraph, TSet<UEdGraphNode*>& OutExpandedNodes);
 
 	/**
 	* Move every node from the source graph to the destination graph. Add Each node that is moved to the OutExpandedNodes set.
@@ -1209,193 +1211,193 @@ protected:
 	* @param OutResult			Pointer to the function result node
 	* @param bIsCollapsedGraph	Whether or not the source graph is collapsed
 	**/
-	static void MoveNodesToGraph(TArray<UEdGraphNode*>& SourceNodes, UEdGraph* DestinationGraph, TSet<UEdGraphNode*>& OutExpandedNodes, UEdGraphNode** OutEntry, UEdGraphNode** OutResult, const bool bIsCollapsedGraph = false);
+	static UE_API void MoveNodesToGraph(TArray<UEdGraphNode*>& SourceNodes, UEdGraph* DestinationGraph, TSet<UEdGraphNode*>& OutExpandedNodes, UEdGraphNode** OutEntry, UEdGraphNode** OutResult, const bool bIsCollapsedGraph = false);
 
 	/** Start editing the defaults for this blueprint */
-	void OnStartEditingDefaultsClicked();
+	UE_API void OnStartEditingDefaultsClicked();
 
 	/** Creates the widgets that go into the tabs (note: does not create the tabs themselves) **/
-	virtual void CreateDefaultTabContents(const TArray<UBlueprint*>& InBlueprints);
+	UE_API virtual void CreateDefaultTabContents(const TArray<UBlueprint*>& InBlueprints);
 
 	/** Create Default Commands **/
-	virtual void CreateDefaultCommands();
+	UE_API virtual void CreateDefaultCommands();
 
 	/** Called when DeleteUnusedVariables button is clicked */
-	void DeleteUnusedVariables_OnClicked();
+	UE_API void DeleteUnusedVariables_OnClicked();
 
 	/** Called when Find In Blueprints menu is opened is clicked */
-	void FindInBlueprints_OnClicked();
+	UE_API void FindInBlueprints_OnClicked();
 
 	//~ Begin FNotifyHook Interface
-	virtual void NotifyPreChange( FProperty* PropertyAboutToChange ) override;
-	virtual void NotifyPostChange( const FPropertyChangedEvent& PropertyChangedEvent, FProperty* PropertyThatChanged) override;
+	UE_API virtual void NotifyPreChange( FProperty* PropertyAboutToChange ) override;
+	UE_API virtual void NotifyPostChange( const FPropertyChangedEvent& PropertyChangedEvent, FProperty* PropertyThatChanged) override;
 	//~ End FNotifyHook Interface
 
 	/** Callback to determine visibility of the public view checkbox in the Defaults editor */
-	bool ShouldShowPublicViewControl() const;
+	UE_API bool ShouldShowPublicViewControl() const;
 
 	/** Callback when properties have finished being handled */
-	virtual void OnFinishedChangingProperties(const FPropertyChangedEvent& PropertyChangedEvent);
+	UE_API virtual void OnFinishedChangingProperties(const FPropertyChangedEvent& PropertyChangedEvent);
 
 	/** On starting to rename node */
-	void OnRenameNode();
-	bool CanRenameNodes() const;
+	UE_API void OnRenameNode();
+	UE_API bool CanRenameNodes() const;
 
 	/** Called when a node's title is being committed for a rename so it can be verified */
-	bool OnNodeVerifyTitleCommit(const FText& NewText, UEdGraphNode* NodeBeingChanged, FText& OutErrorMessage);
+	UE_API bool OnNodeVerifyTitleCommit(const FText& NewText, UEdGraphNode* NodeBeingChanged, FText& OutErrorMessage);
 
 	/**Load macro & function blueprint libraries from asset registry*/
-	void LoadLibrariesFromAssetRegistry();
+	UE_API void LoadLibrariesFromAssetRegistry();
 
 	//~ Begin FEditorUndoClient Interface
-	virtual void	PostUndo(bool bSuccess) override;
-	virtual void	PostRedo(bool bSuccess) override;
+	UE_API virtual void	PostUndo(bool bSuccess) override;
+	UE_API virtual void	PostRedo(bool bSuccess) override;
 	// End of FEditorUndoClient
 
 	/** Setup all the events that the graph editor can handle */
-	virtual void SetupGraphEditorEvents(UEdGraph* InGraph, SGraphEditor::FGraphEditorEvents& InEvents);
+	UE_API virtual void SetupGraphEditorEvents(UEdGraph* InGraph, SGraphEditor::FGraphEditorEvents& InEvents);
 
 	/** Get the graph appearance of the currently focused graph */
-	FGraphAppearanceInfo GetCurrentGraphAppearance() const;
+	UE_API FGraphAppearanceInfo GetCurrentGraphAppearance() const;
 
 	/** Get the graph appearance of a specific graph, GetCurrentGraphAppearance() uses the currently focused graph. */
-	virtual FGraphAppearanceInfo GetGraphAppearance(class UEdGraph* InGraph) const;
+	UE_API virtual FGraphAppearanceInfo GetGraphAppearance(class UEdGraph* InGraph) const;
 
 	/** Whenever new graphs need to be created it will use this schema by default. */
-	virtual TSubclassOf<UEdGraphSchema> GetDefaultSchemaClass() const;
+	UE_API virtual TSubclassOf<UEdGraphSchema> GetDefaultSchemaClass() const;
 
 	/** Attempts to invoke the details tab if it's currently possible to. */
-	void TryInvokingDetailsTab(bool bFlash = true);
+	UE_API void TryInvokingDetailsTab(bool bFlash = true);
 
 	/** Handles activation of a graph editor "quick jump" command */
-	void OnGraphEditorQuickJump(int32 BookmarkIndex);
+	UE_API void OnGraphEditorQuickJump(int32 BookmarkIndex);
 
 	/** Binds a graph editor "quick jump" to the command at the given index */
-	void SetGraphEditorQuickJump(int32 BookmarkIndex);
+	UE_API void SetGraphEditorQuickJump(int32 BookmarkIndex);
 
 	/** Unbinds a graph editor "quick jump" from the command at the given index */
-	void ClearGraphEditorQuickJump(int32 BookmarkIndex);
+	UE_API void ClearGraphEditorQuickJump(int32 BookmarkIndex);
 
 	/** Unbinds all graph editor "quick jump" commands */
-	void ClearAllGraphEditorQuickJumps();
+	UE_API void ClearAllGraphEditorQuickJumps();
 
 	/** Create a graph title bar widget */
-	TSharedRef<SWidget> CreateGraphTitleBarWidget(TSharedRef<FTabInfo> InTabInfo, UEdGraph* InGraph);
+	UE_API TSharedRef<SWidget> CreateGraphTitleBarWidget(TSharedRef<FTabInfo> InTabInfo, UEdGraph* InGraph);
 
 private:
 
 	/** Returns true if modules can be recompiled */
-	static bool CanRecompileModules();
+	static UE_API bool CanRecompileModules();
 
 	/* User wants to edit tunnel via function editor */
-	void OnEditTunnel();
+	UE_API void OnEditTunnel();
 
 protected:
 	
 	/* Create comment node on graph */
-	virtual void OnCreateComment();
+	UE_API virtual void OnCreateComment();
 	
 	/** Create custom event node on graph */
-	virtual void OnCreateCustomEvent();
+	UE_API virtual void OnCreateCustomEvent();
 
 	// Create new graph editor widget for the supplied document container
-	virtual TSharedRef<SGraphEditor> CreateGraphEditorWidget(TSharedRef<class FTabInfo> InTabInfo, class UEdGraph* InGraph);
+	UE_API virtual TSharedRef<SGraphEditor> CreateGraphEditorWidget(TSharedRef<class FTabInfo> InTabInfo, class UEdGraph* InGraph);
 
 private:
 
 	/** Helper to move focused graph when clicking on graph breadcrumb */
-	void OnChangeBreadCrumbGraph( class UEdGraph* InGraph);
+	UE_API void OnChangeBreadCrumbGraph( class UEdGraph* InGraph);
 
 	/** Function to check whether the give graph is a subgraph */
-	static bool IsASubGraph( const class UEdGraph* GraphPtr );
+	static UE_API bool IsASubGraph( const class UEdGraph* GraphPtr );
 
 	/** Creates the Subobject Editor tree component view and the Subobject Viewport. */
-	void CreateSubobjectEditors();
+	UE_API void CreateSubobjectEditors();
 
 	/** Callback when a token is clicked on in the compiler results log */
-	void OnLogTokenClicked(const TSharedRef<class IMessageToken>& Token);
+	UE_API void OnLogTokenClicked(const TSharedRef<class IMessageToken>& Token);
 
 	/** Helper function to navigate the current tab */
-	void NavigateTab(FDocumentTracker::EOpenDocumentCause InCause);
+	UE_API void NavigateTab(FDocumentTracker::EOpenDocumentCause InCause);
 
 	/** Find all references of the selected node. */
-	void OnFindReferences(bool bSearchAllBlueprints, const EGetFindReferenceSearchStringFlags Flags);
+	UE_API void OnFindReferences(bool bSearchAllBlueprints, const EGetFindReferenceSearchStringFlags Flags);
 
 	/** Checks if we can currently find all references of the node selection. */
-	bool CanFindReferences();
+	UE_API bool CanFindReferences();
 
 	/** Called when the user generates a warning tooltip because a connection was invalid */
-	void OnDisallowedPinConnection(const class UEdGraphPin* PinA, const class UEdGraphPin* PinB);
+	UE_API void OnDisallowedPinConnection(const class UEdGraphPin* PinA, const class UEdGraphPin* PinB);
 
 	/**
 	 * Checks to see if the focused Graph-Editor is focused on a animation graph or not.
 	 * 
 	 * @return True if GraphEdPtr's current graph is an animation graph, false if not.
 	 */
-	bool IsEditingAnimGraph() const;
+	UE_API bool IsEditingAnimGraph() const;
 
 	/** Returns whether the currently focused graph is editable or not */
-	bool IsFocusedGraphEditable() const;
+	UE_API bool IsFocusedGraphEditable() const;
 
 	/** Called when the user wants to jump to a node's definition */
-	void OnGoToDefinition();
+	UE_API void OnGoToDefinition();
 
 	/** Checks to see if it is possible to jump to the selected node's definition. */
-	bool CanGoToDefinition() const;
+	UE_API bool CanGoToDefinition() const;
 
 	/** Open documentation for the selected node class */
-	void OnGoToDocumentation();
+	UE_API void OnGoToDocumentation();
 
 	/** Can we open documentation for the selected node */
-	bool CanGoToDocumentation();
+	UE_API bool CanGoToDocumentation();
 
 	/** Util to try and get doc link for the currently selected node */
-	FString GetDocLinkForSelectedNode();
+	UE_API FString GetDocLinkForSelectedNode();
 
 	/** Util to try and get the base URL for the doc link for the currently selected node */
-	FString GetDocLinkBaseUrlForSelectedNode();
+	UE_API FString GetDocLinkBaseUrlForSelectedNode();
 
 	/** Set the enabled state for currently-selected nodes */
-	void OnSetEnabledStateForSelectedNodes(ENodeEnabledState NewState);
+	UE_API void OnSetEnabledStateForSelectedNodes(ENodeEnabledState NewState);
 
 	/** Returns the appropriate check box state representing whether or not the selected nodes are enabled */
-	ECheckBoxState GetEnabledCheckBoxStateForSelectedNodes();
+	UE_API ECheckBoxState GetEnabledCheckBoxStateForSelectedNodes();
 
 	/**
 	 * Load editor settings from disk (docking state, window pos/size, option state, etc).
 	 */
-	virtual void LoadEditorSettings();
+	UE_API virtual void LoadEditorSettings();
 
 	/**
 	 * Saves editor settings to disk (docking state, window pos/size, option state, etc).
 	 */
-	virtual void SaveEditorSettings();
+	UE_API virtual void SaveEditorSettings();
 
 	/** Attempt to match the given enabled state for currently-selected nodes */
-	ECheckBoxState CheckEnabledStateForSelectedNodes(ENodeEnabledState CheckState);
+	UE_API ECheckBoxState CheckEnabledStateForSelectedNodes(ENodeEnabledState CheckState);
 
 	/** Handle undo/redo */
-	void HandleUndoTransaction(const class FTransaction* Transaction);
+	UE_API void HandleUndoTransaction(const class FTransaction* Transaction);
 
 public://@TODO
 
-	virtual bool TransactionObjectAffectsBlueprint(UObject* InTransactedObject);
+	UE_API virtual bool TransactionObjectAffectsBlueprint(UObject* InTransactedObject);
 
 	TSharedPtr<FDocumentTracker> DocumentManager;
 	
 	/** Update all nodes' unrelated states when the graph has changed */
-	void UpdateNodesUnrelatedStatesAfterGraphChange();
+	UE_API void UpdateNodesUnrelatedStatesAfterGraphChange();
 
 	/** Whether event graphs are allowed to be displayed/created in a blueprint */
-	virtual bool AreEventGraphsAllowed() const;
+	UE_API virtual bool AreEventGraphsAllowed() const;
 
 	/** Whether macros are allowed  to be displayed/created in a blueprint */
-	virtual bool AreMacrosAllowed() const;
+	UE_API virtual bool AreMacrosAllowed() const;
 
 	/** Whether delegates are allowed  to be displayed/created in a blueprint */
-	virtual bool AreDelegatesAllowed() const;
+	UE_API virtual bool AreDelegatesAllowed() const;
 
-	FEditorViewportSelectabilityBridge* GetViewportSelectabilityBridge();
+	UE_API FEditorViewportSelectabilityBridge* GetViewportSelectabilityBridge();
 
 protected:
 
@@ -1532,40 +1534,40 @@ protected:
 	bool bWasOpenedInDefaultsMode;
 
 	/** Focus nodes which are related to the selected nodes */
-	void ResetAllNodesUnrelatedStates();
-	void CollectExecUpstreamNodes(UEdGraphNode* CurrentNode, TArray<UEdGraphNode*>& CollectedNodes);
-	void CollectExecDownstreamNodes(UEdGraphNode* CurrentNode, TArray<UEdGraphNode*>& CollectedNodes);
-	void CollectPureDownstreamNodes(UEdGraphNode* CurrentNode, TArray<UEdGraphNode*>& CollectedNodes);
-	void CollectPureUpstreamNodes(UEdGraphNode* CurrentNode, TArray<UEdGraphNode*>& CollectedNodes);
-	void HideUnrelatedNodes();
+	UE_API void ResetAllNodesUnrelatedStates();
+	UE_API void CollectExecUpstreamNodes(UEdGraphNode* CurrentNode, TArray<UEdGraphNode*>& CollectedNodes);
+	UE_API void CollectExecDownstreamNodes(UEdGraphNode* CurrentNode, TArray<UEdGraphNode*>& CollectedNodes);
+	UE_API void CollectPureDownstreamNodes(UEdGraphNode* CurrentNode, TArray<UEdGraphNode*>& CollectedNodes);
+	UE_API void CollectPureUpstreamNodes(UEdGraphNode* CurrentNode, TArray<UEdGraphNode*>& CollectedNodes);
+	UE_API void HideUnrelatedNodes();
 
 	/** Register Menus */
-	void RegisterMenus();
+	UE_API void RegisterMenus();
 
 public:
 	/** Make nodes which are unrelated to the selected nodes fade out */
-	void ToggleHideUnrelatedNodes();
-	bool IsToggleHideUnrelatedNodesChecked() const;
+	UE_API void ToggleHideUnrelatedNodes();
+	UE_API bool IsToggleHideUnrelatedNodesChecked() const;
 	
 	UE_DEPRECATED(5.0, "The Toggle Hide Unrelated Nodes button is always shown now")
-	bool ShouldShowToggleHideUnrelatedNodes(bool bIsToolbar) const;
+	UE_API bool ShouldShowToggleHideUnrelatedNodes(bool bIsToolbar) const;
 
 	/** Make a drop down menu to control the opacity of unrelated nodes */
-	TSharedRef<SWidget> MakeHideUnrelatedNodesOptionsMenu();
-	TOptional<float> HandleUnrelatedNodesOpacityBoxValue() const;
-	void HandleUnrelatedNodesOpacityBoxChanged(float NewOpacity);
-	void OnLockNodeStateCheckStateChanged(ECheckBoxState NewCheckedState);
+	UE_API TSharedRef<SWidget> MakeHideUnrelatedNodesOptionsMenu();
+	UE_API TOptional<float> HandleUnrelatedNodesOpacityBoxValue() const;
+	UE_API void HandleUnrelatedNodesOpacityBoxChanged(float NewOpacity);
+	UE_API void OnLockNodeStateCheckStateChanged(ECheckBoxState NewCheckedState);
 
 
 public:
 	//@TODO: To be moved/merged
-	TSharedPtr<SDockTab> OpenDocument(const UObject* DocumentID, FDocumentTracker::EOpenDocumentCause Cause);
+	UE_API TSharedPtr<SDockTab> OpenDocument(const UObject* DocumentID, FDocumentTracker::EOpenDocumentCause Cause);
 
 	/** Finds the tab associated with the specified asset, and closes if it is open */
-	void CloseDocumentTab(const UObject* DocumentID);
+	UE_API void CloseDocumentTab(const UObject* DocumentID);
 
 	// Finds any open tabs containing the specified document and adds them to the specified array; returns true if at least one is found
-	bool FindOpenTabsContainingDocument(const UObject* DocumentID, /*inout*/ TArray< TSharedPtr<SDockTab> >& Results);
+	UE_API bool FindOpenTabsContainingDocument(const UObject* DocumentID, /*inout*/ TArray< TSharedPtr<SDockTab> >& Results);
 
 public:
 	/** Broadcasts a notification whenever the editor needs associated controls to refresh */
@@ -1650,10 +1652,12 @@ private:
 	// Allow derived editors to add command mappings 
 	virtual void OnCreateGraphEditorCommands(TSharedPtr<FUICommandList> GraphEditorCommandsList) {}
 
-	virtual void OnBlueprintProjectSettingsChanged(UObject*, struct FPropertyChangedEvent&);
-	virtual void OnBlueprintEditorPreferencesChanged(UObject*, struct FPropertyChangedEvent&);
+	UE_API virtual void OnBlueprintProjectSettingsChanged(UObject*, struct FPropertyChangedEvent&);
+	UE_API virtual void OnBlueprintEditorPreferencesChanged(UObject*, struct FPropertyChangedEvent&);
 
 	TUniquePtr<FEditorViewportSelectabilityBridge> ViewportSelectabilityBridge;
 };
 
 #undef LOCTEXT_NAMESPACE
+
+#undef UE_API
diff --git a/Engine/Source/Editor/Kismet/Public/BlueprintEditorContext.h b/Engine/Source/Editor/Kismet/Public/BlueprintEditorContext.h
--- a/Engine/Source/Editor/Kismet/Public/BlueprintEditorContext.h
+++ b/Engine/Source/Editor/Kismet/Public/BlueprintEditorContext.h
@@ -6,17 +6,21 @@
 
 #include "BlueprintEditorContext.generated.h"
 
+#define UE_API KISMET_API
+
 class UBlueprint;
 class FBlueprintEditor;
 
-UCLASS()
-class KISMET_API UBlueprintEditorToolMenuContext : public UObject
+UCLASS(MinimalAPI)
+class UBlueprintEditorToolMenuContext : public UObject
 {
 	GENERATED_BODY()
 public:
 
 	UFUNCTION(BlueprintCallable, Category="Tool Menus")
-	UBlueprint* GetBlueprintObj() const;
+	UE_API UBlueprint* GetBlueprintObj() const;
 
 	TWeakPtr<FBlueprintEditor> BlueprintEditor;
 };
+
+#undef UE_API
diff --git a/Engine/Source/Editor/Kismet/Public/BlueprintEditorModes.h b/Engine/Source/Editor/Kismet/Public/BlueprintEditorModes.h
--- a/Engine/Source/Editor/Kismet/Public/BlueprintEditorModes.h
+++ b/Engine/Source/Editor/Kismet/Public/BlueprintEditorModes.h
@@ -16,19 +16,21 @@
 #include "WorkflowOrientedApp/ApplicationMode.h"
 #include "WorkflowOrientedApp/WorkflowTabManager.h"
 
+#define UE_API KISMET_API
+
 class AActor;
 class FBlueprintEditor;
 class FTabManager;
 class UActorComponent;
 
-struct KISMET_API FBlueprintEditorApplicationModes
+struct FBlueprintEditorApplicationModes
 {
 	// Mode identifiers
-	static const FName StandardBlueprintEditorMode;
-	static const FName BlueprintDefaultsMode;
-	static const FName BlueprintComponentsMode;
-	static const FName BlueprintInterfaceMode;
-	static const FName BlueprintMacroMode;
+	static UE_API const FName StandardBlueprintEditorMode;
+	static UE_API const FName BlueprintDefaultsMode;
+	static UE_API const FName BlueprintComponentsMode;
+	static UE_API const FName BlueprintInterfaceMode;
+	static UE_API const FName BlueprintMacroMode;
 	static FText GetLocalizedMode( const FName InMode )
 	{
 		static TMap< FName, FText > LocModes;
@@ -53,14 +55,14 @@ private:
 
 
 
-class KISMET_API FBlueprintEditorApplicationMode : public FApplicationMode
+class FBlueprintEditorApplicationMode : public FApplicationMode
 {
 public:
-	FBlueprintEditorApplicationMode(TSharedPtr<class FBlueprintEditor> InBlueprintEditor, FName InModeName, FText(*GetLocalizedMode)(const FName), const bool bRegisterViewport = true, const bool bRegisterDefaultsTab = true);
+	UE_API FBlueprintEditorApplicationMode(TSharedPtr<class FBlueprintEditor> InBlueprintEditor, FName InModeName, FText(*GetLocalizedMode)(const FName), const bool bRegisterViewport = true, const bool bRegisterDefaultsTab = true);
 
-	virtual void RegisterTabFactories(TSharedPtr<FTabManager> InTabManager) override;
-	virtual void PreDeactivateMode() override;
-	virtual void PostActivateMode() override;
+	UE_API virtual void RegisterTabFactories(TSharedPtr<FTabManager> InTabManager) override;
+	UE_API virtual void PreDeactivateMode() override;
+	UE_API virtual void PostActivateMode() override;
 public:
 
 protected:
@@ -77,13 +79,13 @@ protected:
 };
 
 
-class KISMET_API FBlueprintDefaultsApplicationMode : public FApplicationMode
+class FBlueprintDefaultsApplicationMode : public FApplicationMode
 {
 public:
-	FBlueprintDefaultsApplicationMode(TSharedPtr<class FBlueprintEditor> InBlueprintEditor);
+	UE_API FBlueprintDefaultsApplicationMode(TSharedPtr<class FBlueprintEditor> InBlueprintEditor);
 
-	virtual void RegisterTabFactories(TSharedPtr<FTabManager> InTabManager) override;
-	virtual void PostActivateMode() override;
+	UE_API virtual void RegisterTabFactories(TSharedPtr<FTabManager> InTabManager) override;
+	UE_API virtual void PostActivateMode() override;
 public:
 
 protected:
@@ -94,17 +96,17 @@ protected:
 };
 
 
-class KISMET_API FBlueprintComponentsApplicationMode : public FApplicationMode
+class FBlueprintComponentsApplicationMode : public FApplicationMode
 {
 public:
-	FBlueprintComponentsApplicationMode(TSharedPtr<class FBlueprintEditor> InBlueprintEditor);
+	UE_API FBlueprintComponentsApplicationMode(TSharedPtr<class FBlueprintEditor> InBlueprintEditor);
 
-	virtual void RegisterTabFactories(TSharedPtr<FTabManager> InTabManager) override;
-	virtual void PreDeactivateMode() override;
-	virtual void PostActivateMode() override;
+	UE_API virtual void RegisterTabFactories(TSharedPtr<FTabManager> InTabManager) override;
+	UE_API virtual void PreDeactivateMode() override;
+	UE_API virtual void PostActivateMode() override;
 
 	TSharedPtr<FBlueprintEditor> GetBlueprintEditor() const { return MyBlueprintEditor.Pin(); }
-	AActor* GetPreviewActor() const;
+	UE_API AActor* GetPreviewActor() const;
 protected:
 	TWeakPtr<FBlueprintEditor> MyBlueprintEditor;
 
@@ -114,14 +116,14 @@ protected:
 	TArray<TWeakObjectPtr<const UActorComponent>> CachedComponentSelection;
 };
 
-class KISMET_API FBlueprintInterfaceApplicationMode : public FApplicationMode
+class FBlueprintInterfaceApplicationMode : public FApplicationMode
 {
 public:
-	FBlueprintInterfaceApplicationMode(TSharedPtr<class FBlueprintEditor> InBlueprintEditor, FName InModeName = FBlueprintEditorApplicationModes::BlueprintInterfaceMode, FText(*GetLocalizedMode)(const FName) = FBlueprintEditorApplicationModes::GetLocalizedMode);
+	UE_API FBlueprintInterfaceApplicationMode(TSharedPtr<class FBlueprintEditor> InBlueprintEditor, FName InModeName = FBlueprintEditorApplicationModes::BlueprintInterfaceMode, FText(*GetLocalizedMode)(const FName) = FBlueprintEditorApplicationModes::GetLocalizedMode);
 
-	virtual void RegisterTabFactories(TSharedPtr<FTabManager> InTabManager) override;
-	virtual void PreDeactivateMode() override;
-	virtual void PostActivateMode() override;
+	UE_API virtual void RegisterTabFactories(TSharedPtr<FTabManager> InTabManager) override;
+	UE_API virtual void PreDeactivateMode() override;
+	UE_API virtual void PostActivateMode() override;
 public:
 
 protected:
@@ -131,14 +133,14 @@ protected:
 	FWorkflowAllowedTabSet BlueprintInterfaceTabFactories;
 };
 
-class KISMET_API FBlueprintMacroApplicationMode : public FApplicationMode
+class FBlueprintMacroApplicationMode : public FApplicationMode
 {
 public:
-	FBlueprintMacroApplicationMode(TSharedPtr<class FBlueprintEditor> InBlueprintEditor);
+	UE_API FBlueprintMacroApplicationMode(TSharedPtr<class FBlueprintEditor> InBlueprintEditor);
 
-	virtual void RegisterTabFactories(TSharedPtr<FTabManager> InTabManager) override;
-	virtual void PreDeactivateMode() override;
-	virtual void PostActivateMode() override;
+	UE_API virtual void RegisterTabFactories(TSharedPtr<FTabManager> InTabManager) override;
+	UE_API virtual void PreDeactivateMode() override;
+	UE_API virtual void PostActivateMode() override;
 public:
 
 protected:
@@ -148,14 +150,14 @@ protected:
 	FWorkflowAllowedTabSet BlueprintMacroTabFactories;
 };
 
-class KISMET_API FBlueprintEditorUnifiedMode : public FApplicationMode
+class FBlueprintEditorUnifiedMode : public FApplicationMode
 {
 public:
-	FBlueprintEditorUnifiedMode(TSharedPtr<class FBlueprintEditor> InBlueprintEditor, FName InModeName, FText(*GetLocalizedMode)( const FName ), const bool bRegisterViewport = true);
+	UE_API FBlueprintEditorUnifiedMode(TSharedPtr<class FBlueprintEditor> InBlueprintEditor, FName InModeName, FText(*GetLocalizedMode)( const FName ), const bool bRegisterViewport = true);
 
-	virtual void RegisterTabFactories(TSharedPtr<FTabManager> InTabManager) override;
-	virtual void PreDeactivateMode() override;
-	virtual void PostActivateMode() override;
+	UE_API virtual void RegisterTabFactories(TSharedPtr<FTabManager> InTabManager) override;
+	UE_API virtual void PreDeactivateMode() override;
+	UE_API virtual void PostActivateMode() override;
 public:
 
 protected:
@@ -170,3 +172,5 @@ protected:
 	// Set of spawnable tabs only usable in blueprint editing mode (not useful in Persona, etc...)
 	FWorkflowAllowedTabSet BlueprintEditorOnlyTabFactories;
 };
+
+#undef UE_API
diff --git a/Engine/Source/Editor/Kismet/Public/BlueprintEditorModule.h b/Engine/Source/Editor/Kismet/Public/BlueprintEditorModule.h
--- a/Engine/Source/Editor/Kismet/Public/BlueprintEditorModule.h
+++ b/Engine/Source/Editor/Kismet/Public/BlueprintEditorModule.h
@@ -22,6 +22,8 @@
 #include "UObject/SoftObjectPath.h"
 #include "WorkflowOrientedApp/WorkflowCentricApplication.h"
 
+#define UE_API KISMET_API
+
 class FBlueprintEditor;
 class FExtender;
 class FFieldClass;
@@ -73,21 +75,21 @@ namespace ERefreshBlueprintEditorReason
 /**
  * Enum editor public interface
  */
-class KISMET_API IUserDefinedEnumEditor : public FAssetEditorToolkit
+class IUserDefinedEnumEditor : public FAssetEditorToolkit
 {
 };
 
 /**
  * Enum editor public interface
  */
-class KISMET_API IUserDefinedStructureEditor : public FAssetEditorToolkit
+class IUserDefinedStructureEditor : public FAssetEditorToolkit
 {
 };
 
 /**
  * Blueprint editor public interface
  */
-class KISMET_API IBlueprintEditor : public FWorkflowCentricApplication
+class IBlueprintEditor : public FWorkflowCentricApplication
 {
 public:
 	virtual void JumpToHyperlink(const UObject* ObjectReference, bool bRequestRename) = 0;
@@ -116,7 +118,7 @@ public:
 	virtual bool CanPasteNodes() const= 0;
 
 	UE_DEPRECATED(5.6, "Slate positions are represented in FVector2f. Please update to use the function acceping FVector2f.")
-	virtual void PasteNodesHere(class UEdGraph* Graph, const FVector2D& Location) UE_SLATE_DEPRECATED_VECTOR_VIRTUAL_FUNCTION { PasteNodesHere(Graph, UE::Slate::CastToVector2f(Location)); };
+	UE_API virtual void PasteNodesHere(class UEdGraph* Graph, const FVector2D& Location) UE_SLATE_DEPRECATED_VECTOR_VIRTUAL_FUNCTION { PasteNodesHere(Graph, UE::Slate::CastToVector2f(Location)); };
 	virtual void PasteNodesHere(class UEdGraph* Graph, const FVector2f& Location) = 0;
 
 	virtual bool GetBoundsForSelectedNodes(class FSlateRect& Rect, float Padding ) = 0;
@@ -407,3 +409,5 @@ private:
 	/** Blueprint debugger state - refactor into SBlueprintDebugger if needed */
 	TUniquePtr<FBlueprintDebugger> BlueprintDebugger;
 };
+
+#undef UE_API
diff --git a/Engine/Source/Editor/Kismet/Public/BlueprintEditorSharedTabFactories.h b/Engine/Source/Editor/Kismet/Public/BlueprintEditorSharedTabFactories.h
--- a/Engine/Source/Editor/Kismet/Public/BlueprintEditorSharedTabFactories.h
+++ b/Engine/Source/Editor/Kismet/Public/BlueprintEditorSharedTabFactories.h
@@ -7,17 +7,19 @@
 #include "Toolkits/AssetEditorToolkit.h"
 #include "WorkflowOrientedApp/WorkflowTabFactory.h"
 
+#define UE_API KISMET_API
+
 #define LOCTEXT_NAMESPACE "BlueprintEditor"
 
 //////////////////////////////////////////////////////////////////////////
 // FCompilerResultsSummoner
 
-struct KISMET_API FCompilerResultsSummoner : public FWorkflowTabFactory
+struct FCompilerResultsSummoner : public FWorkflowTabFactory
 {
 public:
-	FCompilerResultsSummoner(TSharedPtr<class FAssetEditorToolkit> InHostingApp);
+	UE_API FCompilerResultsSummoner(TSharedPtr<class FAssetEditorToolkit> InHostingApp);
 
-	virtual TSharedRef<SWidget> CreateTabBody(const FWorkflowTabSpawnInfo& Info) const override;
+	UE_API virtual TSharedRef<SWidget> CreateTabBody(const FWorkflowTabSpawnInfo& Info) const override;
 
 	virtual FText GetTabToolTipText(const FWorkflowTabSpawnInfo& Info) const override
 	{
@@ -26,3 +28,5 @@ public:
 };
 
 #undef LOCTEXT_NAMESPACE
+
+#undef UE_API
diff --git a/Engine/Source/Editor/Kismet/Public/BlueprintEditorTabs.h b/Engine/Source/Editor/Kismet/Public/BlueprintEditorTabs.h
--- a/Engine/Source/Editor/Kismet/Public/BlueprintEditorTabs.h
+++ b/Engine/Source/Editor/Kismet/Public/BlueprintEditorTabs.h
@@ -4,30 +4,33 @@
 
 #include "CoreMinimal.h"
 
+#define UE_API KISMET_API
+
 class FName;
 
-struct KISMET_API FBlueprintEditorTabs
+struct FBlueprintEditorTabs
 {
 	// Tab identifiers
-	static const FName DetailsID;
-	static const FName DefaultEditorID;
-	static const FName DebugID;
-	static const FName PaletteID;
-	static const FName BookmarksID;
-	static const FName CompilerResultsID;
-	static const FName FindResultsID;
-	static const FName ConstructionScriptEditorID;
-	static const FName SCSViewportID;
-	static const FName MyBlueprintID;
-	static const FName ReplaceNodeReferencesID;
-	static const FName UserDefinedStructureID;
-	static const FName BlueprintDebuggerID;
+	static UE_API const FName DetailsID;
+	static UE_API const FName DefaultEditorID;
+	static UE_API const FName DebugID;
+	static UE_API const FName PaletteID;
+	static UE_API const FName BookmarksID;
+	static UE_API const FName CompilerResultsID;
+	static UE_API const FName FindResultsID;
+	static UE_API const FName ConstructionScriptEditorID;
+	static UE_API const FName SCSViewportID;
+	static UE_API const FName MyBlueprintID;
+	static UE_API const FName ReplaceNodeReferencesID;
+	static UE_API const FName UserDefinedStructureID;
+	static UE_API const FName BlueprintDebuggerID;
 
 	// Document tab identifiers
-	static const FName GraphEditorID;
-	static const FName TimelineEditorID;
+	static UE_API const FName GraphEditorID;
+	static UE_API const FName TimelineEditorID;
 
 private:
 	FBlueprintEditorTabs() {}
 };
 
+#undef UE_API
diff --git a/Engine/Source/Editor/Kismet/Public/BlueprintManagedListDetails.h b/Engine/Source/Editor/Kismet/Public/BlueprintManagedListDetails.h
--- a/Engine/Source/Editor/Kismet/Public/BlueprintManagedListDetails.h
+++ b/Engine/Source/Editor/Kismet/Public/BlueprintManagedListDetails.h
@@ -16,12 +16,14 @@
 #include "UObject/WeakObjectPtr.h"
 #include "UObject/WeakObjectPtrTemplates.h"
 
+#define UE_API KISMET_API
+
 class FDetailWidgetRow;
 class IDetailChildrenBuilder;
 class SWidget;
 
 /** Blueprint managed list details */
-class KISMET_API FBlueprintManagedListDetails : public IDetailCustomNodeBuilder
+class FBlueprintManagedListDetails : public IDetailCustomNodeBuilder
 {
 protected:
 	/** List item node type */
@@ -54,8 +56,8 @@ protected:
 	FDisplayOptions DisplayOptions;
 
 	/** IDetailCustomNodeBuilder interface */
-	virtual void GenerateHeaderRowContent(FDetailWidgetRow& HeaderRow) override;
-	virtual void GenerateChildContent(IDetailChildrenBuilder& ChildrenBuilder) override;
+	UE_API virtual void GenerateHeaderRowContent(FDetailWidgetRow& HeaderRow) override;
+	UE_API virtual void GenerateChildContent(IDetailChildrenBuilder& ChildrenBuilder) override;
 	virtual FName GetName() const override { return *DisplayOptions.TitleText.ToString(); }
 	virtual bool InitiallyCollapsed() const override { return false; }
 	virtual void Tick(float DeltaTime) override {}
@@ -69,9 +71,11 @@ protected:
 	virtual void OnRemoveItem(const FManagedListItem& Item) {}
 
 	/** Helper function to regenerate the details customization */
-	void RegenerateChildContent();
+	UE_API void RegenerateChildContent();
 
 private:
 	/** A delegate to regenerate this list of children */
 	FSimpleDelegate RegenerateChildrenDelegate;
-};
\ No newline at end of file
+};
+
+#undef UE_API
diff --git a/Engine/Source/Editor/Kismet/Public/BlueprintMemberReferenceCustomization.h b/Engine/Source/Editor/Kismet/Public/BlueprintMemberReferenceCustomization.h
--- a/Engine/Source/Editor/Kismet/Public/BlueprintMemberReferenceCustomization.h
+++ b/Engine/Source/Editor/Kismet/Public/BlueprintMemberReferenceCustomization.h
@@ -6,13 +6,15 @@
 #include "IPropertyTypeCustomization.h"
 #include "Templates/SharedPointer.h"
 
+#define UE_API KISMET_API
+
 class FDetailWidgetRow;
 class IDetailChildrenBuilder;
 class IPropertyHandle;
 class SMyBlueprint;
 
 // Property type customization for FMemberReference
-class KISMET_API FBlueprintMemberReferenceDetails : public IPropertyTypeCustomization
+class FBlueprintMemberReferenceDetails : public IPropertyTypeCustomization
 {
 public:
 	static TSharedRef<IPropertyTypeCustomization> MakeInstance()
@@ -32,9 +34,11 @@ private:
 	}
 	
 	// IPropertyTypeCustomization interface
-	virtual void CustomizeHeader(TSharedRef<IPropertyHandle> InStructPropertyHandle, FDetailWidgetRow& InHeaderRow, IPropertyTypeCustomizationUtils& InStructCustomizationUtils) override;
+	UE_API virtual void CustomizeHeader(TSharedRef<IPropertyHandle> InStructPropertyHandle, FDetailWidgetRow& InHeaderRow, IPropertyTypeCustomizationUtils& InStructCustomizationUtils) override;
 	virtual void CustomizeChildren(TSharedRef<IPropertyHandle> InStructPropertyHandle, IDetailChildrenBuilder& InChildBuilder, IPropertyTypeCustomizationUtils& InStructCustomizationUtils) override {};
 
 private:
 	TWeakPtr<SMyBlueprint> MyBlueprint;
 };
+
+#undef UE_API
diff --git a/Engine/Source/Editor/Kismet/Public/BlueprintNamespaceHelper.h b/Engine/Source/Editor/Kismet/Public/BlueprintNamespaceHelper.h
--- a/Engine/Source/Editor/Kismet/Public/BlueprintNamespaceHelper.h
+++ b/Engine/Source/Editor/Kismet/Public/BlueprintNamespaceHelper.h
@@ -8,6 +8,8 @@
 #include "Templates/UniquePtr.h"
 #include "UObject/SoftObjectPath.h"
 
+#define UE_API KISMET_API
+
 class IClassViewerFilter;
 class IPinTypeSelectorFilter;
 class UBlueprint;
@@ -18,28 +20,28 @@ struct FBlueprintNamespacePathTree;
 /**
  * A helper class for managing namespace info for a single Blueprint context.
  */
-class KISMET_API FBlueprintNamespaceHelper
+class FBlueprintNamespaceHelper
 {
 public:
 	/** Default constructor. */
-	FBlueprintNamespaceHelper();
+	UE_API FBlueprintNamespaceHelper();
 
 	/** Note: We explicitly declare/implement the dtor so that forward-declared types can be destroyed. */
-	virtual ~FBlueprintNamespaceHelper();
+	UE_API virtual ~FBlueprintNamespaceHelper();
 
 	/**
 	 * Add a new Blueprint into the current helper context.
 	 * 
 	 * @param InBlueprint	A Blueprint object being edited.
 	 */
-	void AddBlueprint(const UBlueprint* InBlueprint);
+	UE_API void AddBlueprint(const UBlueprint* InBlueprint);
 
 	/**
 	 * Add a namespace identifier into the Blueprint editor context that's managed by this instance.
 	 *
 	 * @param Namespace		The namespace identifier string to add. Should resemble "X.Y.Z" as the format.
 	 */
-	void AddNamespace(const FString& Namespace);
+	UE_API void AddNamespace(const FString& Namespace);
 
 	/**
 	 * Add multiple namespace identifiers into the Blueprint editor context that's managed by this instance.
@@ -60,7 +62,7 @@ public:
 	 * 
 	 * @param Namespace		The namespace identifier string to remove. Should resemble "X.Y.Z" as the format.
 	 */
-	void RemoveNamespace(const FString& Namespace);
+	UE_API void RemoveNamespace(const FString& Namespace);
 
 	/**
 	 * Tests a namespace identifier to see if it's been imported. 
@@ -68,7 +70,7 @@ public:
 	 * @param TestNamespace	A namespace identifier string. Should resemble "X.Y.Z" as the format.
 	 * @return TRUE if the namespace identifier exists within the Blueprint editor context that's managed by this instance.
 	 */
-	bool IsIncludedInNamespaceList(const FString& TestNamespace) const;
+	UE_API bool IsIncludedInNamespaceList(const FString& TestNamespace) const;
 
 	/**
 	 * Tests an asset's namespace to see if it's been imported.
@@ -76,7 +78,7 @@ public:
 	 * @param InAssetData	Asset data context.
 	 * @return TRUE if the asset's underlying namespace was imported into the Blueprint editor context that's managed by this instance.
 	 */
-	bool IsImportedAsset(const FAssetData& InAssetData) const;
+	UE_API bool IsImportedAsset(const FAssetData& InAssetData) const;
 
 	/**
 	 * Tests an object's namespace to see if it's been imported.
@@ -84,7 +86,7 @@ public:
 	 * @param InObject	Hard reference to a (loaded) object.
 	 * @return TRUE if the object's underlying namespace was imported into the Blueprint editor context that's managed by this instance.
 	 */
-	bool IsImportedObject(const UObject* InObject) const;
+	UE_API bool IsImportedObject(const UObject* InObject) const;
 
 	/**
 	 * Tests an object's namespace to see if it's been imported.
@@ -92,7 +94,7 @@ public:
 	 * @param InObjectPath	Soft reference to an object. May be loaded or unloaded.
 	 * @return TRUE if the object's underlying namespace was imported into the Blueprint editor context that's managed by this instance.
 	 */
-	bool IsImportedObject(const FSoftObjectPath& InObjectPath) const;
+	UE_API bool IsImportedObject(const FSoftObjectPath& InObjectPath) const;
 
 	/**
 	 * @return A class viewer filter that can be used with class picker widgets to filter out class types whose namespaces were not imported into the Blueprint editor context that's managed by this instance.
@@ -113,7 +115,7 @@ public:
 	/**
 	 * Utility method used to keep console flags in sync with the current Blueprint editor settings environment.
 	 */
-	static void RefreshEditorFeatureConsoleFlags();
+	static UE_API void RefreshEditorFeatureConsoleFlags();
 
 private:
 	// Path tree that stores imported namespace identifiers for the associated Blueprint.
@@ -124,4 +126,6 @@ private:
 
 	// For use with the pin type selector widget in order to filter pin type items by namespace.
 	TSharedPtr<IPinTypeSelectorFilter> PinTypeSelectorFilter;
-};
\ No newline at end of file
+};
+
+#undef UE_API
