Updated headers using UnrealCodeFixup to make sure dllstorage is on methods/staticvars instead of types. Part 5/n

diff --git a/Engine/Source/Editor/Kismet/Public/BlueprintNamespaceRegistry.h b/Engine/Source/Editor/Kismet/Public/BlueprintNamespaceRegistry.h
--- a/Engine/Source/Editor/Kismet/Public/BlueprintNamespaceRegistry.h
+++ b/Engine/Source/Editor/Kismet/Public/BlueprintNamespaceRegistry.h
@@ -11,6 +11,8 @@
 #include "UObject/SoftObjectPath.h"
 #include "UObject/UObjectGlobals.h"
 
+#define UE_API KISMET_API
+
 class FName;
 class UObject;
 struct FAssetData;
@@ -19,30 +21,30 @@ struct FBlueprintNamespacePathTree;
 /**
  * A shared utility class that keeps track of registered Blueprint namespace identifiers sourced from objects and assets in the editor.
  */
-class KISMET_API FBlueprintNamespaceRegistry
+class FBlueprintNamespaceRegistry
 {
 public:
-	~FBlueprintNamespaceRegistry();
+	UE_API ~FBlueprintNamespaceRegistry();
 
 	/**
 	 * Provides public singleton access.
 	 */
-	static FBlueprintNamespaceRegistry& Get();
+	static UE_API FBlueprintNamespaceRegistry& Get();
 
 	/**
 	 * One-time initialization method; separated from the ctor so it can be called explicitly.
 	 */
-	void Initialize();
+	UE_API void Initialize();
 
 	/**
 	 * One-time shutdown method; separated from the dtor so it can be called explicitly.
 	 */
-	void Shutdown();
+	UE_API void Shutdown();
 
 	/**
 	 * @return TRUE if the given path identifier is currently registered.
 	 */
-	bool IsRegisteredPath(const FString& InPath) const;
+	UE_API bool IsRegisteredPath(const FString& InPath) const;
 	
 	/**
 	 * @return TRUE if the given path identifier is inclusive of any registered paths.
@@ -52,52 +54,52 @@ public:
 	 * Also note if a registered path is removed, inclusive paths may still be valid. For instance, if both "MyProject.MyNamespace" and
 	 * "MyProject.MyNamespace_2" are registered paths, and "MyProject.MyNamespace_2" is removed, "MyProject" is still an inclusive path.
 	 */
-	bool IsInclusivePath(const FString& InPath) const;
+	UE_API bool IsInclusivePath(const FString& InPath) const;
 
 	/**
 	 * @param InPath	Path identifier string (e.g. "X.Y" or "X.Y.").
 	 * @param OutNames	On output, an array containing the set of names rooted to the given path (e.g. "Z" in "X.Y.Z").
 	 */
-	void GetNamesUnderPath(const FString& InPath, TArray<FName>& OutNames) const;
+	UE_API void GetNamesUnderPath(const FString& InPath, TArray<FName>& OutNames) const;
 
 	/**
 	 * @param OutPaths	On output, contains the full set of all currently-registered namespace identifier paths.
 	 */
-	void GetAllRegisteredPaths(TArray<FString>& OutPaths) const;
+	UE_API void GetAllRegisteredPaths(TArray<FString>& OutPaths) const;
 
 	/**
 	 * Adds an explicit namespace identifier to the registry if not already included.
 	 * 
 	 * @param InPath	Path identifier string (e.g. "X.Y").
 	 */
-	void RegisterNamespace(const FString& InPath);
+	UE_API void RegisterNamespace(const FString& InPath);
 
 	/**
 	 * Recreates the namespace registry.
 	 */
-	void Rebuild();
+	UE_API void Rebuild();
 
 protected:
-	FBlueprintNamespaceRegistry();
+	UE_API FBlueprintNamespaceRegistry();
 
 	/** Asset registry event handler methods. */
-	void OnAssetAdded(const FAssetData& AssetData);
-	void OnAssetRemoved(const FAssetData& AssetData);
-	void OnAssetRenamed(const FAssetData& AssetData, const FString& InOldName);
-	void OnAssetRegistryFilesLoaded();
+	UE_API void OnAssetAdded(const FAssetData& AssetData);
+	UE_API void OnAssetRemoved(const FAssetData& AssetData);
+	UE_API void OnAssetRenamed(const FAssetData& AssetData, const FString& InOldName);
+	UE_API void OnAssetRegistryFilesLoaded();
 
 	/** Namespace identifier registration methods. */
-	void FindAndRegisterAllNamespaces();
-	void RegisterNamespace(const UObject* InObject);
-	void RegisterNamespace(const FAssetData& AssetData);
+	UE_API void FindAndRegisterAllNamespaces();
+	UE_API void RegisterNamespace(const UObject* InObject);
+	UE_API void RegisterNamespace(const FAssetData& AssetData);
 
 	/** Console command implementations (debugging/testing). */
-	void ToggleDefaultNamespace();
-	void DumpAllRegisteredPaths();
-	void OnDefaultNamespaceTypeChanged();
+	UE_API void ToggleDefaultNamespace();
+	UE_API void DumpAllRegisteredPaths();
+	UE_API void OnDefaultNamespaceTypeChanged();
 
 	/** Handler for hot reload / live coding completion events. */
-	void OnReloadComplete(EReloadCompleteReason InReason);
+	UE_API void OnReloadComplete(EReloadCompleteReason InReason);
 
 private:
 	/** Indicates whether the registry has been initialized. */
@@ -116,4 +118,6 @@ private:
 
 	/** Internal set of objects to exclude during namespace registration. */
 	TSet<FSoftObjectPath> ExcludedObjectPaths;
-};
\ No newline at end of file
+};
+
+#undef UE_API
diff --git a/Engine/Source/Editor/Kismet/Public/BlueprintNamespaceUtilities.h b/Engine/Source/Editor/Kismet/Public/BlueprintNamespaceUtilities.h
--- a/Engine/Source/Editor/Kismet/Public/BlueprintNamespaceUtilities.h
+++ b/Engine/Source/Editor/Kismet/Public/BlueprintNamespaceUtilities.h
@@ -9,6 +9,8 @@
 #include "Misc/CoreMiscDefines.h"
 #include "UObject/SoftObjectPath.h"
 
+#define UE_API KISMET_API
+
 class FProperty;
 class UBlueprint;
 class UObject;
@@ -26,7 +28,7 @@ enum class EDefaultBlueprintNamespaceType
 };
 
 /** A wrapper struct around various Blueprint namespace utility and support methods. */
-struct KISMET_API FBlueprintNamespaceUtilities
+struct FBlueprintNamespaceUtilities
 {
 public:
 	/**
@@ -35,7 +37,7 @@ public:
 	 * @param InAssetData	Input asset data.
 	 * @return The unique Blueprint namespace identifier associated with the given asset, or an empty string if the asset belongs to the global namespace (default).
 	 */
-	static FString GetAssetNamespace(const FAssetData& InAssetData);
+	static UE_API FString GetAssetNamespace(const FAssetData& InAssetData);
 
 	/**
 	 * Analyzes the given object to determine its explicitly-assigned namespace identifier, or otherwise returns its default namespace.
@@ -43,7 +45,7 @@ public:
 	 * @param InObject		A reference to the input object.
 	 * @return The unique Blueprint namespace identifier associated with the given object, or an empty string if the object belongs to the global namespace (default).
 	 */
-	static FString GetObjectNamespace(const UObject* InObject);
+	static UE_API FString GetObjectNamespace(const UObject* InObject);
 
 	/**
 	 * Analyzes the given object path to determine its explicitly-assigned namespace identifier, or otherwise returns its default namespace.
@@ -51,7 +53,7 @@ public:
 	 * @param InObjectPath	Path to the given object (may not be loaded yet).
 	 * @return The unique Blueprint namespace identifier associated with the given object (even if unloaded), or an empty string if the object belongs to the global namespace (default).
 	 */
-	static FString GetObjectNamespace(const FSoftObjectPath& InObjectPath);
+	static UE_API FString GetObjectNamespace(const FSoftObjectPath& InObjectPath);
 
 	/**
 	 * Analyzes a property value to determine explicitly-assigned namespace identifiers from any object references, or otherwise returns the default namespace for each occurrence (default).
@@ -60,7 +62,7 @@ public:
 	 * @param InContainer	The source address of the struct/object containing the property's value.
 	 * @param OutNamespaces	Zero or more unique namespace identifier(s) referenced by the property value. An entry with an empty string equates to the default global namespace.
 	 */
-	static void GetPropertyValueNamespaces(const FProperty* InProperty, const void* InContainer, TSet<FString>& OutNamespaces);
+	static UE_API void GetPropertyValueNamespaces(const FProperty* InProperty, const void* InContainer, TSet<FString>& OutNamespaces);
 
 	UE_DEPRECATED(5.1, "Please use the updated version that removes the InStruct parameter (no longer needed).")
 	static void GetPropertyValueNamespaces(const UStruct* InStruct, const FProperty* InProperty, const void* InContainer, TSet<FString>& OutNamespaces)
@@ -73,7 +75,7 @@ public:
 	 * 
 	 * @param OutNamespaces	Zero or more unique namespace identifier(s) representing the shared global namespace set.
 	 */
-	static void GetSharedGlobalImports(TSet<FString>& OutNamespaces);
+	static UE_API void GetSharedGlobalImports(TSet<FString>& OutNamespaces);
 
 	/**
 	 * Gathers the set of default namespaces that are implicitly imported by a given Blueprint asset.
@@ -82,7 +84,7 @@ public:
 	 * @param OutNamespaces	Zero or more unique namespace identifier(s) representing the Blueprint asset's default namespace set.
 	 */
 	UE_DEPRECATED(5.1, "Please use GetDefaultImportsForObject instead")
-	static void GetDefaultImportsForBlueprint(const UBlueprint* InBlueprint, TSet<FString>& OutNamespaces);
+	static UE_API void GetDefaultImportsForBlueprint(const UBlueprint* InBlueprint, TSet<FString>& OutNamespaces);
 
 	/**
 	 * Gathers the set of default namespaces that are implicitly imported by a given object's type.
@@ -90,26 +92,26 @@ public:
 	 * @param InObject		A reference to a loaded object or type.
 	 * @param OutNamespaces	Zero or more unique namespace identifier(s) representing the type's default namespace set.
 	 */
-	static void GetDefaultImportsForObject(const UObject* InObject, TSet<FString>& OutNamespaces);
+	static UE_API void GetDefaultImportsForObject(const UObject* InObject, TSet<FString>& OutNamespaces);
 
 	/**
 	 * Sets the default Blueprint namespace type that objects/assets should use when not explicitly assigned.
 	 *
 	 * @param InType		Default namespace type to use.
 	 */
-	static void SetDefaultBlueprintNamespaceType(EDefaultBlueprintNamespaceType InType);
+	static UE_API void SetDefaultBlueprintNamespaceType(EDefaultBlueprintNamespaceType InType);
 
 	/** @return The default Blueprint namespace type objects/assets should use. Currently used for debugging/testing. */
-	static EDefaultBlueprintNamespaceType GetDefaultBlueprintNamespaceType();
+	static UE_API EDefaultBlueprintNamespaceType GetDefaultBlueprintNamespaceType();
 
 	/** Delegate invoked whenever the default Blueprint namespace type changes. */
 	DECLARE_MULTICAST_DELEGATE(FOnDefaultBlueprintNamespaceTypeChanged);
-	static FOnDefaultBlueprintNamespaceTypeChanged& OnDefaultBlueprintNamespaceTypeChanged();
+	static UE_API FOnDefaultBlueprintNamespaceTypeChanged& OnDefaultBlueprintNamespaceTypeChanged();
 
 	/**
 	 * Refresh the Blueprint editor environment to align with current namespace editor feature settings.
 	 */
-	static void RefreshBlueprintEditorFeatures();
+	static UE_API void RefreshBlueprintEditorFeatures();
 
 	/**
 	 * Helper method to convert a package path to a Blueprint namespace identifier string.
@@ -117,5 +119,7 @@ public:
 	 * @param InPackagePath		Package path input string.
 	 * @param OutNamespacePath	The input string converted to namespace identifier format.
 	 */
-	static void ConvertPackagePathToNamespacePath(const FString& InPackagePath, FString& OutNamespacePath);
-};
\ No newline at end of file
+	static UE_API void ConvertPackagePathToNamespacePath(const FString& InPackagePath, FString& OutNamespacePath);
+};
+
+#undef UE_API
diff --git a/Engine/Source/Editor/Kismet/Public/Debugging/SKismetDebugTreeView.h b/Engine/Source/Editor/Kismet/Public/Debugging/SKismetDebugTreeView.h
--- a/Engine/Source/Editor/Kismet/Public/Debugging/SKismetDebugTreeView.h
+++ b/Engine/Source/Editor/Kismet/Public/Debugging/SKismetDebugTreeView.h
@@ -28,6 +28,8 @@
 #include "Widgets/Views/STableViewBase.h"
 #include "Widgets/Views/STreeView.h"
 
+#define UE_API KISMET_API
+
 class FBreakpointParentItem;
 class FDebugLineItem;
 class FName;
@@ -51,7 +53,7 @@ struct FGeometry;
 typedef TSharedPtr<class FDebugLineItem> FDebugTreeItemPtr;
 
 // The base class for a line entry in the debugging tree view
-class KISMET_API FDebugLineItem : public TSharedFromThis<FDebugLineItem>
+class FDebugLineItem : public TSharedFromThis<FDebugLineItem>
 {
 public:
 	friend class FLineItemWithChildren; // used by FLineItemWithChildren::EnsureChildIsAdded
@@ -79,22 +81,22 @@ public:
 	virtual ~FDebugLineItem() {}
 
 	// Create the widget for the name column
-	virtual TSharedRef<SWidget> GenerateNameWidget(TSharedPtr<FString> InSearchString);
+	UE_API virtual TSharedRef<SWidget> GenerateNameWidget(TSharedPtr<FString> InSearchString);
 
 	// Create the widget for the value column
-	virtual TSharedRef<SWidget> GenerateValueWidget(TSharedPtr<FString> InSearchString);
+	UE_API virtual TSharedRef<SWidget> GenerateValueWidget(TSharedPtr<FString> InSearchString);
 
 	// Add standard context menu items that can act on any node of the tree
-	void MakeMenu(class FMenuBuilder& MenuBuilder, bool bInDebuggerTab);
+	UE_API void MakeMenu(class FMenuBuilder& MenuBuilder, bool bInDebuggerTab);
 
 	// Add context menu items that can act on this node of the tree
-	virtual void ExtendContextMenu(class FMenuBuilder& MenuBuilder, bool bInDebuggerTab);
+	UE_API virtual void ExtendContextMenu(class FMenuBuilder& MenuBuilder, bool bInDebuggerTab);
 
 	// Gather all of the children
 	virtual void GatherChildrenBase(TArray<FDebugTreeItemPtr>& OutChildren, const FString& InSearchString, bool bRespectSearch = true) {}
 
 	// returns whether this tree node has children (used by drop down arrows)
-	virtual bool HasChildren() const;
+	UE_API virtual bool HasChildren() const;
 
 	// only line items inheriting from FLineItemWithChildren can have children
 	virtual bool CanHaveChildren() { return false; }
@@ -108,35 +110,35 @@ public:
 	}
 
 	// returns a widget that will go to the left of the Name Widget.
-	virtual TSharedRef<SWidget> GetNameIcon();
+	UE_API virtual TSharedRef<SWidget> GetNameIcon();
 
 	// returns a widget that will go to the left of the Value Widget.
-	virtual TSharedRef<SWidget> GetValueIcon();
+	UE_API virtual TSharedRef<SWidget> GetValueIcon();
 
 	// returns current value of Search String
-	FText GetHighlightText(const TSharedPtr<FString> InSearchString) const;
+	UE_API FText GetHighlightText(const TSharedPtr<FString> InSearchString) const;
 
 	// Helper function to try to get the blueprint for a given object;
 	//   Returns the blueprint that was used to create the instance if there was one
 	//   Returns the object itself if it is already a blueprint
 	//   Otherwise returns NULL
-	static UBlueprint* GetBlueprintForObject(UObject* ParentObject);
+	static UE_API UBlueprint* GetBlueprintForObject(UObject* ParentObject);
 
-	static UBlueprintGeneratedClass* GetClassForObject(UObject* ParentObject);
+	static UE_API UBlueprintGeneratedClass* GetClassForObject(UObject* ParentObject);
 
-	static bool IsDebugLineTypeActive(EDebugLineType Type);
-	static void OnDebugLineTypeActiveChanged(ECheckBoxState CheckState, EDebugLineType Type);
+	static UE_API bool IsDebugLineTypeActive(EDebugLineType Type);
+	static UE_API void OnDebugLineTypeActiveChanged(ECheckBoxState CheckState, EDebugLineType Type);
 
-	static void SetBreakpointParentItemBlueprint(FDebugTreeItemPtr InBreakpointParentItem, TWeakObjectPtr<UBlueprint> InBlueprint);
+	static UE_API void SetBreakpointParentItemBlueprint(FDebugTreeItemPtr InBreakpointParentItem, TWeakObjectPtr<UBlueprint> InBlueprint);
 
 	// updates bVisible and bParentsMatchSearch based on this node alone
-	void UpdateSearch(const FString& InSearchString, ESearchFlags SearchFlags);
+	UE_API void UpdateSearch(const FString& InSearchString, ESearchFlags SearchFlags);
 
-	bool IsVisible();
-	bool DoParentsMatchSearch();
+	UE_API bool IsVisible();
+	UE_API bool DoParentsMatchSearch();
 
 	// @return The text to display in the value column, unless GenerateValueWidget is overridden
-	virtual FText GetDescription() const;
+	UE_API virtual FText GetDescription() const;
 protected:
 
 	// Cannot create an instance of this class, it's just for use as a base class
@@ -159,21 +161,21 @@ protected:
 	[[maybe_unused]] virtual void UpdateData(const FDebugLineItem& NewerData) {}
 
 	// @return The actual underlying name of the item, for search compatibility
-	virtual FText GetName() const;
+	UE_API virtual FText GetName() const;
 
 	// @return The text to display in the name column, unless GenerateNameWidget is overridden
-	virtual FText GetDisplayName() const;
+	UE_API virtual FText GetDisplayName() const;
 
-	bool HasName() const;
-	bool HasValue() const;
-	void CopyNameToClipboard() const;
-	void CopyValueToClipboard() const;
+	UE_API bool HasName() const;
+	UE_API bool HasValue() const;
+	UE_API void CopyNameToClipboard() const;
+	UE_API void CopyValueToClipboard() const;
 
 protected:
 	// Type of action (a kind of RTTI for the tree, really only used to accelerate Compare checks)
 	EDebugLineType Type;
 
-	static uint16 ActiveTypeBitset;
+	static UE_API uint16 ActiveTypeBitset;
 
 	// true if self or any recursive children match the search
 	bool bVisible = false;
@@ -184,7 +186,7 @@ protected:
 //////////////////////////////////////////////////////////////////////////
 // SKismetDebugTreeView
 
-class KISMET_API SKismetDebugTreeView : public SCompoundWidget
+class SKismetDebugTreeView : public SCompoundWidget
 {
 public:
 	SLATE_BEGIN_ARGS(SKismetDebugTreeView)
@@ -200,42 +202,42 @@ public:
 
 		SLATE_END_ARGS()
 public:
-	void Construct(const FArguments& InArgs);
+	UE_API void Construct(const FArguments& InArgs);
 
 	// SWidget interface
-	virtual void Tick(const FGeometry& AllottedGeometry, const double InCurrentTime, const float InDeltaTime) override;
+	UE_API virtual void Tick(const FGeometry& AllottedGeometry, const double InCurrentTime, const float InDeltaTime) override;
 	// End of SWidget interface
 
 	/** Adds a root level TreeItem */
-	void AddTreeItemUnique(const FDebugTreeItemPtr& Item);
+	UE_API void AddTreeItemUnique(const FDebugTreeItemPtr& Item);
 
 	/** Removes a root level TreeItem */
-	bool RemoveTreeItem(const FDebugTreeItemPtr& Item);
+	UE_API bool RemoveTreeItem(const FDebugTreeItemPtr& Item);
 
 	/** Clears all root level Tree Items */
-	void ClearTreeItems();
+	UE_API void ClearTreeItems();
 
 	/** Updates search text and requests the Tree be updated */
-	void SetSearchText(const FText& InSearchText);
+	UE_API void SetSearchText(const FText& InSearchText);
 
 	/** Requests for an update of the tree at the next tick */
-	void RequestUpdateFilteredItems();
+	UE_API void RequestUpdateFilteredItems();
 
 	/** Returns the Array of Root Tree Items */
-	const TArray<FDebugTreeItemPtr>& GetRootTreeItems() const;
+	UE_API const TArray<FDebugTreeItemPtr>& GetRootTreeItems() const;
 
 	// Passthrough functions to the TreeView
-	int32 GetSelectedItems(TArray<FDebugTreeItemPtr>& OutItems);
-	void ClearExpandedItems();
-	bool IsScrolling() const;
-	void SetItemExpansion(FDebugTreeItemPtr InItem, bool bInShouldExpandItem);
+	UE_API int32 GetSelectedItems(TArray<FDebugTreeItemPtr>& OutItems);
+	UE_API void ClearExpandedItems();
+	UE_API bool IsScrolling() const;
+	UE_API void SetItemExpansion(FDebugTreeItemPtr InItem, bool bInShouldExpandItem);
 private:
 	/** Updates the Filtered Items based on current Root Tree Items */
-	void UpdateFilteredItems();
+	UE_API void UpdateFilteredItems();
 
-	TSharedRef<ITableRow> OnGenerateRow(FDebugTreeItemPtr InItem, const TSharedRef<STableViewBase>& OwnerTable);
-	void OnGetChildren(FDebugTreeItemPtr InParent, TArray<FDebugTreeItemPtr>& OutChildren);
-	TSharedPtr<SWidget> OnMakeContextMenu() const;
+	UE_API TSharedRef<ITableRow> OnGenerateRow(FDebugTreeItemPtr InItem, const TSharedRef<STableViewBase>& OwnerTable);
+	UE_API void OnGetChildren(FDebugTreeItemPtr InParent, TArray<FDebugTreeItemPtr>& OutChildren);
+	UE_API TSharedPtr<SWidget> OnMakeContextMenu() const;
 
 	TSharedPtr< STreeView< FDebugTreeItemPtr > > TreeView;
 	TArray<FDebugTreeItemPtr> FilteredTreeRoots;
@@ -252,13 +254,15 @@ private:
 	bool bInDebuggerTab;
 public:
 	/** Accessible functions for creating tree items */
-	static FDebugTreeItemPtr MakeTraceStackParentItem();
-	static FDebugTreeItemPtr MakeBreakpointParentItem(TWeakObjectPtr<UBlueprint> InBlueprint);
-	static FDebugTreeItemPtr MakeMessageItem(const FString& InMessage);
-	static FDebugTreeItemPtr MakeParentItem(UObject* InObject);
-	static FDebugTreeItemPtr MakeWatchLineItem(const UEdGraphPin* InPinRef, UObject* InDebugObject);
+	static UE_API FDebugTreeItemPtr MakeTraceStackParentItem();
+	static UE_API FDebugTreeItemPtr MakeBreakpointParentItem(TWeakObjectPtr<UBlueprint> InBlueprint);
+	static UE_API FDebugTreeItemPtr MakeMessageItem(const FString& InMessage);
+	static UE_API FDebugTreeItemPtr MakeParentItem(UObject* InObject);
+	static UE_API FDebugTreeItemPtr MakeWatchLineItem(const UEdGraphPin* InPinRef, UObject* InDebugObject);
 	
 	/** Column Id's for custom header rows */
-	static const FName ColumnId_Name;
-	static const FName ColumnId_Value;
+	static UE_API const FName ColumnId_Name;
+	static UE_API const FName ColumnId_Value;
 };
+
+#undef UE_API
diff --git a/Engine/Source/Editor/Kismet/Public/DetailsDiff.h b/Engine/Source/Editor/Kismet/Public/DetailsDiff.h
--- a/Engine/Source/Editor/Kismet/Public/DetailsDiff.h
+++ b/Engine/Source/Editor/Kismet/Public/DetailsDiff.h
@@ -10,43 +10,45 @@
 #include "Widgets/SWidget.h"
 #include "Widgets/Layout/LinkableScrollBar.h"
 
+#define UE_API KISMET_API
+
 class FPropertyPath;
 class IDetailsView;
 class UObject;
 
 /** Struct to handle showing details for an object and provide an interface for listing all differences */
-class KISMET_API FDetailsDiff
+class FDetailsDiff
 {
 public:
-	FDetailsDiff( const UObject* InObject, bool bScrollbarOnLeft = false );	
+	UE_API FDetailsDiff( const UObject* InObject, bool bScrollbarOnLeft = false );	
 
 	UE_DEPRECATED(5.6, "FOnDisplayedPropertiesChanged removed, interact with IDetailsView directly")
 	DECLARE_DELEGATE( FOnDisplayedPropertiesChanged );
 	
 	PRAGMA_DISABLE_DEPRECATION_WARNINGS
 	UE_DEPRECATED(5.6, "InOnDisplayedPropertiesChanged, interact with IDetailsView directly")
-	FDetailsDiff( const UObject* InObject, FOnDisplayedPropertiesChanged InOnDisplayedPropertiesChanged, bool bScrollbarOnLeft = false );
-	PRAGMA_ENABLE_DEPRECATION_WARNINGS
+	UE_API FDetailsDiff( const UObject* InObject, FOnDisplayedPropertiesChanged InOnDisplayedPropertiesChanged, bool bScrollbarOnLeft = false );
+	UE_API PRAGMA_ENABLE_DEPRECATION_WARNINGS
 
 	/** Attempt to highlight the property with the given path, may not always succeed */
 	void HighlightProperty( const FPropertySoftPath& PropertyName );
 
 	/** Returns actual widget that is used to display details */
-	TSharedRef< IDetailsView > DetailsWidget() const;
+	UE_API TSharedRef< IDetailsView > DetailsWidget() const;
 
 	/** Returns object being displayed */
 	const UObject* GetDisplayedObject() const { return DisplayedObject; }
 
 	/** Returns a list of all properties that would be diffed */
-	TArray<FPropertySoftPath> GetDisplayedProperties() const;
+	UE_API TArray<FPropertySoftPath> GetDisplayedProperties() const;
 
 	/** Perform a diff against another view, ordering either by display order or by remove/add/change */
-	void DiffAgainst(const FDetailsDiff& Newer, TArray<FSingleObjectDiffEntry>& OutDifferences, bool bSortByDisplayOrder = false) const;
+	UE_API void DiffAgainst(const FDetailsDiff& Newer, TArray<FSingleObjectDiffEntry>& OutDifferences, bool bSortByDisplayOrder = false) const;
 
 	/** Link the two details panels so they scroll in sync with one another */
-	static void LinkScrolling(FDetailsDiff& LeftPanel, FDetailsDiff& RightPanel, const TAttribute<TArray<FVector2f>>& ScrollRate);
+	static UE_API void LinkScrolling(FDetailsDiff& LeftPanel, FDetailsDiff& RightPanel, const TAttribute<TArray<FVector2f>>& ScrollRate);
 
-	static TSharedRef<IDetailsView> CreateDetailsView(const UObject* InObject, TSharedPtr<SScrollBar> ExternalScrollbar = nullptr, bool bScrollbarOnLeft = false);
+	static UE_API TSharedRef<IDetailsView> CreateDetailsView(const UObject* InObject, TSharedPtr<SScrollBar> ExternalScrollbar = nullptr, bool bScrollbarOnLeft = false);
 
 private:
 	TArray< FPropertyPath > DisplayedProperties;
@@ -56,3 +58,5 @@ private:
 
 	TSharedPtr<SLinkableScrollBar> ScrollBar;
 };
+
+#undef UE_API
diff --git a/Engine/Source/Editor/Kismet/Public/FindInBlueprintManager.h b/Engine/Source/Editor/Kismet/Public/FindInBlueprintManager.h
--- a/Engine/Source/Editor/Kismet/Public/FindInBlueprintManager.h
+++ b/Engine/Source/Editor/Kismet/Public/FindInBlueprintManager.h
@@ -38,6 +38,8 @@
 
 #include "FindInBlueprintManager.generated.h"
 
+#define UE_API KISMET_API
+
 class FArchive;
 class SDockTab;
 class SWidget;
@@ -67,83 +69,83 @@ typedef TSharedPtr<FImaginaryFiBData, ESPMode::ThreadSafe> FImaginaryFiBDataShar
 /**
  *Const values for Find-in-Blueprints to tag searchable data
  */
-struct KISMET_API FFindInBlueprintSearchTags
+struct FFindInBlueprintSearchTags
 {
 	/** Properties tag, for Blueprint variables */
-	static const FText FiB_Properties;
+	static UE_API const FText FiB_Properties;
 
 	/** Components tags */
-	static const FText FiB_Components;
-	static const FText FiB_IsSCSComponent;
+	static UE_API const FText FiB_Components;
+	static UE_API const FText FiB_IsSCSComponent;
 	/** End Components tags */
 
 	/** Nodes tag */
-	static const FText FiB_Nodes;
+	static UE_API const FText FiB_Nodes;
 
 	/** Schema Name tag, to identify the schema that a graph uses */
-	static const FText FiB_SchemaName;
+	static UE_API const FText FiB_SchemaName;
 	/** Uber graphs tag */
-	static const FText FiB_UberGraphs;
+	static UE_API const FText FiB_UberGraphs;
 	/** Function graph tag */
-	static const FText FiB_Functions;
+	static UE_API const FText FiB_Functions;
 	/** Macro graph tag */
-	static const FText FiB_Macros;
+	static UE_API const FText FiB_Macros;
 	/** Sub graph tag, for any sub-graphs in a Blueprint */
-	static const FText FiB_SubGraphs;
+	static UE_API const FText FiB_SubGraphs;
 	/** Extension graph tag, for any graphs in a Blueprint Extension */
-	static const FText FiB_ExtensionGraphs;
+	static UE_API const FText FiB_ExtensionGraphs;
 	/** Blueprint extension tag. */
-	static const FText FiB_Extensions;
+	static UE_API const FText FiB_Extensions;
 
 	/** Name tag */
-	static const FText FiB_Name;
+	static UE_API const FText FiB_Name;
 	/** Native Name tag */
-	static const FText FiB_NativeName;
+	static UE_API const FText FiB_NativeName;
 	/** Class Name tag */
-	static const FText FiB_ClassName;
+	static UE_API const FText FiB_ClassName;
 	/** NodeGuid tag */
-	static const FText FiB_NodeGuid;
+	static UE_API const FText FiB_NodeGuid;
 	/** Default value */
-	static const FText FiB_DefaultValue;
+	static UE_API const FText FiB_DefaultValue;
 	/** Tooltip tag */
-	static const FText FiB_Tooltip;
+	static UE_API const FText FiB_Tooltip;
 	/** Description tag */
-	static const FText FiB_Description;
+	static UE_API const FText FiB_Description;
 	/** Comment tag */
-	static const FText FiB_Comment;
+	static UE_API const FText FiB_Comment;
 	/** Path tag */
-	static const FText FiB_Path;
+	static UE_API const FText FiB_Path;
 	/** Parent Class tag */
-	static const FText FiB_ParentClass;
+	static UE_API const FText FiB_ParentClass;
 	/** Interfaces tag */
-	static const FText FiB_Interfaces;
+	static UE_API const FText FiB_Interfaces;
 	/** Class that originally defined the function tag */
-	static const FText FiB_FuncOriginClass;
+	static UE_API const FText FiB_FuncOriginClass;
 
 	/** Pin type tags */
 
 	/** Pins tag */
-	static const FText FiB_Pins;
+	static UE_API const FText FiB_Pins;
 
 	/** Pin Category tag */
-	static const FText FiB_PinCategory;
+	static UE_API const FText FiB_PinCategory;
 	/** Pin Sub-Category tag */
-	static const FText FiB_PinSubCategory;
+	static UE_API const FText FiB_PinSubCategory;
 	/** Pin object class tag */
-	static const FText FiB_ObjectClass;
+	static UE_API const FText FiB_ObjectClass;
 	/** Pin IsArray tag */
-	static const FText FiB_IsArray;
+	static UE_API const FText FiB_IsArray;
 	/** Pin IsReference tag */
-	static const FText FiB_IsReference;
+	static UE_API const FText FiB_IsReference;
 	/** Glyph icon tag */
-	static const FText FiB_Glyph;
+	static UE_API const FText FiB_Glyph;
 	/** Style set the glyph belongs to */
-	static const FText FiB_GlyphStyleSet;
+	static UE_API const FText FiB_GlyphStyleSet;
 	/** Glyph icon color tag */
-	static const FText FiB_GlyphColor;
+	static UE_API const FText FiB_GlyphColor;
 
 	// Identifier for metadata storage, completely unsearchable tag
-	static const FText FiBMetaDataTag;
+	static UE_API const FText FiBMetaDataTag;
 	/** End const values for Find-in-Blueprint */
 };
 
@@ -282,13 +284,13 @@ struct FSearchTagDataPair
 	FText Value;
 };
 
-struct KISMET_API FFiBMD
+struct FFiBMD
 {
-	static const FString FiBSearchableMD;
-	static const FString FiBSearchableShallowMD;
-	static const FString FiBSearchableExplicitMD;
-	static const FString FiBSearchableHiddenExplicitMD;
-	static const FString FiBSearchableFormatVersionMD;
+	static UE_API const FString FiBSearchableMD;
+	static UE_API const FString FiBSearchableShallowMD;
+	static UE_API const FString FiBSearchableExplicitMD;
+	static UE_API const FString FiBSearchableHiddenExplicitMD;
+	static UE_API const FString FiBSearchableFormatVersionMD;
 };
 
 /** Which assets to index for caching */
@@ -329,7 +331,7 @@ enum class EFiBCacheOpFlags
 ENUM_CLASS_FLAGS(EFiBCacheOpFlags);
 
 /** Options to configure the bulk caching task */
-struct KISMET_API FFindInBlueprintCachingOptions
+struct FFindInBlueprintCachingOptions
 {
 	/** Type of caching operation */
 	EFiBCacheOpType OpType = EFiBCacheOpType::CachePendingAssets;
@@ -360,32 +362,32 @@ ENUM_CLASS_FLAGS(EAddOrUpdateBlueprintSearchMetadataFlags);
 // FFindInBlueprintsResult
 
 /* Item that matched the search results */
-class KISMET_API FFindInBlueprintsResult : public TSharedFromThis< FFindInBlueprintsResult >
+class FFindInBlueprintsResult : public TSharedFromThis< FFindInBlueprintsResult >
 {
 public:
 	FFindInBlueprintsResult() = default;
 	virtual ~FFindInBlueprintsResult() = default;
 
 	/* Create a root */
-	explicit FFindInBlueprintsResult(const FText& InDisplayText);
+	UE_API explicit FFindInBlueprintsResult(const FText& InDisplayText);
 
 	/* Called when user clicks on the search item */
-	virtual FReply OnClick();
+	UE_API virtual FReply OnClick();
 
 	/* Get Category for this search result */
-	virtual FText GetCategory() const;
+	UE_API virtual FText GetCategory() const;
 
 	/* Create an icon to represent the result */
-	virtual TSharedRef<SWidget>	CreateIcon() const;
+	UE_API virtual TSharedRef<SWidget>	CreateIcon() const;
 
 	/** Finalizes any content for the search data that was unsafe to do on a separate thread */
 	virtual void FinalizeSearchData() {};
 
 	/* Gets the comment on this node if any */
-	FString GetCommentText() const;
+	UE_API FString GetCommentText() const;
 
 	/** gets the blueprint housing all these search results */
-	UBlueprint* GetParentBlueprint() const;
+	UE_API UBlueprint* GetParentBlueprint() const;
 
 	/**
 	* Parses search info for specific data important for displaying the search result in an easy to understand format
@@ -398,10 +400,10 @@ public:
 	virtual void ParseSearchInfo(FText InKey, FText InValue) {};
 
 	/** Returns the Object represented by this search information give the Blueprint it can be found in */
-	virtual UObject* GetObject(UBlueprint* InBlueprint) const;
+	UE_API virtual UObject* GetObject(UBlueprint* InBlueprint) const;
 
 	/** Returns the display string for the row */
-	FText GetDisplayString() const;
+	UE_API FText GetDisplayString() const;
 
 public:
 	/*Any children listed under this category */
@@ -422,7 +424,7 @@ typedef TSharedPtr<FFindInBlueprintsResult> FSearchResult;
 ////////////////////////////////////
 // FStreamSearch
 
-struct KISMET_API FStreamSearchOptions
+struct FStreamSearchOptions
 {
 	/** Filter to limit the FilteredImaginaryResults to */
 	enum ESearchQueryFilter ImaginaryDataFilter;
@@ -441,40 +443,40 @@ struct KISMET_API FStreamSearchOptions
 /**
  * Async task for searching Blueprints
  */
-class KISMET_API FStreamSearch : public FRunnable
+class FStreamSearch : public FRunnable
 {
 public:
 	/** Constructor */
-	FStreamSearch(const FString& InSearchValue, const FStreamSearchOptions& InSearchOptions = FStreamSearchOptions());
+	UE_API FStreamSearch(const FString& InSearchValue, const FStreamSearchOptions& InSearchOptions = FStreamSearchOptions());
 
 	/** Begin FRunnable Interface */
-	virtual bool Init() override;
+	UE_API virtual bool Init() override;
 
-	virtual uint32 Run() override;
+	UE_API virtual uint32 Run() override;
 
-	virtual void Stop() override;
+	UE_API virtual void Stop() override;
 
-	virtual void Exit() override;
+	UE_API virtual void Exit() override;
 	/** End FRunnable Interface */
 
 	/** Brings the thread to a safe stop before continuing. */
-	void EnsureCompletion();
+	UE_API void EnsureCompletion();
 
 	/** Returns TRUE if the thread is done with it's work. */
-	bool IsComplete() const;
+	UE_API bool IsComplete() const;
 
 	/** Returns TRUE if Stop() was called while work is still pending. */
-	bool WasStopped() const;
+	UE_API bool WasStopped() const;
 
 	/**
 	 * Appends the items filtered through the search filter to the passed array
 	 *
 	 * @param OutItemsFound		All the items found since last queried
 	 */
-	void GetFilteredItems(TArray<TSharedPtr<class FFindInBlueprintsResult>>& OutItemsFound);
+	UE_API void GetFilteredItems(TArray<TSharedPtr<class FFindInBlueprintsResult>>& OutItemsFound);
 
 	/** Helper function to query the percent complete this search is */
-	float GetPercentComplete() const;
+	UE_API float GetPercentComplete() const;
 
 	/** Returns the Out-of-Date Blueprint count */
 	int32 GetOutOfDateCount() const
@@ -483,7 +485,7 @@ public:
 	}
 
 	/** Returns the FilteredImaginaryResults from the search query, these results have been filtered by the ImaginaryDataFilter. */
-	void GetFilteredImaginaryResults(TArray<FImaginaryFiBDataSharedPtr>& OutFilteredImaginaryResults);
+	UE_API void GetFilteredImaginaryResults(TArray<FImaginaryFiBDataSharedPtr>& OutFilteredImaginaryResults);
 
 public:
 	/** Thread to run the cleanup FRunnable on */
@@ -522,19 +524,19 @@ private:
 // FFindInBlueprintSearchManager
 
 /** Singleton manager for handling all Blueprint searches, helps to manage the going progress of Blueprints, and is thread-safe. */
-class KISMET_API FFindInBlueprintSearchManager : public FTickableEditorObject
+class FFindInBlueprintSearchManager : public FTickableEditorObject
 {
 public:
-	static FFindInBlueprintSearchManager* Instance;
-	static FFindInBlueprintSearchManager& Get();
+	static UE_API FFindInBlueprintSearchManager* Instance;
+	static UE_API FFindInBlueprintSearchManager& Get();
 
-	FFindInBlueprintSearchManager();
-	~FFindInBlueprintSearchManager();
+	UE_API FFindInBlueprintSearchManager();
+	UE_API ~FFindInBlueprintSearchManager();
 
 	//~ Begin FTickableObject Interface
-	virtual void Tick(float DeltaTime) override;
-	virtual bool IsTickable() const override;
-	virtual TStatId GetStatId() const override;
+	UE_API virtual void Tick(float DeltaTime) override;
+	UE_API virtual bool IsTickable() const override;
+	UE_API virtual TStatId GetStatId() const override;
 	//~ End FTickableObject Interface
 
 	/**
@@ -543,7 +545,7 @@ public:
 	 * @param InSearchData		Search data to be applied
 	 * @param bAllowNewEntry	Whether to allow a new entry to be added if a match is not found
 	 */
-	void ApplySearchDataToDatabase(FSearchData InSearchData, bool bAllowNewEntry = false);
+	UE_API void ApplySearchDataToDatabase(FSearchData InSearchData, bool bAllowNewEntry = false);
 
 	/**
 	 * Given an asset path, locate and return a copy of its matching search data in the index cache.
@@ -551,7 +553,7 @@ public:
 	 * @param InAssetPath				Asset path (search index key).
 	 * @return							Matching search data from the index cache. Will return invalid (empty) search data if a matching entry was not found.
 	 */
-	FSearchData GetSearchDataForAssetPath(const FSoftObjectPath& InAssetPath);
+	UE_API FSearchData GetSearchDataForAssetPath(const FSoftObjectPath& InAssetPath);
 
 	/**
 	 * Gathers the Blueprint's search metadata and adds or updates it in the cache
@@ -560,14 +562,14 @@ public:
 	 * @param InFlags			Flags to assist with controlling this function's behavior
 	 * @param InVersion			Allows the caller to override the format version to use for caching
 	 */
-	void AddOrUpdateBlueprintSearchMetadata(UBlueprint* InBlueprint, EAddOrUpdateBlueprintSearchMetadataFlags InFlags = EAddOrUpdateBlueprintSearchMetadataFlags::None, EFiBVersion InVersion = EFiBVersion::FIB_VER_NONE);
+	UE_API void AddOrUpdateBlueprintSearchMetadata(UBlueprint* InBlueprint, EAddOrUpdateBlueprintSearchMetadataFlags InFlags = EAddOrUpdateBlueprintSearchMetadataFlags::None, EFiBVersion InVersion = EFiBVersion::FIB_VER_NONE);
 
 	/**
 	 * Starts a search query, the FiB manager handles where the thread is at in the search query at all times so that post-save of the cache to disk it can correct the index
 	 *
 	 * @param InSearchOriginator		Pointer to the thread object that the query originates from
 	 */
-	void BeginSearchQuery(const class FStreamSearch* InSearchOriginator);
+	UE_API void BeginSearchQuery(const class FStreamSearch* InSearchOriginator);
 
 	/**
 	 * Continues a search query, returning a single piece of search data
@@ -576,7 +578,7 @@ public:
 	 * @param OutSearchData				Result of the search, if there is any Blueprint search data still available to query
 	 * @return							TRUE if the search was successful, FALSE if the search is complete
 	 */
-	bool ContinueSearchQuery(const class FStreamSearch* InSearchOriginator, FSearchData& OutSearchData);
+	UE_API bool ContinueSearchQuery(const class FStreamSearch* InSearchOriginator, FSearchData& OutSearchData);
 
 	/**
 	 * This function ensures that the passed in search query ends in a safe manner. The search will no longer be valid to this manager, though it does not destroy any objects.
@@ -584,7 +586,7 @@ public:
 	 *
 	 * @param InSearchOriginator	Identifying search stream to be stopped
 	 */
-	void EnsureSearchQueryEnds(const class FStreamSearch* InSearchOriginator);
+	UE_API void EnsureSearchQueryEnds(const class FStreamSearch* InSearchOriginator);
 
 	/**
 	 * Query how far along a search thread is
@@ -592,7 +594,7 @@ public:
 	 * @param OutSearchData				Result of the search, if there is any Blueprint search data still available to query
 	 * @return							Percent along the search thread is
 	 */
-	float GetPercentComplete(const class FStreamSearch* InSearchOriginator) const;
+	UE_API float GetPercentComplete(const class FStreamSearch* InSearchOriginator) const;
 
 	/**
 	 * Query for a single, specific Blueprint's search data.
@@ -601,16 +603,16 @@ public:
 	 * @param bInRebuildSearchData		When TRUE the search data will be freshly collected
 	 * @return							The search data, or invalid (empty) search data if not found
 	 */
-	FSearchData QuerySingleBlueprint(UBlueprint* InBlueprint, bool bInRebuildSearchData);
+	UE_API FSearchData QuerySingleBlueprint(UBlueprint* InBlueprint, bool bInRebuildSearchData);
 
 	/** Processes the encoded string value in the SearchData into the intermediate format, return true if string and version were valid */
-	bool ProcessEncodedValueForUnloadedBlueprint(FSearchData& SearchData);
+	UE_API bool ProcessEncodedValueForUnloadedBlueprint(FSearchData& SearchData);
 
 	/** Returns the number of unindexed Blueprints, either due to not having been indexed before, or AR data being out-of-date */
-	int32 GetNumberUnindexedAssets() const;
+	UE_API int32 GetNumberUnindexedAssets() const;
 
 	/** Returns the number of uncached assets during an active indexing operation */
-	int32 GetNumberUncachedAssets() const;
+	UE_API int32 GetNumberUncachedAssets() const;
 
 	/**
 	 * Starts a task to cache Blueprints at a rate of 1 per tick
@@ -618,12 +620,12 @@ public:
 	 * @param InSourceWidget				The source FindInBlueprints widget, this widget will be informed when caching is complete
 	 * @param InCachingOptions				Options to configure the caching task
 	 */
-	void CacheAllAssets(TWeakPtr< class SFindInBlueprints > InSourceWidget, const FFindInBlueprintCachingOptions& InCachingOptions);
+	UE_API void CacheAllAssets(TWeakPtr< class SFindInBlueprints > InSourceWidget, const FFindInBlueprintCachingOptions& InCachingOptions);
 
 	/**
 	 * Exports a list of all unindexed assets to Saved/FindInBlueprints_OutdatedAssetList.txt
 	 */
-	void ExportOutdatedAssetList();
+	UE_API void ExportOutdatedAssetList();
 
 	/**
 	 * Starts the actual caching process
@@ -631,19 +633,19 @@ public:
 	 * @param bInSourceControlActive		TRUE if source control is active
 	 * @param bInCheckoutAndSave			TRUE if the system should checkout and save all assets that need to be reindexed
 	 */
-	void OnCacheAllUnindexedAssets(bool bInSourceControlActive, bool bInCheckoutAndSave);
+	UE_API void OnCacheAllUnindexedAssets(bool bInSourceControlActive, bool bInCheckoutAndSave);
 
 	/** Stops the caching process where it currently is at, the rest can be continued later */
-	void CancelCacheAll(SFindInBlueprints* InFindInBlueprintWidget);
+	UE_API void CancelCacheAll(SFindInBlueprints* InFindInBlueprintWidget);
 
 	/** Returns the current index in the caching */
-	int32 GetCurrentCacheIndex() const;
+	UE_API int32 GetCurrentCacheIndex() const;
 
 	/** Returns the path of the current Blueprint being cached */
-	FSoftObjectPath GetCurrentCacheBlueprintPath() const;
+	UE_API FSoftObjectPath GetCurrentCacheBlueprintPath() const;
 
 	/** Returns the progress complete on the caching */
-	float GetCacheProgress() const;
+	UE_API float GetCacheProgress() const;
 
 	/** Returns the list of Blueprint paths that failed to cache */
 	TSet<FSoftObjectPath> GetFailedToCachePathList() const { return FailedToCachePaths; }
@@ -655,26 +657,26 @@ public:
 	bool HasCachingFailed() const { return FailedToCachePaths.Num() > 0; };
 
 	/** Callback to note that Blueprint caching is started */
-	void StartedCachingBlueprints(EFiBCacheOpType InCacheOpType, EFiBCacheOpFlags InCacheOpFlags);
+	UE_API void StartedCachingBlueprints(EFiBCacheOpType InCacheOpType, EFiBCacheOpFlags InCacheOpFlags);
 
 	/**
 	 * Callback to note that Blueprint caching is complete
 	 *
 	 * @param InNumberCached		The number of Blueprints cached, to be chopped off the existing array so the rest (if any) can be finished later
 	 */
-	void FinishedCachingBlueprints(EFiBCacheOpType InCacheOpType, EFiBCacheOpFlags InCacheOpFlags, int32 InNumberCached, TSet<FSoftObjectPath>& InFailedToCacheList);
+	UE_API void FinishedCachingBlueprints(EFiBCacheOpType InCacheOpType, EFiBCacheOpFlags InCacheOpFlags, int32 InNumberCached, TSet<FSoftObjectPath>& InFailedToCacheList);
 
 	/** Returns TRUE if Blueprints are being cached. */
-	bool IsCacheInProgress() const;
+	UE_API bool IsCacheInProgress() const;
 
 	/** Returns TRUE if unindexed Blueprints are being cached (since this can block the UI) */
-	bool IsUnindexedCacheInProgress() const;
+	UE_API bool IsUnindexedCacheInProgress() const;
 
 	/** Returns TRUE if we're still inside the initial asset discovery and registration stage */
-	bool IsAssetDiscoveryInProgress() const;
+	UE_API bool IsAssetDiscoveryInProgress() const;
 
 	/** Returns TRUE if there are one or more active asynchronous search queries */
-	bool IsAsyncSearchQueryInProgress() const;
+	UE_API bool IsAsyncSearchQueryInProgress() const;
 
 	/** Returns a weak reference to the widget that initiated the current caching operation */
 	TWeakPtr<SFindInBlueprints> GetSourceCachingWidget() const { return SourceCachingWidget; }
@@ -690,77 +692,77 @@ public:
 	bool ShouldEnableSearchResultTemplates() const { return !bDisableSearchResultTemplates; }
 
 	/** Find or create the global find results widget */
-	TSharedPtr<SFindInBlueprints> GetGlobalFindResults();
+	UE_API TSharedPtr<SFindInBlueprints> GetGlobalFindResults();
 
 	/** Enable or disable the global find results tab feature */
-	void EnableGlobalFindResults(bool bEnable);
+	UE_API void EnableGlobalFindResults(bool bEnable);
 
 	/** Close any orphaned global find results tabs for a particular tab manager */
-	void CloseOrphanedGlobalFindResultsTabs(TSharedPtr<class FTabManager> TabManager);
+	UE_API void CloseOrphanedGlobalFindResultsTabs(TSharedPtr<class FTabManager> TabManager);
 
 	/** Returns true if a global find results tab is currently open */
 	bool IsGlobalFindResultsOpen() const { return GlobalFindResults.Num() > 0; }
 
-	void GlobalFindResultsClosed(const TSharedRef<SFindInBlueprints>& FindResults);
+	UE_API void GlobalFindResultsClosed(const TSharedRef<SFindInBlueprints>& FindResults);
 
 	/** Dumps the full index cache to the given stream (for debugging purposes) */
-	void DumpCache(FArchive& Ar);
+	UE_API void DumpCache(FArchive& Ar);
 
 	/** Randomly initiates a partial rebuild of the index cache (for debugging purposes) */
-	void DoRandomIndexCacheRebuild();
+	UE_API void DoRandomIndexCacheRebuild();
 
 public:
 	/** Converts a string of hex characters, previously converted by ConvertFTextToHexString, to an FText. */
-	static FText ConvertHexStringToFText(FString InHexString);
+	static UE_API FText ConvertHexStringToFText(FString InHexString);
 
 	/** Serializes an FText to memory and converts the memory into a string of hex characters */
-	static FString ConvertFTextToHexString(FText InValue);
+	static UE_API FString ConvertFTextToHexString(FText InValue);
 
 	/** Given a fully constructed Find-in-Blueprint FString of searchable data, will parse and construct a JsonObject */
-	static TSharedPtr< class FJsonObject > ConvertJsonStringToObject(FSearchDataVersionInfo InVersionInfo, const FString& InJsonString, TMap<int32, FText>& OutFTextLookupTable);
+	static UE_API TSharedPtr< class FJsonObject > ConvertJsonStringToObject(FSearchDataVersionInfo InVersionInfo, const FString& InJsonString, TMap<int32, FText>& OutFTextLookupTable);
 
 	/** Generates a human-readable search index for the given Blueprint (for debugging purposes) */
-	static FString GenerateSearchIndexForDebugging(UBlueprint* InBlueprint);
+	static UE_API FString GenerateSearchIndexForDebugging(UBlueprint* InBlueprint);
 
 private:
 	/** Initializes the FiB manager */
-	void Initialize();
+	UE_API void Initialize();
 
 	/** Callback hook during pre-garbage collection, pauses all processing searches so that the GC can do it's job */ 
-	void PauseFindInBlueprintSearch();
+	UE_API void PauseFindInBlueprintSearch();
 
 	/** Callback hook during post-garbage collection, saves the cache to disk and unpauses all processing searches */
-	void UnpauseFindInBlueprintSearch();
+	UE_API void UnpauseFindInBlueprintSearch();
 
 	/** Callback hook from the Asset Registry when an asset is added */
-	void OnAssetAdded(const struct FAssetData& InAssetData);
+	UE_API void OnAssetAdded(const struct FAssetData& InAssetData);
 
 	/** Callback hook from the Asset Registry, marks the asset for deletion from the cache */
-	void OnAssetRemoved(const struct FAssetData& InAssetData);
+	UE_API void OnAssetRemoved(const struct FAssetData& InAssetData);
 
 	/** Callback hook from the Asset Registry, marks the asset for deletion from the cache */
-	void OnAssetRenamed(const struct FAssetData& InAssetData, const FString& InOldName);
+	UE_API void OnAssetRenamed(const struct FAssetData& InAssetData, const FString& InOldName);
 
 	/** Callback hook from the Asset Registry, signals that the initial asset discovery stage has been completed */
-	void OnAssetRegistryFilesLoaded();
+	UE_API void OnAssetRegistryFilesLoaded();
 
 	/** Callback hook from the Asset Registry when an asset is loaded */
-	void OnAssetLoaded(class UObject* InAsset);
+	UE_API void OnAssetLoaded(class UObject* InAsset);
 
 	/** Callback from Kismet when a Blueprint is unloaded */
-	void OnBlueprintUnloaded(class UBlueprint* InBlueprint);
+	UE_API void OnBlueprintUnloaded(class UBlueprint* InBlueprint);
 
 	/** Callback hook from the Reload manager that indicates that a module has been reloaded */
-	void OnReloadComplete(EReloadCompleteReason Reason);
+	UE_API void OnReloadComplete(EReloadCompleteReason Reason);
 
 	/** Returns a copy of the search data that's cached at the given index. Will return invalid (empty) search data if the index is out of range */
-	FSearchData GetSearchDataForIndex(int32 CacheIndex);
+	UE_API FSearchData GetSearchDataForIndex(int32 CacheIndex);
 
 	/** Cleans the cache of any excess data from Blueprints that have been moved, renamed, or deleted. Occurs during post-garbage collection */
-	void CleanCache();
+	UE_API void CleanCache();
 
 	/** Builds the cache from all available Blueprint assets that the asset registry has discovered at the time of this function. Occurs on startup */
-	void BuildCache();
+	UE_API void BuildCache();
 
 	/**
 	 * Helper to properly add a Blueprint's SearchData to the database
@@ -768,25 +770,25 @@ private:
 	 * @param InSearchData		Data to add to the database
 	 * @return					Index into the SearchArray for looking up the added item
 	 */
-	int32 AddSearchDataToDatabase(FSearchData InSearchData);
+	UE_API int32 AddSearchDataToDatabase(FSearchData InSearchData);
 
 	/** Removes a Blueprint from being managed by the FiB system by passing in the UBlueprint's path */
-	void RemoveBlueprintByPath(const FSoftObjectPath& InPath);
+	UE_API void RemoveBlueprintByPath(const FSoftObjectPath& InPath);
 
 	/** Adds a new search database entry for unloaded asset data */
-	void AddUnloadedBlueprintSearchMetadata(const FAssetData& InAssetData);
+	UE_API void AddUnloadedBlueprintSearchMetadata(const FAssetData& InAssetData);
 
 	/** Begins the process of extracting FiB data from an unloaded asset */
-	void ExtractUnloadedFiBData(const FAssetData& InAssetData, FString* InFiBData, FName InKeyForFiBData, EFiBVersion InFiBDataVersion);
+	UE_API void ExtractUnloadedFiBData(const FAssetData& InAssetData, FString* InFiBData, FName InKeyForFiBData, EFiBVersion InFiBDataVersion);
 
 	/** Determines the global find results tab label */
-	FText GetGlobalFindResultsTabLabel(int32 TabIdx);
+	UE_API FText GetGlobalFindResultsTabLabel(int32 TabIdx);
 
 	/** Handler for a request to spawn a new global find results tab */
-	TSharedRef<SDockTab> SpawnGlobalFindResultsTab(const FSpawnTabArgs& SpawnTabArgs, int32 TabIdx);
+	UE_API TSharedRef<SDockTab> SpawnGlobalFindResultsTab(const FSpawnTabArgs& SpawnTabArgs, int32 TabIdx);
 
 	/** Creates and opens a new global find results tab */
-	TSharedPtr<SFindInBlueprints> OpenGlobalFindResultsTab();
+	UE_API TSharedPtr<SFindInBlueprints> OpenGlobalFindResultsTab();
 
 protected:
 	/** Contains info about an active search query */
@@ -809,13 +811,13 @@ protected:
 	typedef TSharedPtr<FActiveSearchQuery, ESPMode::ThreadSafe> FActiveSearchQueryPtr;
 
 	/** Thread-safe access to the active search query that's mapped to the given stream search */
-	FActiveSearchQueryPtr FindSearchQuery(const class FStreamSearch* InSearchOriginator) const;
+	UE_API FActiveSearchQueryPtr FindSearchQuery(const class FStreamSearch* InSearchOriginator) const;
 
 	/** Returns the next pending search data for the given query and advances the index to the next entry */
-	FSearchData GetNextSearchDataForQuery(const FStreamSearch* InSearchOriginator, FActiveSearchQueryPtr InSearchQueryPtr, bool bCheckDeferredList);
+	UE_API FSearchData GetNextSearchDataForQuery(const FStreamSearch* InSearchOriginator, FActiveSearchQueryPtr InSearchQueryPtr, bool bCheckDeferredList);
 
 	/** If searches are paused, blocks the calling thread until searching is resumed */
-	void BlockSearchQueryIfPaused();
+	UE_API void BlockSearchQueryIfPaused();
 
 private:
 	/** Maps the Blueprint paths to their index in the SearchArray */
@@ -907,7 +909,7 @@ private:
 	bool bDisableImmediateAssetDiscovery;
 };
 
-struct KISMET_API FDisableGatheringDataOnScope
+struct FDisableGatheringDataOnScope
 {
 	bool bOriginallyEnabled;
 	FDisableGatheringDataOnScope() : bOriginallyEnabled(FFindInBlueprintSearchManager::Get().IsGatheringDataEnabled())
@@ -919,3 +921,5 @@ struct KISMET_API FDisableGatheringDataOnScope
 		FFindInBlueprintSearchManager::Get().EnableGatheringData(bOriginallyEnabled);
 	}
 };
+
+#undef UE_API
diff --git a/Engine/Source/Editor/Kismet/Public/FindInBlueprints.h b/Engine/Source/Editor/Kismet/Public/FindInBlueprints.h
--- a/Engine/Source/Editor/Kismet/Public/FindInBlueprints.h
+++ b/Engine/Source/Editor/Kismet/Public/FindInBlueprints.h
@@ -37,6 +37,8 @@
 #include "Widgets/Views/STreeView.h"
 #include "FindInBlueprints.generated.h"
 
+#define UE_API KISMET_API
+
 class FBlueprintEditor;
 class FImaginaryFiBData;
 class FJsonValue;
@@ -275,7 +277,7 @@ enum class EFiBIndexAllPermission
 };
 
 /*Widget for searching for (functions/events) across all blueprints or just a single blueprint */
-class KISMET_API SFindInBlueprints: public SCompoundWidget
+class SFindInBlueprints: public SCompoundWidget
 {
 public:
 	SLATE_BEGIN_ARGS( SFindInBlueprints )
@@ -290,11 +292,11 @@ public:
 		SLATE_ARGUMENT(TSharedPtr<SDockTab>, ContainingTab)
 	SLATE_END_ARGS()
 
-	void Construct(const FArguments& InArgs, TSharedPtr<class FBlueprintEditor> InBlueprintEditor = nullptr);
-	~SFindInBlueprints();
+	UE_API void Construct(const FArguments& InArgs, TSharedPtr<class FBlueprintEditor> InBlueprintEditor = nullptr);
+	UE_API ~SFindInBlueprints();
 
 	/** Focuses this widget's search box, and changes the mode as well, and optionally the search terms */
-	void FocusForUse(bool bSetFindWithinBlueprint, FString NewSearchTerms = FString(), bool bSelectFirstResult = false);
+	UE_API void FocusForUse(bool bSetFindWithinBlueprint, FString NewSearchTerms = FString(), bool bSelectFirstResult = false);
 
 	/**
 	 * Submits a search query
@@ -304,26 +306,26 @@ public:
 	 * @param InSearchOptions						Optional search parameters.
 	 * @param InOnSearchComplete					Callback when the search is complete, passing the filtered imaginary data (if any).
 	 */
-	void MakeSearchQuery(FString InSearchString, bool bInIsFindWithinBlueprint, const FStreamSearchOptions& InSearchOptions = FStreamSearchOptions(), FOnSearchComplete InOnSearchComplete = FOnSearchComplete());
+	UE_API void MakeSearchQuery(FString InSearchString, bool bInIsFindWithinBlueprint, const FStreamSearchOptions& InSearchOptions = FStreamSearchOptions(), FOnSearchComplete InOnSearchComplete = FOnSearchComplete());
 
 	/** Called when caching Blueprints is started */
-	void OnCacheStarted(EFiBCacheOpType InOpType, EFiBCacheOpFlags InOpFlags);
+	UE_API void OnCacheStarted(EFiBCacheOpType InOpType, EFiBCacheOpFlags InOpFlags);
 	
 	/** Called when caching Blueprints is complete */
-	void OnCacheComplete(EFiBCacheOpType InOpType, EFiBCacheOpFlags InOpFlags);
+	UE_API void OnCacheComplete(EFiBCacheOpType InOpType, EFiBCacheOpFlags InOpFlags);
 
 	/**
 	 * Asynchronously caches all Blueprints below a specified version.
 	 *
 	 * @param InOptions								Options to configure the caching task
 	 */
-	void CacheAllBlueprints(const FFindInBlueprintCachingOptions& InOptions);
+	UE_API void CacheAllBlueprints(const FFindInBlueprintCachingOptions& InOptions);
 
 	/** If this is a global find results widget, returns the host tab's unique ID. Otherwise, returns NAME_None. */
-	FName GetHostTabId() const;
+	UE_API FName GetHostTabId() const;
 
 	/** If this is a global find results widget, ask the host tab to close */
-	void CloseHostTab();
+	UE_API void CloseHostTab();
 
 	/** Determines if this context does not accept syncing from an external source */
 	bool IsLocked() const
@@ -332,113 +334,113 @@ public:
 	}
 
 	/** Determines whether a search query is actively in progress */
-	bool IsSearchInProgress() const;
+	UE_API bool IsSearchInProgress() const;
 
 	/** SWidget overrides */
-	virtual FReply OnKeyDown(const FGeometry& MyGeometry, const FKeyEvent& InKeyEvent) override;
+	UE_API virtual FReply OnKeyDown(const FGeometry& MyGeometry, const FKeyEvent& InKeyEvent) override;
 
 	/** Clears the currently visible results */
-	void ClearResults();
+	UE_API void ClearResults();
 
 private:
 	/** Processes results of the ongoing async stream search */
-	EActiveTimerReturnType UpdateSearchResults( double InCurrentTime, float InDeltaTime );
+	UE_API EActiveTimerReturnType UpdateSearchResults( double InCurrentTime, float InDeltaTime );
 
 	/** Register any Find-in-Blueprint commands */
-	void RegisterCommands();
+	UE_API void RegisterCommands();
 
 	/*Called when user changes the text they are searching for */
-	void OnSearchTextChanged(const FText& Text);
+	UE_API void OnSearchTextChanged(const FText& Text);
 
 	/*Called when user changes commits text to the search box */
-	void OnSearchTextCommitted(const FText& Text, ETextCommit::Type CommitType);
+	UE_API void OnSearchTextCommitted(const FText& Text, ETextCommit::Type CommitType);
 
 	/* Get the children of a row */
-	void OnGetChildren( FSearchResult InItem, TArray< FSearchResult >& OutChildren );
+	UE_API void OnGetChildren( FSearchResult InItem, TArray< FSearchResult >& OutChildren );
 
 	/* Called when user double clicks on a new result */
-	void OnTreeSelectionDoubleClicked( FSearchResult Item );
+	UE_API void OnTreeSelectionDoubleClicked( FSearchResult Item );
 
 	/* Called when a new row is being generated */
-	TSharedRef<ITableRow> OnGenerateRow(FSearchResult InItem, const TSharedRef<STableViewBase>& OwnerTable);
+	UE_API TSharedRef<ITableRow> OnGenerateRow(FSearchResult InItem, const TSharedRef<STableViewBase>& OwnerTable);
 	
 	/** Launches a thread for streaming more content into the results widget */
-	void LaunchStreamThread(const FString& InSearchValue, const FStreamSearchOptions& InSearchOptions = FStreamSearchOptions(), FOnSearchComplete InOnSearchComplete = FOnSearchComplete());
+	UE_API void LaunchStreamThread(const FString& InSearchValue, const FStreamSearchOptions& InSearchOptions = FStreamSearchOptions(), FOnSearchComplete InOnSearchComplete = FOnSearchComplete());
 
 	/** Returns the percent complete on the search for the progress bar */
-	TOptional<float> GetPercentCompleteSearch() const;
+	UE_API TOptional<float> GetPercentCompleteSearch() const;
 
 	/** Returns the search bar visiblity for the given widget */
-	EVisibility GetSearchBarWidgetVisiblity(EFiBSearchBarWidget InSearchBarWidget) const;
+	UE_API EVisibility GetSearchBarWidgetVisiblity(EFiBSearchBarWidget InSearchBarWidget) const;
 
 	/** Adds the "cache" bar at the bottom of the Find-in-Blueprints widget, to notify the user that the search is incomplete */
-	void ConditionallyAddCacheBar();
+	UE_API void ConditionallyAddCacheBar();
 
 	/** Callback to remove the "cache" bar when a button is pressed */
-	FReply OnRemoveCacheBar();
+	UE_API FReply OnRemoveCacheBar();
 
 	/** Callback to return the cache bar's display text, informing the user of the situation */
-	FText GetCacheBarStatusText() const;
+	UE_API FText GetCacheBarStatusText() const;
 
 	/** Callback to return the current asset name during a cache operation */
-	FText GetCacheBarCurrentAssetName() const;
+	UE_API FText GetCacheBarCurrentAssetName() const;
 
 	/** Whether user is allowed to initiate loading and indexing all blueprints with out-of-date metadata */
-	bool CanCacheAllUnindexedBlueprints() const;
+	UE_API bool CanCacheAllUnindexedBlueprints() const;
 
 	/** Callback to cache all unindexed Blueprints */
-	FReply OnCacheAllUnindexedBlueprints();
+	UE_API FReply OnCacheAllUnindexedBlueprints();
 
 	/** Callback to export a list of all blueprints that need reindexing */
-	FReply OnExportUnindexedAssetList();
+	UE_API FReply OnExportUnindexedAssetList();
 
 	/** Callback to cache all Blueprints according to the given options */
-	FReply OnCacheAllBlueprints(const FFindInBlueprintCachingOptions& InOptions);
+	UE_API FReply OnCacheAllBlueprints(const FFindInBlueprintCachingOptions& InOptions);
 
 	/** Callback to cancel the caching process */
-	FReply OnCancelCacheAll();
+	UE_API FReply OnCancelCacheAll();
 
 	/** Retrieves the current index of the Blueprint caching process */
-	int32 GetCurrentCacheIndex() const;
+	UE_API int32 GetCurrentCacheIndex() const;
 
 	/** Gets the percent complete of the caching process */
-	TOptional<float> GetPercentCompleteCache() const;
+	UE_API TOptional<float> GetPercentCompleteCache() const;
 
 	/** Returns the caching bar's visibility, it goes invisible when there is nothing to be cached. The next search will remove this bar or make it visible again */
-	EVisibility GetCacheBarVisibility() const;
+	UE_API EVisibility GetCacheBarVisibility() const;
 
 	/** Returns the cache bar visibility for the given widget */
-	EVisibility GetCacheBarWidgetVisibility(EFiBCacheBarWidget InCacheBarWidget) const;
+	UE_API EVisibility GetCacheBarWidgetVisibility(EFiBCacheBarWidget InCacheBarWidget) const;
 
 	/** Returns TRUE if Blueprint caching is in progress */
-	bool IsCacheInProgress() const;
+	UE_API bool IsCacheInProgress() const;
 
 	/** Returns the BG image used for the cache bar */
-	const FSlateBrush* GetCacheBarImage() const;
+	UE_API const FSlateBrush* GetCacheBarImage() const;
 
 	/** Callback to build the context menu when right clicking in the tree */
-	TSharedPtr<SWidget> OnContextMenuOpening();
+	UE_API TSharedPtr<SWidget> OnContextMenuOpening();
 
 	/** Helper function to select all items */
-	void SelectAllItemsHelper(FSearchResult InItemToSelect);
+	UE_API void SelectAllItemsHelper(FSearchResult InItemToSelect);
 
 	/** Callback when user attempts to select all items in the search results */
-	void OnSelectAllAction();
+	UE_API void OnSelectAllAction();
 
 	/** Callback when user attempts to copy their selection in the Find-in-Blueprints */
-	void OnCopyAction();
+	UE_API void OnCopyAction();
 
 	/** Called when the user clicks the global find results button */
-	FReply OnOpenGlobalFindResults();
+	UE_API FReply OnOpenGlobalFindResults();
 
 	/** Called when the host tab is closed (if valid) */
-	void OnHostTabClosed(TSharedRef<SDockTab> DockTab);
+	UE_API void OnHostTabClosed(TSharedRef<SDockTab> DockTab);
 
 	/** Called when the lock button is clicked in a global find results tab */
-	FReply OnLockButtonClicked();
+	UE_API FReply OnLockButtonClicked();
 
 	/** Returns the image used for the lock button in a global find results tab */
-	const FSlateBrush* OnGetLockButtonImage() const;
+	UE_API const FSlateBrush* OnGetLockButtonImage() const;
 
 private:
 	/** Pointer back to the blueprint editor that owns us */
@@ -510,3 +512,5 @@ private:
 	/** True if cache bar should remain visible after a caching operation has ended */
 	bool bKeepCacheBarProgressVisible;
 };
+
+#undef UE_API
diff --git a/Engine/Source/Editor/Kismet/Public/MyBlueprintItemDragDropAction.h b/Engine/Source/Editor/Kismet/Public/MyBlueprintItemDragDropAction.h
--- a/Engine/Source/Editor/Kismet/Public/MyBlueprintItemDragDropAction.h
+++ b/Engine/Source/Editor/Kismet/Public/MyBlueprintItemDragDropAction.h
@@ -13,21 +13,23 @@
 #include "UObject/Class.h"
 #include "UObject/UnrealType.h"
 
+#define UE_API KISMET_API
+
 class FText;
 class UBlueprint;
 class UEdGraph;
 struct FEdGraphSchemaAction;
 
 /** DragDropAction class for drag and dropping an item from the My Blueprints tree (e.g., variable or function) */
-class KISMET_API FMyBlueprintItemDragDropAction : public FGraphSchemaActionDragDropAction
+class FMyBlueprintItemDragDropAction : public FGraphSchemaActionDragDropAction
 {
 public:
 	DRAG_DROP_OPERATOR_TYPE(FMyBlueprintItemDragDropAction, FGraphSchemaActionDragDropAction)
 
 	// FGraphEditorDragDropAction interface
-	virtual FReply DroppedOnAction(TSharedRef<FEdGraphSchemaAction> Action) override;
-	virtual FReply DroppedOnCategory(FText Category) override;
-	virtual void HoverTargetChanged() override;
+	UE_API virtual FReply DroppedOnAction(TSharedRef<FEdGraphSchemaAction> Action) override;
+	UE_API virtual FReply DroppedOnCategory(FText Category) override;
+	UE_API virtual void HoverTargetChanged() override;
 	// End of FGraphEditorDragDropAction
 
 	/** Set if operation is modified by alt */
@@ -38,7 +40,7 @@ public:
 
 protected:
 	 /** Constructor */
-	FMyBlueprintItemDragDropAction();
+	UE_API FMyBlueprintItemDragDropAction();
 
 	virtual UBlueprint* GetSourceBlueprint() const
 	{
@@ -51,8 +53,8 @@ protected:
 		return GetSourceBlueprint() == InBlueprint;
 	}
 
-	void SetFeedbackMessageError(const FText& Message);
-	void SetFeedbackMessageOK(const FText& Message);
+	UE_API void SetFeedbackMessageError(const FText& Message);
+	UE_API void SetFeedbackMessageOK(const FText& Message);
 
 protected:
 	/** Was ctrl held down at start of drag */
@@ -62,3 +64,5 @@ protected:
 	/** Analytic delegate to track node creation */
 	FNodeCreationAnalytic AnalyticCallback;
 };
+
+#undef UE_API
diff --git a/Engine/Source/Editor/Kismet/Public/ReviewComments.h b/Engine/Source/Editor/Kismet/Public/ReviewComments.h
--- a/Engine/Source/Editor/Kismet/Public/ReviewComments.h
+++ b/Engine/Source/Editor/Kismet/Public/ReviewComments.h
@@ -3,6 +3,8 @@
 #pragma once
 #include "CoreMinimal.h"
 #include "ReviewComments.generated.h"
+
+#define UE_API KISMET_API
 class FJsonObject;
 
 UENUM()
@@ -31,7 +33,7 @@ enum class EReviewTopicType : uint8
 	Unrecognised
 };
 
-struct KISMET_API FReviewCommentContext
+struct FReviewCommentContext
 {
 public:
 	// file being commented on
@@ -51,21 +53,21 @@ public:
 	// category this comment will appear under in the diff results
 	TOptional<FString> Category;
 	
-	FString ToJson() const;
-	TSharedPtr<FJsonObject> ToJsonObject() const;
-	static FReviewCommentContext FromJson(const FString& Json);
-	static FReviewCommentContext FromJson(const TSharedPtr<FJsonObject>& JsonObject);
+	UE_API FString ToJson() const;
+	UE_API TSharedPtr<FJsonObject> ToJsonObject() const;
+	static UE_API FReviewCommentContext FromJson(const FString& Json);
+	static UE_API FReviewCommentContext FromJson(const TSharedPtr<FJsonObject>& JsonObject);
 };
 
-struct KISMET_API FReviewTopic
+struct FReviewTopic
 {
 public:
 	FString ChangelistNum;
 	EReviewTopicType TopicType;
 
 	// returns topic in format "TopicType/ChangelistNum"
-	FString ToString() const;
-	static FReviewTopic FromString(const FString& ReviewTopic);
+	UE_API FString ToString() const;
+	static UE_API FReviewTopic FromString(const FString& ReviewTopic);
 	bool operator==(const FReviewTopic& Other) const
 	{
 		return ChangelistNum == Other.ChangelistNum;
@@ -77,7 +79,7 @@ inline uint32 GetTypeHash(const FReviewTopic& Topic)
 	return GetTypeHash(Topic.ChangelistNum);
 }
 
-struct KISMET_API FReviewComment
+struct FReviewComment
 {
 public:	
 	// info about CL being commented on
@@ -108,13 +110,13 @@ public:
 	// flags
 	bool bIsClosed : 1 = false;
 
-	FString ToJson() const;
-	TSharedPtr<FJsonObject> ToJsonObject() const;
-	static FReviewComment FromJson(const FString& Json);
-	static FReviewComment FromJson(const TSharedPtr<FJsonObject>& JsonObject);
+	UE_API FString ToJson() const;
+	UE_API TSharedPtr<FJsonObject> ToJsonObject() const;
+	static UE_API FReviewComment FromJson(const FString& Json);
+	static UE_API FReviewComment FromJson(const TSharedPtr<FJsonObject>& JsonObject);
 };
 
-class KISMET_API IReviewCommentAPI : public TSharedFromThis<IReviewCommentAPI>
+class IReviewCommentAPI : public TSharedFromThis<IReviewCommentAPI>
 {
 public:
 	DECLARE_DELEGATE_TwoParams(OnGetCommentsComplete, const TArray<FReviewComment>&, const FString& /*ErrorMessage*/)
@@ -138,3 +140,5 @@ public:
 	// retrieves the review associated with a CL. If there is none, it creates a new review and returns that.
 	virtual void GetOrCreateReviewTopicForCL(const FString &ChangelistNum, const OnGetReviewTopicForCLComplete& OnComplete) const = 0;
 };
+
+#undef UE_API
diff --git a/Engine/Source/Editor/Kismet/Public/SBlueprintDiff.h b/Engine/Source/Editor/Kismet/Public/SBlueprintDiff.h
--- a/Engine/Source/Editor/Kismet/Public/SBlueprintDiff.h
+++ b/Engine/Source/Editor/Kismet/Public/SBlueprintDiff.h
@@ -28,6 +28,8 @@
 #include "Widgets/Views/STableViewBase.h"
 #include "Widgets/Views/STreeView.h"
 
+#define UE_API KISMET_API
+
 class FAsyncDetailViewDiff;
 class SLinkableScrollBar;
 class FBlueprintDifferenceTreeEntry;
@@ -73,37 +75,37 @@ namespace DiffWidgetUtils
 }
 
 /** Panel used to display the blueprint */
-struct KISMET_API FDiffPanel
+struct FDiffPanel
 {
-	FDiffPanel();
+	UE_API FDiffPanel();
 
 	/** Initializes the panel, can be moved into constructor if diff and merge clients are made more uniform: */
-	void InitializeDiffPanel();
+	UE_API void InitializeDiffPanel();
 
 	/** Generate a panel for NewGraph diffed against OldGraph */
-	void GeneratePanel(UEdGraph* NewGraph, UEdGraph* OldGraph);
+	UE_API void GeneratePanel(UEdGraph* NewGraph, UEdGraph* OldGraph);
 	
 	/** Generate a panel that displays the Graph and reflects the items in the DiffResults */
-	void GeneratePanel(UEdGraph* Graph, TSharedPtr<TArray<FDiffSingleResult>> DiffResults, TAttribute<int32> FocusedDiffResult);
+	UE_API void GeneratePanel(UEdGraph* Graph, TSharedPtr<TArray<FDiffSingleResult>> DiffResults, TAttribute<int32> FocusedDiffResult);
 
 	/** Generate the 'MyBlueprint' widget, which is private to this module */
-	TSharedRef<class SWidget> GenerateMyBlueprintWidget();
+	UE_API TSharedRef<class SWidget> GenerateMyBlueprintWidget();
 
 	/** Called when user hits keyboard shortcut to copy nodes */
-	void CopySelectedNodes();
+	UE_API void CopySelectedNodes();
 
 	/** Gets whatever nodes are selected in the Graph Editor */
-	FGraphPanelSelectionSet GetSelectedNodes() const;
+	UE_API FGraphPanelSelectionSet GetSelectedNodes() const;
 
 	/** Can user copy any of the selected nodes? */
-	bool CanCopyNodes() const;
+	UE_API bool CanCopyNodes() const;
 
 	/** Functions used to focus/find a particular change in a diff result */
-	void FocusDiff(UEdGraphPin& Pin);
-	void FocusDiff(UEdGraphNode& Node);
+	UE_API void FocusDiff(UEdGraphPin& Pin);
+	UE_API void FocusDiff(UEdGraphNode& Node);
 
-	TSharedRef<SWidget> GetMyBlueprintWidget() const;
-	TSharedRef<SWidget> GetDetailsWidget() const;
+	UE_API TSharedRef<SWidget> GetMyBlueprintWidget() const;
+	UE_API TSharedRef<SWidget> GetDetailsWidget() const;
 
 	/** The blueprint that owns the graph we are showing */
 	const UBlueprint*				Blueprint;
@@ -136,7 +138,7 @@ private:
 };
 
 /* Visual Diff between two Blueprints*/
-class  KISMET_API SBlueprintDiff: public SCompoundWidget
+class  SBlueprintDiff: public SCompoundWidget
 {
 public:
 	DECLARE_DELEGATE_TwoParams( FOpenInDefaults, const class UBlueprint* , const class UBlueprint* );
@@ -150,87 +152,87 @@ public:
 			SLATE_ARGUMENT(TSharedPtr<SWindow>, ParentWindow)
 	SLATE_END_ARGS()
 
-	void Construct(const FArguments& InArgs);
-	virtual ~SBlueprintDiff();
+	UE_API void Construct(const FArguments& InArgs);
+	UE_API virtual ~SBlueprintDiff();
 
-	virtual void Tick(const FGeometry& AllottedGeometry, const double InCurrentTime, const float InDeltaTime) override;
+	UE_API virtual void Tick(const FGeometry& AllottedGeometry, const double InCurrentTime, const float InDeltaTime) override;
 
 	/** Called when a new Graph is clicked on by user */
-	void OnGraphChanged(FGraphToDiff* Diff);
+	UE_API void OnGraphChanged(FGraphToDiff* Diff);
 
 	/** Called when blueprint is modified */
-	void OnBlueprintChanged(UBlueprint* InBlueprint);
+	UE_API void OnBlueprintChanged(UBlueprint* InBlueprint);
 
 	/** Called when user clicks on a new graph list item */
-	void OnGraphSelectionChanged(TSharedPtr<FGraphToDiff> Item, ESelectInfo::Type SelectionType);
+	UE_API void OnGraphSelectionChanged(TSharedPtr<FGraphToDiff> Item, ESelectInfo::Type SelectionType);
 
 	/** Called when user clicks on an entry in the listview of differences */
-	void OnDiffListSelectionChanged(TSharedPtr<struct FDiffResultItem> TheDiff);
+	UE_API void OnDiffListSelectionChanged(TSharedPtr<struct FDiffResultItem> TheDiff);
 
 	/** Helper function for generating an empty widget */
-	static TSharedRef<SWidget> DefaultEmptyPanel();
+	static UE_API TSharedRef<SWidget> DefaultEmptyPanel();
 
 	/** Helper function to create a window that holds a diff widget */
-	static TSharedPtr<SWindow> CreateDiffWindow(FText WindowTitle, const UBlueprint* OldBlueprint, const UBlueprint* NewBlueprint, const struct FRevisionInfo& OldRevision, const struct FRevisionInfo& NewRevision);
+	static UE_API TSharedPtr<SWindow> CreateDiffWindow(FText WindowTitle, const UBlueprint* OldBlueprint, const UBlueprint* NewBlueprint, const struct FRevisionInfo& OldRevision, const struct FRevisionInfo& NewRevision);
 
 	/** Helper function to create a window that holds a diff widget that defaults a window title*/
-	static TSharedPtr<SWindow> CreateDiffWindow(const UBlueprint* OldBlueprint, const UBlueprint* NewBlueprint, const FRevisionInfo& OldRevision, const FRevisionInfo& NewRevision, const UClass* BlueprintClass);
+	static UE_API TSharedPtr<SWindow> CreateDiffWindow(const UBlueprint* OldBlueprint, const UBlueprint* NewBlueprint, const FRevisionInfo& OldRevision, const FRevisionInfo& NewRevision, const UClass* BlueprintClass);
 
 protected:
 	/** Called when user clicks button to go to next difference */
-	void NextDiff();
+	UE_API void NextDiff();
 
 	/** Called when user clicks button to go to prev difference */
-	void PrevDiff();
+	UE_API void PrevDiff();
 
 	/** Called to determine whether we have a list of differences to cycle through */
-	bool HasNextDiff() const;
-	bool HasPrevDiff() const;
+	UE_API bool HasNextDiff() const;
+	UE_API bool HasPrevDiff() const;
 
 	/** Find the FGraphToDiff that displays the graph with GraphPath relative path */
-	FGraphToDiff* FindGraphToDiffEntry(const FString& GraphPath);
+	UE_API FGraphToDiff* FindGraphToDiffEntry(const FString& GraphPath);
 
 	/** Bring these revisions of graph into focus on main display*/
-	void FocusOnGraphRevisions(FGraphToDiff* Diff);
+	UE_API void FocusOnGraphRevisions(FGraphToDiff* Diff);
 
 	/** Create a list item entry graph that exists in at least one of the blueprints */
-	void CreateGraphEntry(class UEdGraph* GraphOld, class UEdGraph* GraphNew);
+	UE_API void CreateGraphEntry(class UEdGraph* GraphOld, class UEdGraph* GraphNew);
 		
 	/** Disable the focus on a particular pin */
-	void DisablePinDiffFocus();
+	UE_API void DisablePinDiffFocus();
 
 	/** User toggles the option to lock the views between the two blueprints */
-	void OnToggleLockView();
+	UE_API void OnToggleLockView();
 	
 	/** User toggles the option to change the split view mode betwwen vertical and horizontal */
-	void OnToggleSplitViewMode();
+	UE_API void OnToggleSplitViewMode();
 
 	/** Reset the graph editor, called when user switches graphs to display*/
-	void ResetGraphEditors();
+	UE_API void ResetGraphEditors();
 
 	/** Get the image to show for the toggle lock option*/
-	FSlateIcon GetLockViewImage() const;
+	UE_API FSlateIcon GetLockViewImage() const;
 	
 	/** Get the image to show for the toggle split view mode option*/
-	FSlateIcon GetSplitViewModeImage() const;
+	UE_API FSlateIcon GetSplitViewModeImage() const;
 
 	/** List of graphs to diff, are added to panel last */
 	TArray<TSharedPtr<FGraphToDiff>> Graphs;
 
 	/** Get Graph editor associated with this Graph */
-	FDiffPanel& GetDiffPanelForNode(UEdGraphNode& Node);
+	UE_API FDiffPanel& GetDiffPanelForNode(UEdGraphNode& Node);
 
 	/** Event handler that updates the graph view when user selects a new graph */
-	void HandleGraphChanged(FGraphToDiff* Diff);
+	UE_API void HandleGraphChanged(FGraphToDiff* Diff);
 	
 	/** Function used to generate the list of differences and the widgets needed to calculate that list */
-	void GenerateDifferencesList();
+	UE_API void GenerateDifferencesList();
 
 	/** Checks if a graph is valid for diffing */
-	bool IsGraphDiffNeeded(class UEdGraph* InGraph) const;
+	UE_API bool IsGraphDiffNeeded(class UEdGraph* InGraph) const;
 
 	/** Called when editor may need to be closed */
-	void OnCloseAssetEditor(UObject* Asset, EAssetEditorCloseReason CloseReason);
+	UE_API void OnCloseAssetEditor(UObject* Asset, EAssetEditorCloseReason CloseReason);
 
 	struct FDiffControl
 	{
@@ -244,23 +246,23 @@ protected:
 		TSharedPtr< class IDiffControl > DiffControl;
 	};
 
-	FDiffControl GenerateBlueprintTypePanel();
-	FDiffControl GenerateMyBlueprintPanel();
-	FDiffControl GenerateGraphPanel();
-	FDiffControl GenerateDefaultsPanel();
-	FDiffControl GenerateClassSettingsPanel();
-	FDiffControl GenerateComponentsPanel();
-	FDiffControl GenerateGeneralFileCommentEntries();
+	UE_API FDiffControl GenerateBlueprintTypePanel();
+	UE_API FDiffControl GenerateMyBlueprintPanel();
+	UE_API FDiffControl GenerateGraphPanel();
+	UE_API FDiffControl GenerateDefaultsPanel();
+	UE_API FDiffControl GenerateClassSettingsPanel();
+	UE_API FDiffControl GenerateComponentsPanel();
+	UE_API FDiffControl GenerateGeneralFileCommentEntries();
 
-	TSharedRef<SOverlay> GenerateGraphWidgetForPanel(FDiffPanel& OutDiffPanel) const;
-	TSharedRef<SBox> GenerateRevisionInfoWidgetForPanel(TSharedPtr<SWidget>& OutGeneratedWidget,const FText& InRevisionText) const;
+	UE_API TSharedRef<SOverlay> GenerateGraphWidgetForPanel(FDiffPanel& OutDiffPanel) const;
+	UE_API TSharedRef<SBox> GenerateRevisionInfoWidgetForPanel(TSharedPtr<SWidget>& OutGeneratedWidget,const FText& InRevisionText) const;
 
 	/** Accessor and event handler for toggling between diff view modes (defaults, components, graph view, interface, macro): */
-	void SetCurrentMode(FName NewMode);
+	UE_API void SetCurrentMode(FName NewMode);
 	FName GetCurrentMode() const { return CurrentMode; }
-	void OnModeChanged(const FName& InNewViewMode) const;
+	UE_API void OnModeChanged(const FName& InNewViewMode) const;
 
-	void UpdateTopSectionVisibility(const FName& InNewViewMode) const;
+	UE_API void UpdateTopSectionVisibility(const FName& InNewViewMode) const;
 
 	FName CurrentMode;
 
@@ -321,3 +323,4 @@ protected:
 };
 
 
+#undef UE_API
diff --git a/Engine/Source/Editor/Kismet/Public/SBlueprintEditorSelectedDebugObjectWidget.h b/Engine/Source/Editor/Kismet/Public/SBlueprintEditorSelectedDebugObjectWidget.h
--- a/Engine/Source/Editor/Kismet/Public/SBlueprintEditorSelectedDebugObjectWidget.h
+++ b/Engine/Source/Editor/Kismet/Public/SBlueprintEditorSelectedDebugObjectWidget.h
@@ -24,6 +24,8 @@
 #include "Widgets/Input/SComboBox.h"
 #include "Widgets/SCompoundWidget.h"
 
+#define UE_API KISMET_API
+
 class STextComboBox;
 class SWidget;
 class UBlueprint;
@@ -70,70 +72,70 @@ struct FBlueprintDebugObjectInstance
 	}
 };
 
-class KISMET_API SBlueprintEditorSelectedDebugObjectWidget : public SCompoundWidget
+class SBlueprintEditorSelectedDebugObjectWidget : public SCompoundWidget
 {
 public:
 	SLATE_BEGIN_ARGS(SBlueprintEditorSelectedDebugObjectWidget){}
 	SLATE_END_ARGS()
 
-	void Construct(const FArguments& InArgs, TSharedPtr<FBlueprintEditor> InBlueprintEditor);
+	UE_API void Construct(const FArguments& InArgs, TSharedPtr<FBlueprintEditor> InBlueprintEditor);
 
 	// SWidget interface
-	virtual void Tick( const FGeometry& AllottedGeometry, const double InCurrentTime, const float InDeltaTime ) override;
+	UE_API virtual void Tick( const FGeometry& AllottedGeometry, const double InCurrentTime, const float InDeltaTime ) override;
 	// End of SWidget interface
 
 	/** Adds an object to the list of debug choices */
-	void AddDebugObject(UObject* TestObject, const FString& TestObjectName = FString());
+	UE_API void AddDebugObject(UObject* TestObject, const FString& TestObjectName = FString());
 
 private:
 	UBlueprint* GetBlueprintObj() const { return BlueprintEditor.Pin()->GetBlueprintObj(); }
 
 	/** Creates a list of all debug objects **/
-	void GenerateDebugObjectInstances(bool bRestoreSelection);
+	UE_API void GenerateDebugObjectInstances(bool bRestoreSelection);
 
 	/** Generate list of active PIE worlds to debug **/
-	void GenerateDebugWorldNames(bool bRestoreSelection);
+	UE_API void GenerateDebugWorldNames(bool bRestoreSelection);
 
 	/** Refresh the widget. **/
-	void OnRefresh();
+	UE_API void OnRefresh();
 
 	/** Returns the entry for the current debug actor */
-	TSharedPtr<FBlueprintDebugObjectInstance> GetDebugObjectInstance() const;
+	UE_API TSharedPtr<FBlueprintDebugObjectInstance> GetDebugObjectInstance() const;
 
 	/** Returns the name of the current debug actor */
-	TSharedPtr<FString> GetDebugWorldName() const;
+	UE_API TSharedPtr<FString> GetDebugWorldName() const;
 
 	/** Handles the selection changed event for the debug actor combo box */
-	void DebugObjectSelectionChanged(TSharedPtr<FBlueprintDebugObjectInstance> NewSelection, ESelectInfo::Type SelectInfo);
+	UE_API void DebugObjectSelectionChanged(TSharedPtr<FBlueprintDebugObjectInstance> NewSelection, ESelectInfo::Type SelectInfo);
 
-	void DebugWorldSelectionChanged(TSharedPtr<FString> NewSelection, ESelectInfo::Type SelectInfo);
+	UE_API void DebugWorldSelectionChanged(TSharedPtr<FString> NewSelection, ESelectInfo::Type SelectInfo);
 
 	/** Called when user clicks button to select the current object being debugged */
-	void SelectedDebugObject_OnClicked();
+	UE_API void SelectedDebugObject_OnClicked();
 
 	/** Returns true if a debug actor is currently selected */
-	bool IsDebugObjectSelected() const;
+	UE_API bool IsDebugObjectSelected() const;
 
-	EVisibility IsDebugWorldComboVisible() const;
+	UE_API EVisibility IsDebugWorldComboVisible() const;
 
 	/* Returns the string to indicate no debug object is selected */
-	const FString& GetNoDebugString() const;
+	UE_API const FString& GetNoDebugString() const;
 
-	const FString& GetDebugAllWorldsString() const;
+	UE_API const FString& GetDebugAllWorldsString() const;
 
-	TSharedRef<SWidget> OnGetActiveDetailSlotContent(bool bChangedToHighDetail);
+	UE_API TSharedRef<SWidget> OnGetActiveDetailSlotContent(bool bChangedToHighDetail);
 
 	/** Helper method to construct a debug object label string */
-	FString MakeDebugObjectLabel(UObject* TestObject, bool bAddContextIfSelectedInEditor, bool bAddSpawnedContext) const;
+	UE_API FString MakeDebugObjectLabel(UObject* TestObject, bool bAddContextIfSelectedInEditor, bool bAddSpawnedContext) const;
 
 	/** Fills in data for a specific instance */
-	void FillDebugObjectInstance(TSharedPtr<FBlueprintDebugObjectInstance> Instance);
+	UE_API void FillDebugObjectInstance(TSharedPtr<FBlueprintDebugObjectInstance> Instance);
 
 	/** Called to create a widget for each debug object item */
-	TSharedRef<SWidget> CreateDebugObjectItemWidget(TSharedPtr<FBlueprintDebugObjectInstance> InItem);
+	UE_API TSharedRef<SWidget> CreateDebugObjectItemWidget(TSharedPtr<FBlueprintDebugObjectInstance> InItem);
 
 	/** Returns the combo button label to use for the currently-selected debug object item */
-	FText GetSelectedDebugObjectTextLabel() const;
+	UE_API FText GetSelectedDebugObjectTextLabel() const;
 
 private:
 	/** Pointer back to the blueprint editor tool that owns us */
@@ -153,3 +155,5 @@ private:
 
 	TWeakObjectPtr<UObject> LastObjectObserved;
 };
+
+#undef UE_API
diff --git a/Engine/Source/Editor/Kismet/Public/SBlueprintEditorToolbar.h b/Engine/Source/Editor/Kismet/Public/SBlueprintEditorToolbar.h
--- a/Engine/Source/Editor/Kismet/Public/SBlueprintEditorToolbar.h
+++ b/Engine/Source/Editor/Kismet/Public/SBlueprintEditorToolbar.h
@@ -16,6 +16,8 @@
 #include "Widgets/SWidget.h"
 #include "WorkflowOrientedApp/SModeWidget.h"
 
+#define UE_API KISMET_API
+
 class FBlueprintEditor;
 class FExtender;
 class FMenuBuilder;
@@ -29,21 +31,21 @@ struct FToolMenuContext;
 /**
  * Kismet menu
  */
-class KISMET_API FKismet2Menu
+class FKismet2Menu
 {
 public:
-	static void SetupBlueprintEditorMenu(const FName MainMenuName);
+	static UE_API void SetupBlueprintEditorMenu(const FName MainMenuName);
 	
 protected:
-	static void FillFileMenuBlueprintSection(UToolMenu* Menu);
+	static UE_API void FillFileMenuBlueprintSection(UToolMenu* Menu);
 
-	static void FillEditMenu(UToolMenu* Menu);
+	static UE_API void FillEditMenu(UToolMenu* Menu);
 
-	static void FillViewMenu(UToolMenu* Menu);
+	static UE_API void FillViewMenu(UToolMenu* Menu);
 
-	static void FillDebugMenu(UToolMenu* Menu);
+	static UE_API void FillDebugMenu(UToolMenu* Menu);
 
-	static void FillDeveloperMenu(UToolMenu* Menu);
+	static UE_API void FillDeveloperMenu(UToolMenu* Menu);
 };
 
 
@@ -78,31 +80,32 @@ public:
 
 
 
-class KISMET_API FBlueprintEditorToolbar : public TSharedFromThis<FBlueprintEditorToolbar>
+class FBlueprintEditorToolbar : public TSharedFromThis<FBlueprintEditorToolbar>
 {
 public:
 	FBlueprintEditorToolbar(TSharedPtr<FBlueprintEditor> InBlueprintEditor)
 		: BlueprintEditor(InBlueprintEditor) {}
 
-	void AddBlueprintGlobalOptionsToolbar(UToolMenu* InMenu, bool bRegisterViewport = false);
-	void AddCompileToolbar(UToolMenu* InMenu);
-	void AddNewToolbar(UToolMenu* InMenu);
-	void AddScriptingToolbar(UToolMenu* InMenu);
-	void AddDebuggingToolbar(UToolMenu* InMenu);
+	UE_API void AddBlueprintGlobalOptionsToolbar(UToolMenu* InMenu, bool bRegisterViewport = false);
+	UE_API void AddCompileToolbar(UToolMenu* InMenu);
+	UE_API void AddNewToolbar(UToolMenu* InMenu);
+	UE_API void AddScriptingToolbar(UToolMenu* InMenu);
+	UE_API void AddDebuggingToolbar(UToolMenu* InMenu);
 
 	/** Returns the current status icon for the blueprint being edited */
-	FSlateIcon GetStatusImage() const;
+	UE_API FSlateIcon GetStatusImage() const;
 
 	/** Returns the current status as text for the blueprint being edited */
-	FText GetStatusTooltip() const;
+	UE_API FText GetStatusTooltip() const;
 
 	/** Diff current blueprint against the specified revision */
-	static void DiffAgainstRevision(class UBlueprint* Current, int32 OldRevision);
+	static UE_API void DiffAgainstRevision(class UBlueprint* Current, int32 OldRevision);
 
-	static TSharedRef<SWidget> MakeDiffMenu(const UBlueprintEditorToolMenuContext* InContext);
+	static UE_API TSharedRef<SWidget> MakeDiffMenu(const UBlueprintEditorToolMenuContext* InContext);
 
 protected:
 	/** Pointer back to the blueprint editor tool that owns us */
 	TWeakPtr<FBlueprintEditor> BlueprintEditor;
 };
 
+#undef UE_API
diff --git a/Engine/Source/Editor/Kismet/Public/SBlueprintNamespaceEntry.h b/Engine/Source/Editor/Kismet/Public/SBlueprintNamespaceEntry.h
--- a/Engine/Source/Editor/Kismet/Public/SBlueprintNamespaceEntry.h
+++ b/Engine/Source/Editor/Kismet/Public/SBlueprintNamespaceEntry.h
@@ -28,6 +28,8 @@
 #include "Widgets/Views/STableViewBase.h"
 #include "Widgets/Views/STreeView.h"
 
+#define UE_API KISMET_API
+
 class ITableRow;
 class SComboButton;
 class SSearchBox;
@@ -35,7 +37,7 @@ class SSuggestionTextBox;
 class SWidget;
 
 /** A custom widget class that provides support for Blueprint namespace entry and/or selection. */
-class KISMET_API SBlueprintNamespaceEntry : public SCompoundWidget
+class SBlueprintNamespaceEntry : public SCompoundWidget
 {
 public:
 	DECLARE_DELEGATE_OneParam(FOnNamespaceSelected, const FString&);
@@ -80,14 +82,14 @@ PRAGMA_ENABLE_DEPRECATION_WARNINGS
 
 	SLATE_END_ARGS()
 
-	void Construct(const FArguments& InArgs);
+	UE_API void Construct(const FArguments& InArgs);
 
 	/**
 	 * Set the current namespace to the given identifier.
 	 * 
 	 * @param InNamespace	New namespace identifier. May be an empty string.
 	 */
-	void SetCurrentNamespace(const FString& InNamespace);
+	UE_API void SetCurrentNamespace(const FString& InNamespace);
 
 protected:
 	struct FPathTreeNodeItem
@@ -97,22 +99,22 @@ protected:
 		bool bIsSelectable = true;
 	};
 	
-	void OnTextChanged(const FText& InText);
-	void OnTextCommitted(const FText& NewText, ETextCommit::Type InTextCommit);
-	void OnShowingSuggestions(const FString& InputText, TArray<FString>& OutSuggestions);
-	TSharedRef<SWidget> OnGetNamespaceTreeMenuContent();
-	TSharedRef<ITableRow> OnGenerateRowForNamespaceTreeItem(TSharedPtr<FPathTreeNodeItem> Item, const TSharedRef<STableViewBase>& OwnerTable);
-	void OnGetChildrenForNamespaceTreeItem(TSharedPtr<FPathTreeNodeItem> Item, TArray<TSharedPtr<FPathTreeNodeItem>>& OutChildren);
-	void OnNamespaceTreeFilterTextChanged(const FText& InText);
-	void OnNamespaceTreeSelectionChanged(TSharedPtr<FPathTreeNodeItem> Item, ESelectInfo::Type SelectInfo);
-	bool OnIsNamespaceTreeItemSelectable(TSharedPtr<FPathTreeNodeItem> Item) const;
-	FText GetCurrentNamespaceText() const;
-
-	void PopulateNamespaceTree();
-	void SelectNamespace(const FString& InNamespace);
-	void ExpandAllTreeViewItems(const TArray<TSharedPtr<FPathTreeNodeItem>>* NodeListPtr = nullptr);
-	const TSharedPtr<FPathTreeNodeItem>* FindTreeViewNode(const FString& NodePath, const TArray<TSharedPtr<FPathTreeNodeItem>>* NodeListPtr = nullptr) const;
-PRAGMA_DISABLE_DEPRECATION_WARNINGS
+	UE_API void OnTextChanged(const FText& InText);
+	UE_API void OnTextCommitted(const FText& NewText, ETextCommit::Type InTextCommit);
+	UE_API void OnShowingSuggestions(const FString& InputText, TArray<FString>& OutSuggestions);
+	UE_API TSharedRef<SWidget> OnGetNamespaceTreeMenuContent();
+	UE_API TSharedRef<ITableRow> OnGenerateRowForNamespaceTreeItem(TSharedPtr<FPathTreeNodeItem> Item, const TSharedRef<STableViewBase>& OwnerTable);
+	UE_API void OnGetChildrenForNamespaceTreeItem(TSharedPtr<FPathTreeNodeItem> Item, TArray<TSharedPtr<FPathTreeNodeItem>>& OutChildren);
+	UE_API void OnNamespaceTreeFilterTextChanged(const FText& InText);
+	UE_API void OnNamespaceTreeSelectionChanged(TSharedPtr<FPathTreeNodeItem> Item, ESelectInfo::Type SelectInfo);
+	UE_API bool OnIsNamespaceTreeItemSelectable(TSharedPtr<FPathTreeNodeItem> Item) const;
+	UE_API FText GetCurrentNamespaceText() const;
+
+	UE_API void PopulateNamespaceTree();
+	UE_API void SelectNamespace(const FString& InNamespace);
+	UE_API void ExpandAllTreeViewItems(const TArray<TSharedPtr<FPathTreeNodeItem>>* NodeListPtr = nullptr);
+	UE_API const TSharedPtr<FPathTreeNodeItem>* FindTreeViewNode(const FString& NodePath, const TArray<TSharedPtr<FPathTreeNodeItem>>* NodeListPtr = nullptr) const;
+UE_API PRAGMA_DISABLE_DEPRECATION_WARNINGS
 	static void HandleLegacyOnFilterNamespaceList(TSet<FString>& OutNamespacesToExclude, FOnFilterNamespaceList LegacyDelegate);
 PRAGMA_ENABLE_DEPRECATION_WARNINGS
 
@@ -131,6 +133,8 @@ private:
 	FOnGetNamespacesToExclude OnGetNamespacesToExclude;
 	TAttribute<FText> ExcludedNamespaceTooltipText;
 
-	static float NamespaceListBorderPadding;
-	static float NamespaceListMinDesiredWidth;
-};
\ No newline at end of file
+	static UE_API float NamespaceListBorderPadding;
+	static UE_API float NamespaceListMinDesiredWidth;
+};
+
+#undef UE_API
diff --git a/Engine/Source/Editor/Kismet/Public/SBlueprintRevisionMenu.h b/Engine/Source/Editor/Kismet/Public/SBlueprintRevisionMenu.h
--- a/Engine/Source/Editor/Kismet/Public/SBlueprintRevisionMenu.h
+++ b/Engine/Source/Editor/Kismet/Public/SBlueprintRevisionMenu.h
@@ -15,12 +15,14 @@
 #include "Widgets/DeclarativeSyntaxSupport.h"
 #include "Widgets/SCompoundWidget.h"
 
+#define UE_API KISMET_API
+
 class FUpdateStatus;
 class SVerticalBox;
 class UBlueprint;
 struct FRevisionInfo;
 
-class KISMET_API SBlueprintRevisionMenu : public SCompoundWidget
+class SBlueprintRevisionMenu : public SCompoundWidget
 {
 	DECLARE_DELEGATE_OneParam(FOnRevisionSelected, FRevisionInfo const&)
 
@@ -32,20 +34,20 @@ public:
 		SLATE_EVENT(FOnRevisionSelected, OnRevisionSelected)
 	SLATE_END_ARGS()
 
-	~SBlueprintRevisionMenu();
+	UE_API ~SBlueprintRevisionMenu();
 
-	void Construct(const FArguments& InArgs, UBlueprint const* Blueprint);
+	UE_API void Construct(const FArguments& InArgs, UBlueprint const* Blueprint);
 
 private: 
 	/** Delegate used to determine the visibility 'in progress' widgets */
-	EVisibility GetInProgressVisibility() const;
+	UE_API EVisibility GetInProgressVisibility() const;
 	/** Delegate used to determine the visibility of the cancel button */
-	EVisibility GetCancelButtonVisibility() const;
+	UE_API EVisibility GetCancelButtonVisibility() const;
 
 	/** Delegate used to cancel a source control operation in progress */
-	FReply OnCancelButtonClicked() const;
+	UE_API FReply OnCancelButtonClicked() const;
 	/** Callback for when the source control operation is complete */
-	void OnSourceControlQueryComplete(const FSourceControlOperationRef& InOperation, ECommandResult::Type InResult);
+	UE_API void OnSourceControlQueryComplete(const FSourceControlOperationRef& InOperation, ECommandResult::Type InResult);
 
 	/**  */
 	bool bIncludeLocalRevision;
@@ -60,3 +62,5 @@ private:
 	/** The state of the SCC query */
 	uint32 SourceControlQueryState;
 };
+
+#undef UE_API
diff --git a/Engine/Source/Editor/Kismet/Public/SCSDiff.h b/Engine/Source/Editor/Kismet/Public/SCSDiff.h
--- a/Engine/Source/Editor/Kismet/Public/SCSDiff.h
+++ b/Engine/Source/Editor/Kismet/Public/SCSDiff.h
@@ -8,6 +8,8 @@
 #include "UObject/NameTypes.h"
 #include "Widgets/SWidget.h"
 
+#define UE_API KISMET_API
+
 class FSCSEditorTreeNode;
 class FSubobjectEditorTreeNode;
 class SKismetInspector;
@@ -17,21 +19,21 @@ class SWidget;
 class UBlueprint;
 
 /** Struct to support diffing the component tree for a blueprint */
-class KISMET_API FSCSDiff
+class FSCSDiff
 {
 public:
-	FSCSDiff(const class UBlueprint* InBlueprint);
+	UE_API FSCSDiff(const class UBlueprint* InBlueprint);
 
-	void HighlightProperty(FName VarName, FPropertySoftPath Property);
-	TSharedRef< SWidget > TreeWidget();
+	UE_API void HighlightProperty(FName VarName, FPropertySoftPath Property);
+	UE_API TSharedRef< SWidget > TreeWidget();
 
-	TArray< FSCSResolvedIdentifier > GetDisplayedHierarchy() const;
+	UE_API TArray< FSCSResolvedIdentifier > GetDisplayedHierarchy() const;
 
 	const UBlueprint* GetBlueprint() const { return Blueprint; }
 
 protected:
-	void OnSCSEditorUpdateSelectionFromNodes(const TArray<TSharedPtr<FSubobjectEditorTreeNode>>& SelectedNodes);
-	void OnSCSEditorHighlightPropertyInDetailsView(const class FPropertyPath& InPropertyPath);
+	UE_API void OnSCSEditorUpdateSelectionFromNodes(const TArray<TSharedPtr<FSubobjectEditorTreeNode>>& SelectedNodes);
+	UE_API void OnSCSEditorHighlightPropertyInDetailsView(const class FPropertyPath& InPropertyPath);
 
 private:
 	TSharedPtr< class SWidget > ContainerWidget;
@@ -41,3 +43,5 @@ private:
 	/** Blueprint we are inspecting */
 	UBlueprint* Blueprint = nullptr;
 };
+
+#undef UE_API
diff --git a/Engine/Source/Editor/Kismet/Public/SCSEditorExtensionContext.h b/Engine/Source/Editor/Kismet/Public/SCSEditorExtensionContext.h
--- a/Engine/Source/Editor/Kismet/Public/SCSEditorExtensionContext.h
+++ b/Engine/Source/Editor/Kismet/Public/SCSEditorExtensionContext.h
@@ -17,8 +17,8 @@ class SSubobjectEditor;
 class SSubobjectInstanceEditor;
 class UE_DEPRECATED(5.0, "USCSEditorExtensionContext has been deprecated, use USubobjectEditorExtensionContext instead.") USCSEditorExtensionContext;
 
-UCLASS()
-class KISMET_API USCSEditorExtensionContext : public UObject
+UCLASS(MinimalAPI)
+class USCSEditorExtensionContext : public UObject
 {
 	GENERATED_BODY()
 
