Updated headers using UnrealCodeFixup to make sure dllstorage is on methods/staticvars instead of types. Part 5/n

diff --git a/Engine/Source/Editor/Kismet/Public/SDetailsDiff.h b/Engine/Source/Editor/Kismet/Public/SDetailsDiff.h
--- a/Engine/Source/Editor/Kismet/Public/SDetailsDiff.h
+++ b/Engine/Source/Editor/Kismet/Public/SDetailsDiff.h
@@ -21,7 +21,7 @@ class FDetailsDiffControl;
 enum class EAssetEditorCloseReason : uint8;
 
 /** Panel used to display the details */
-struct KISMET_API FDetailsDiffPanel
+struct FDetailsDiffPanel
 {
 	/** The asset that owns the details view we are showing */
 	const UObject*				Object = nullptr;
@@ -40,7 +40,7 @@ private:
 };
 
 /* Visual Diff between two Assets */
-class  KISMET_API SDetailsDiff: public SCompoundWidget
+class  SDetailsDiff: public SCompoundWidget
 {
 public:
 
@@ -67,54 +67,54 @@ public:
 		SLATE_EVENT(FRowHighlightColor, RowHighlightColor)
 	SLATE_END_ARGS()
 
-	void Construct(const FArguments& InArgs);
-	virtual ~SDetailsDiff();
+	KISMET_API void Construct(const FArguments& InArgs);
+	KISMET_API virtual ~SDetailsDiff();
 
-	virtual void Tick(const FGeometry& AllottedGeometry, const double InCurrentTime, const float InDeltaTime) override;
+	KISMET_API virtual void Tick(const FGeometry& AllottedGeometry, const double InCurrentTime, const float InDeltaTime) override;
 
 	/** Called when user clicks on an entry in the listview of differences */
-	void OnDiffListSelectionChanged(TSharedPtr<struct FDiffResultItem> TheDiff);
+	KISMET_API void OnDiffListSelectionChanged(TSharedPtr<struct FDiffResultItem> TheDiff);
 
 	/** Helper function for generating an empty widget */
-	static TSharedRef<SWidget> DefaultEmptyPanel();
+	static KISMET_API TSharedRef<SWidget> DefaultEmptyPanel();
 
 	/** Helper function to create a window that holds a diff widget */
-	static TSharedRef<SDetailsDiff> CreateDiffWindow(FText WindowTitle, const UObject* OldObject, const UObject* NewObject, const struct FRevisionInfo& OldRevision, const struct FRevisionInfo& NewRevision);
+	static KISMET_API TSharedRef<SDetailsDiff> CreateDiffWindow(FText WindowTitle, const UObject* OldObject, const UObject* NewObject, const struct FRevisionInfo& OldRevision, const struct FRevisionInfo& NewRevision);
 
 	/** Helper function to create a window that holds a diff widget (default window title) */
-	static TSharedRef<SDetailsDiff> CreateDiffWindow(const UObject* OldObject, const UObject* NewObject, const FRevisionInfo& OldRevision, const FRevisionInfo& NewRevision, const UClass* ObjectClass);
+	static KISMET_API TSharedRef<SDetailsDiff> CreateDiffWindow(const UObject* OldObject, const UObject* NewObject, const FRevisionInfo& OldRevision, const FRevisionInfo& NewRevision, const UClass* ObjectClass);
 
 	/** Enables actions like "Choose Left" and "Choose Right" which will modify the OutputObject */
-	void SetOutputObject(UObject* OutputObject);
+	KISMET_API void SetOutputObject(UObject* OutputObject);
 	
 	/** Return a serialized buffer of change requests made by the user */
-	UObject* GetOutputObject() const;
+	KISMET_API UObject* GetOutputObject() const;
 
 	/** Returns whether SetOutputObject was called with a valid object */
-	bool IsOutputEnabled() const;
+	KISMET_API bool IsOutputEnabled() const;
 
-	void ReportMergeConflicts(const TMap<FString, TMap<FPropertySoftPath, ETreeDiffResult>>& Conflicts);
+	KISMET_API void ReportMergeConflicts(const TMap<FString, TMap<FPropertySoftPath, ETreeDiffResult>>& Conflicts);
 
 	DECLARE_MULTICAST_DELEGATE_OneParam(FOnWindowClosedEvent, TSharedRef<SDetailsDiff>)
 	FOnWindowClosedEvent OnWindowClosedEvent;
 protected:
 	/** Called when user clicks button to go to next difference */
-	void NextDiff();
+	KISMET_API void NextDiff();
 
 	/** Called when user clicks button to go to prev difference */
-	void PrevDiff();
+	KISMET_API void PrevDiff();
 
 	/** Called to determine whether we have a list of differences to cycle through */
-	bool HasNextDiff() const;
-	bool HasPrevDiff() const;
+	KISMET_API bool HasNextDiff() const;
+	KISMET_API bool HasPrevDiff() const;
 	
 	/** Function used to generate the list of differences and the widgets needed to calculate that list */
-	void GenerateDifferencesList();
+	KISMET_API void GenerateDifferencesList();
 
 	/** Called when editor may need to be closed */
-	void OnCloseAssetEditor(UObject* Asset, EAssetEditorCloseReason CloseReason);
+	KISMET_API void OnCloseAssetEditor(UObject* Asset, EAssetEditorCloseReason CloseReason);
 
-	void OnObjectReplaced(const FCoreUObjectDelegates::FReplacementObjectMap& Replacements);
+	KISMET_API void OnObjectReplaced(const FCoreUObjectDelegates::FReplacementObjectMap& Replacements);
 
 	struct FDiffControl
 	{
@@ -128,17 +128,17 @@ protected:
 		TSharedPtr< class IDiffControl > DiffControl;
 	};
 
-	FDiffControl GenerateDetailsPanel(const TFunction<const UObject*(const UObject*)>& Redirector = nullptr);
+	KISMET_API FDiffControl GenerateDetailsPanel(const TFunction<const UObject*(const UObject*)>& Redirector = nullptr);
 
-	TSharedRef<SBox> GenerateRevisionInfoWidgetForPanel(TSharedPtr<SWidget>& OutGeneratedWidget,const FText& InRevisionText) const;
+	KISMET_API TSharedRef<SBox> GenerateRevisionInfoWidgetForPanel(TSharedPtr<SWidget>& OutGeneratedWidget,const FText& InRevisionText) const;
 
 	/** Accessor and event handler for toggling between diff view modes (defaults, components, graph view, interface, macro): */
-	void SetCurrentMode(FName NewMode);
-	void RefreshCurrentModePanel();
+	KISMET_API void SetCurrentMode(FName NewMode);
+	KISMET_API void RefreshCurrentModePanel();
 	FName GetCurrentMode() const { return CurrentMode; }
-	void OnModeChanged(const FName& InNewViewMode) const;
+	KISMET_API void OnModeChanged(const FName& InNewViewMode) const;
 
-	void UpdateTopSectionVisibility(const FName& InNewViewMode) const;
+	KISMET_API void UpdateTopSectionVisibility(const FName& InNewViewMode) const;
 
 	FName CurrentMode;
 
diff --git a/Engine/Source/Editor/Kismet/Public/SFilterableObjectList.h b/Engine/Source/Editor/Kismet/Public/SFilterableObjectList.h
--- a/Engine/Source/Editor/Kismet/Public/SFilterableObjectList.h
+++ b/Engine/Source/Editor/Kismet/Public/SFilterableObjectList.h
@@ -27,6 +27,8 @@
 #include "Widgets/Views/STableRow.h"
 #include "Widgets/Views/STableViewBase.h"
 
+#define UE_API KISMET_API
+
 class ITableRow;
 class SSearchBox;
 class STableViewBase;
@@ -35,20 +37,20 @@ class SWidget;
 //////////////////////////////////////////////////////////////////////////
 // SFilterableObjectList
 
-class KISMET_API SFilterableObjectList : public SCompoundWidget
+class SFilterableObjectList : public SCompoundWidget
 {
 public:
 	SLATE_BEGIN_ARGS( SFilterableObjectList ){}
 	SLATE_END_ARGS()
 public:
-	void Construct(const FArguments& InArgs);
+	UE_API void Construct(const FArguments& InArgs);
 
-	void InternalConstruct();
+	UE_API void InternalConstruct();
 
 protected:
 	// Interface for derived classes to implement
-	virtual void RebuildObjectList();
-	virtual FString GetSearchableText(UObject* Object);
+	UE_API virtual void RebuildObjectList();
+	UE_API virtual FString GetSearchableText(UObject* Object);
 	
 	/** Return value for GenerateRowForObject */
 	struct FListRow
@@ -66,24 +68,24 @@ protected:
 	};
 
 	/** Make a table row for the list of filterable widgets */
-	virtual FListRow GenerateRowForObject(UObject* Object);
+	UE_API virtual FListRow GenerateRowForObject(UObject* Object);
 	
 
 	// End of interface for derived classes to implement
 protected:
-	void RefilterObjectList();
+	UE_API void RefilterObjectList();
 
-	TSharedRef<ITableRow> OnGenerateTableRow(UObject* InData, const TSharedRef<STableViewBase>& OwnerTable);
+	UE_API TSharedRef<ITableRow> OnGenerateTableRow(UObject* InData, const TSharedRef<STableViewBase>& OwnerTable);
 
-	FReply OnRefreshButtonClicked();
+	UE_API FReply OnRefreshButtonClicked();
 
-	void OnFilterTextChanged(const FText& InFilterText);
+	UE_API void OnFilterTextChanged(const FText& InFilterText);
 
-	EVisibility GetFilterStatusVisibility() const;
-	FText GetFilterStatusText() const;
+	UE_API EVisibility GetFilterStatusVisibility() const;
+	UE_API FText GetFilterStatusText() const;
 
-	bool IsFilterActive() const;
-	void ReapplyFilter();
+	UE_API bool IsFilterActive() const;
+	UE_API void ReapplyFilter();
 protected:
 	/** Widget containing the object list */
 	TSharedPtr< SListView<UObject*> > ObjectListWidget;
@@ -97,3 +99,5 @@ protected:
 	/** List of objects to show that have passed the keyword filtering */
 	TArray<UObject*> FilteredObjectList;
 };
+
+#undef UE_API
diff --git a/Engine/Source/Editor/Kismet/Public/SKismetInspector.h b/Engine/Source/Editor/Kismet/Public/SKismetInspector.h
--- a/Engine/Source/Editor/Kismet/Public/SKismetInspector.h
+++ b/Engine/Source/Editor/Kismet/Public/SKismetInspector.h
@@ -23,6 +23,8 @@
 #include "Widgets/SCompoundWidget.h"
 #include "Widgets/SWidget.h"
 
+#define UE_API KISMET_API
+
 class FBlueprintEditor;
 class FProperty;
 class FStructOnScope;
@@ -42,7 +44,7 @@ typedef TSet<class UObject*> FInspectorSelectionSet;
 // SKismetInspector
 
 /** Widget that shows properties and tools related to the selected node(s) */
-class KISMET_API SKismetInspector : public SCompoundWidget, public FGCObject
+class SKismetInspector : public SCompoundWidget, public FGCObject
 {
 public:
 	SLATE_BEGIN_ARGS( SKismetInspector )
@@ -70,7 +72,7 @@ public:
 		SLATE_ARGUMENT( bool, ShowSectionSelector)
 	SLATE_END_ARGS()
 	
-	void Construct(const FArguments& InArgs);
+	UE_API void Construct(const FArguments& InArgs);
 
 	/** Options for ShowDetails */
 	struct FShowDetailsOptions
@@ -96,44 +98,44 @@ public:
 	};
 
 	// SWidget interface
-	virtual void Tick( const FGeometry& AllottedGeometry, const double InCurrentTime, const float InDeltaTime ) override;
+	UE_API virtual void Tick( const FGeometry& AllottedGeometry, const double InCurrentTime, const float InDeltaTime ) override;
 	// End of SWidget interface
 
 	/** Update the inspector window to show information on the supplied object */
-	void ShowDetailsForSingleObject(UObject* Object, const FShowDetailsOptions& Options = FShowDetailsOptions());
+	UE_API void ShowDetailsForSingleObject(UObject* Object, const FShowDetailsOptions& Options = FShowDetailsOptions());
 
 	/** Update the inspector window to show information on the supplied objects */
-	void ShowDetailsForObjects(const TArray<UObject*>& PropertyObjects, const FShowDetailsOptions& Options = FShowDetailsOptions());
+	UE_API void ShowDetailsForObjects(const TArray<UObject*>& PropertyObjects, const FShowDetailsOptions& Options = FShowDetailsOptions());
 
 	/** Update the inspector window to show single struct. This invalidates ShowDetailsForObjects */
-	void ShowSingleStruct(TSharedPtr<FStructOnScope> InStructToDisplay);
+	UE_API void ShowSingleStruct(TSharedPtr<FStructOnScope> InStructToDisplay);
 
 	/** Used to control visibility of a property in the property window */
-	bool IsPropertyVisible( const struct FPropertyAndParent& PropertyAndParent ) const;
+	UE_API bool IsPropertyVisible( const struct FPropertyAndParent& PropertyAndParent ) const;
 
 	/** Turns on or off details customization for components */
-	void EnableComponentDetailsCustomization(bool bEnable);
+	UE_API void EnableComponentDetailsCustomization(bool bEnable);
 
 	TSharedPtr<class IDetailsView> GetPropertyView() const { return PropertyView; }
 
-	void SetOwnerTab(TSharedRef<SDockTab> Tab);
-	TSharedPtr<SDockTab> GetOwnerTab() const;
+	UE_API void SetOwnerTab(TSharedRef<SDockTab> Tab);
+	UE_API TSharedPtr<SDockTab> GetOwnerTab() const;
 
 	/** @return true if the object is in the selection set. */
-	bool IsSelected(UObject* Object) const;
+	UE_API bool IsSelected(UObject* Object) const;
 
 	/** returns the list of selected objects */
-	const TArray< TWeakObjectPtr<UObject> >& GetSelectedObjects() const;
+	UE_API const TArray< TWeakObjectPtr<UObject> >& GetSelectedObjects() const;
 
-	virtual void AddReferencedObjects(FReferenceCollector& Collector) override;
-	virtual FString GetReferencerName() const override;
+	UE_API virtual void AddReferencedObjects(FReferenceCollector& Collector) override;
+	UE_API virtual FString GetReferencerName() const override;
 
 protected:
 	/** Update the inspector window to show information on the supplied objects */
-	void UpdateFromObjects(const TArray<UObject*>& PropertyObjects, struct FKismetSelectionInfo& SelectionInfo, const FShowDetailsOptions& Options);
+	UE_API void UpdateFromObjects(const TArray<UObject*>& PropertyObjects, struct FKismetSelectionInfo& SelectionInfo, const FShowDetailsOptions& Options);
 
 	/** Add this property and all its child properties to SelectedObjectProperties */
-	void AddPropertiesRecursive(FProperty* Property);
+	UE_API void AddPropertiesRecursive(FProperty* Property);
 
 	/** Pointer back to the kismet 2 tool that owns us */
 	TWeakPtr<FBlueprintEditor> BlueprintEditorPtr;
@@ -199,46 +201,48 @@ protected:
 	TSharedPtr<class IStructureDetailsView> StructureDetailsView;
 
 	/** Update the inspector window to show information on the single struct */
-	void UpdateFromSingleStruct(const TSharedPtr<FStructOnScope>& InStructToDisplay);
+	UE_API void UpdateFromSingleStruct(const TSharedPtr<FStructOnScope>& InStructToDisplay);
 
 	/** Is struct view property read only */
-	bool IsStructViewPropertyReadOnly(const struct FPropertyAndParent& PropertyAndParent) const;
+	UE_API bool IsStructViewPropertyReadOnly(const struct FPropertyAndParent& PropertyAndParent) const;
 
 protected:
 	/** Show properties of the selected object */
-	void SetPropertyWindowContents(TArray<UObject*> Objects);
+	UE_API void SetPropertyWindowContents(TArray<UObject*> Objects);
 
 	/** Returns whether the property view be visible */
-	EVisibility GetPropertyViewVisibility() const;
+	UE_API EVisibility GetPropertyViewVisibility() const;
 
 	/** Returns whether the properties in the view should be editable */
-	bool IsPropertyEditingEnabled() const;
+	UE_API bool IsPropertyEditingEnabled() const;
 
 	/** Returns whether the warning about an inherited component not being editable should be visible */
-	EVisibility GetInheritedBlueprintComponentWarningVisibility() const;
+	UE_API EVisibility GetInheritedBlueprintComponentWarningVisibility() const;
 
 	/** Opens the parent blueprint when the hyperlink in the warning is clicked */
-	void OnInheritedBlueprintComponentWarningHyperlinkClicked(const FSlateHyperlinkRun::FMetadata& Metadata);
+	UE_API void OnInheritedBlueprintComponentWarningHyperlinkClicked(const FSlateHyperlinkRun::FMetadata& Metadata);
 
 	/**
 	 * Generates a widget that is used to edit the specified object array contextually.  This widget
 	 * will be displayed along with a property view in the level editor
 	 */
-	TSharedRef<SWidget> MakeContextualEditingWidget(struct FKismetSelectionInfo& SelectionInfo, const FShowDetailsOptions& Options);
+	UE_API TSharedRef<SWidget> MakeContextualEditingWidget(struct FKismetSelectionInfo& SelectionInfo, const FShowDetailsOptions& Options);
 
 	/**
 	 * Generates the text for the title in the contextual editing widget
 	 */
-	FText GetContextualEditingWidgetTitle() const;
+	UE_API FText GetContextualEditingWidgetTitle() const;
 
-	ECheckBoxState GetPublicViewCheckboxState() const;
-	void SetPublicViewCheckboxState(ECheckBoxState InIsChecked);
+	UE_API ECheckBoxState GetPublicViewCheckboxState() const;
+	UE_API void SetPublicViewCheckboxState(ECheckBoxState InIsChecked);
 
-	bool IsAnyParentOrContainerSelected(const FPropertyAndParent& PropertyAndParent) const;
+	UE_API bool IsAnyParentOrContainerSelected(const FPropertyAndParent& PropertyAndParent) const;
 
 	/** Callback invoked after a value change on the selected object(s) */
-	void OnFinishedChangingProperties(const FPropertyChangedEvent& InPropertyChangedEvent);
+	UE_API void OnFinishedChangingProperties(const FPropertyChangedEvent& InPropertyChangedEvent);
 
 	/** Auto-import any namespaces associated with a property's value into the current editor context */
-	void ImportNamespacesForPropertyValue(const FProperty* InProperty) const;
+	UE_API void ImportNamespacesForPropertyValue(const FProperty* InProperty) const;
 };
+
+#undef UE_API
diff --git a/Engine/Source/Editor/Kismet/Public/SMyBlueprint.h b/Engine/Source/Editor/Kismet/Public/SMyBlueprint.h
--- a/Engine/Source/Editor/Kismet/Public/SMyBlueprint.h
+++ b/Engine/Source/Editor/Kismet/Public/SMyBlueprint.h
@@ -27,6 +27,8 @@
 #include "Widgets/SWidget.h"
 #include "WorkflowOrientedApp/WorkflowTabManager.h"
 
+#define UE_API KISMET_API
+
 class FBlueprintEditor;
 class FMenuBuilder;
 class FUICommandInfo;
@@ -82,59 +84,59 @@ public:
 	virtual void RegisterCommands() override;
 };
 
-class KISMET_API SMyBlueprint : public SCompoundWidget
+class SMyBlueprint : public SCompoundWidget
 {
 public:
 	SLATE_BEGIN_ARGS( SMyBlueprint ) {}
 	SLATE_END_ARGS()
 
-	void Construct(const FArguments& InArgs, TWeakPtr<FBlueprintEditor> InBlueprintEditor, const UBlueprint* InBlueprint = nullptr);
-	~SMyBlueprint();
+	UE_API void Construct(const FArguments& InArgs, TWeakPtr<FBlueprintEditor> InBlueprintEditor, const UBlueprint* InBlueprint = nullptr);
+	UE_API ~SMyBlueprint();
 
 	void SetInspector( TSharedPtr<SKismetInspector> InInspector ) { Inspector = InInspector ; }
 
 	/* SWidget interface */
-	virtual void Tick(const FGeometry& AllottedGeometry, const double InCurrentTime, const float InDeltaTime);
-	virtual FReply OnKeyDown(const FGeometry& MyGeometry, const FKeyEvent& InKeyEvent);
+	UE_API virtual void Tick(const FGeometry& AllottedGeometry, const double InCurrentTime, const float InDeltaTime);
+	UE_API virtual FReply OnKeyDown(const FGeometry& MyGeometry, const FKeyEvent& InKeyEvent);
 
 	/* Reset the last pin type settings to default. */
-	void ResetLastPinType();
+	UE_API void ResetLastPinType();
 
 	/** Refreshes the graph action menu */
-	void Refresh();
+	UE_API void Refresh();
 	void SetFocusedGraph(UEdGraph* InEdGraph) { EdGraph = InEdGraph; }
 	
 	/** Accessor for getting the current selection as a K2 graph */
-	FEdGraphSchemaAction_K2Graph* SelectionAsGraph() const;
+	UE_API FEdGraphSchemaAction_K2Graph* SelectionAsGraph() const;
 
 	/** Accessor for getting the current selection as a K2 enum */
-	FEdGraphSchemaAction_K2Enum* SelectionAsEnum() const;
+	UE_API FEdGraphSchemaAction_K2Enum* SelectionAsEnum() const;
 
 	/** Accessor for getting the current selection as a K2 enum */
-	FEdGraphSchemaAction_K2Struct* SelectionAsStruct() const;
+	UE_API FEdGraphSchemaAction_K2Struct* SelectionAsStruct() const;
 
 	/** Accessor for getting the current selection as a K2 var */
-	FEdGraphSchemaAction_K2Var* SelectionAsVar() const;
+	UE_API FEdGraphSchemaAction_K2Var* SelectionAsVar() const;
 	
 	/** Accessor for getting the current selection as a K2 delegate */
-	FEdGraphSchemaAction_K2Delegate* SelectionAsDelegate() const;
+	UE_API FEdGraphSchemaAction_K2Delegate* SelectionAsDelegate() const;
 
 	/** Accessor for getting the current selection as a K2 event */
-	FEdGraphSchemaAction_K2Event* SelectionAsEvent() const;
+	UE_API FEdGraphSchemaAction_K2Event* SelectionAsEvent() const;
 	
 	/** Accessor for getting the current selection as a K2 Input Action */
-	FEdGraphSchemaAction_K2InputAction* SelectionAsInputAction() const;
+	UE_API FEdGraphSchemaAction_K2InputAction* SelectionAsInputAction() const;
 
 	/** Accessor for getting the current selection as a K2 local var */
-	FEdGraphSchemaAction_K2LocalVar* SelectionAsLocalVar() const;
+	UE_API FEdGraphSchemaAction_K2LocalVar* SelectionAsLocalVar() const;
 
 	/** Accessor for getting the current selection as a K2 blueprint base variable */
-	FEdGraphSchemaAction_BlueprintVariableBase* SelectionAsBlueprintVariable() const;
+	UE_API FEdGraphSchemaAction_BlueprintVariableBase* SelectionAsBlueprintVariable() const;
 
 	/** Accessor for determining if the current selection is a category*/
-	bool SelectionIsCategory() const;
+	UE_API bool SelectionIsCategory() const;
 	
-	void EnsureLastPinTypeValid();
+	UE_API void EnsureLastPinTypeValid();
 
 	/** Gets the last pin type selected by this widget, or by the function editor */
 	FEdGraphPinType& GetLastPinTypeUsed() {EnsureLastPinTypeValid(); return LastPinType;}
@@ -153,96 +155,96 @@ public:
 	 * Fills the supplied array with the currently selected objects
 	 * @param OutSelectedItems The array to fill.
 	 */
-	void GetSelectedItemsForContextMenu(TArray<FComponentEventConstructionData>& OutSelectedItems) const;
+	UE_API void GetSelectedItemsForContextMenu(TArray<FComponentEventConstructionData>& OutSelectedItems) const;
 
 	/** Called to reset the search filter */
-	void OnResetItemFilter();
+	UE_API void OnResetItemFilter();
 
 	/** Selects an item by name in either the main graph action menu or the local one */
-	void SelectItemByName(const FName& ItemName, ESelectInfo::Type SelectInfo = ESelectInfo::Direct, int32 SectionId = INDEX_NONE, bool bIsCategory = false);
+	UE_API void SelectItemByName(const FName& ItemName, ESelectInfo::Type SelectInfo = ESelectInfo::Direct, int32 SectionId = INDEX_NONE, bool bIsCategory = false);
 
 	/** Clears the selection in the graph action menus */
-	void ClearGraphActionMenuSelection();
+	UE_API void ClearGraphActionMenuSelection();
 
 	/** Initiates a rename on the selected action node, if possible */
-	void OnRequestRenameOnActionNode();
+	UE_API void OnRequestRenameOnActionNode();
 
 	/** Expands any category with the associated name */
-	void ExpandCategory(const FText& CategoryName);
+	UE_API void ExpandCategory(const FText& CategoryName);
 
 	/** Move the category before the target category */
-	bool MoveCategoryBeforeCategory( const FText& CategoryToMove, const FText& TargetCategory );
+	UE_API bool MoveCategoryBeforeCategory( const FText& CategoryToMove, const FText& TargetCategory );
 	
 	/** Callbacks for Paste Commands */
-	void OnPasteGeneric();
-	bool CanPasteGeneric();
+	UE_API void OnPasteGeneric();
+	UE_API bool CanPasteGeneric();
 private:
 	/** Creates widgets for the graph schema actions */
-	TSharedRef<SWidget> OnCreateWidgetForAction(struct FCreateWidgetForActionData* const InCreateData);
+	UE_API TSharedRef<SWidget> OnCreateWidgetForAction(struct FCreateWidgetForActionData* const InCreateData);
 
 	/** Callback used to populate all actions list in SGraphActionMenu */
-	void CollectAllActions(FGraphActionListBuilderBase& OutAllActions);
-	void CollectStaticSections(TArray<int32>& StaticSectionIDs);
-	void AddEventForFunctionGraph(UEdGraph* EdGraph, int32 const SectionId, FGraphActionSort& SortList, const FText& ParentCategory = FText::GetEmpty(), bool bAddChildGraphs = true) const;
-	void GetChildGraphs(UEdGraph* EdGraph, int32 const SectionId, FGraphActionSort& SortList, const FText& ParentCategory = FText::GetEmpty()) const;
-	void GetChildEvents(UEdGraph const* EdGraph, int32 const SectionId, FGraphActionSort& SortList, const FText& ParentCategory = FText::GetEmpty(), bool bInAddChildGraphs = true) const;
-	void GetLocalVariables(FGraphActionSort& SortList) const;
+	UE_API void CollectAllActions(FGraphActionListBuilderBase& OutAllActions);
+	UE_API void CollectStaticSections(TArray<int32>& StaticSectionIDs);
+	UE_API void AddEventForFunctionGraph(UEdGraph* EdGraph, int32 const SectionId, FGraphActionSort& SortList, const FText& ParentCategory = FText::GetEmpty(), bool bAddChildGraphs = true) const;
+	UE_API void GetChildGraphs(UEdGraph* EdGraph, int32 const SectionId, FGraphActionSort& SortList, const FText& ParentCategory = FText::GetEmpty()) const;
+	UE_API void GetChildEvents(UEdGraph const* EdGraph, int32 const SectionId, FGraphActionSort& SortList, const FText& ParentCategory = FText::GetEmpty(), bool bInAddChildGraphs = true) const;
+	UE_API void GetLocalVariables(FGraphActionSort& SortList) const;
 	
 	/** Handles the visibility of the local action list */
-	EVisibility GetLocalActionsListVisibility() const;
+	UE_API EVisibility GetLocalActionsListVisibility() const;
 
 	/** Callbacks for the graph action menu */
-	FReply OnActionDragged(const TArray< TSharedPtr<FEdGraphSchemaAction> >& InActions, const FPointerEvent& MouseEvent);
-	FReply OnCategoryDragged(const FText& InCategory, const FPointerEvent& MouseEvent);
-	void OnActionSelected(const TArray< TSharedPtr<FEdGraphSchemaAction> >& InActions);
-	void OnActionSelectedHelper(TSharedPtr<FEdGraphSchemaAction> InAction, TWeakPtr< FBlueprintEditor > InBlueprintEditor, UBlueprint* Blueprint, TSharedRef<SKismetInspector> Inspector);
-	void OnGlobalActionSelected(const TArray< TSharedPtr<FEdGraphSchemaAction> >& InActions, ESelectInfo::Type InSelectionType);
-	void OnActionDoubleClicked(const TArray< TSharedPtr<FEdGraphSchemaAction> >& InActions);
-	void ExecuteAction(TSharedPtr<FEdGraphSchemaAction> InAction);
-	TSharedPtr<SWidget> OnContextMenuOpening();
-
-	TSharedRef<SWidget> CreateAddNewMenuWidget();
-	void BuildAddNewMenu(FMenuBuilder& MenuBuilder);
-	TSharedRef<SWidget> CreateAddToSectionButton(int32 InSectionID, TWeakPtr<SWidget> WeakRowWidget, FText AddNewText, FName MetaDataTag);
-
-	void OnCategoryNameCommitted(const FText& InNewText, ETextCommit::Type InTextCommit, TWeakPtr< struct FGraphActionNode > InAction );
-	bool CanRequestRenameOnActionNode(TWeakPtr<struct FGraphActionNode> InSelectedNode) const;
-	FText OnGetSectionTitle( int32 InSectionID );
-	TSharedRef<SWidget> OnGetSectionWidget( TSharedRef<SWidget> RowWidget, int32 InSectionID );
-	EVisibility OnGetSectionTextVisibility(TWeakPtr<SWidget> RowWidget, int32 InSectionID) const;
-	TSharedRef<SWidget> OnGetFunctionListMenu();
-	void BuildOverridableFunctionsMenu(FMenuBuilder& MenuBuilder);
-	FReply OnAddButtonClickedOnSection(int32 InSectionID);
-	bool CanAddNewElementToSection(int32 InSectionID) const;
-
-	bool HandleActionMatchesName(FEdGraphSchemaAction* InAction, const FName& InName) const;
+	UE_API FReply OnActionDragged(const TArray< TSharedPtr<FEdGraphSchemaAction> >& InActions, const FPointerEvent& MouseEvent);
+	UE_API FReply OnCategoryDragged(const FText& InCategory, const FPointerEvent& MouseEvent);
+	UE_API void OnActionSelected(const TArray< TSharedPtr<FEdGraphSchemaAction> >& InActions);
+	UE_API void OnActionSelectedHelper(TSharedPtr<FEdGraphSchemaAction> InAction, TWeakPtr< FBlueprintEditor > InBlueprintEditor, UBlueprint* Blueprint, TSharedRef<SKismetInspector> Inspector);
+	UE_API void OnGlobalActionSelected(const TArray< TSharedPtr<FEdGraphSchemaAction> >& InActions, ESelectInfo::Type InSelectionType);
+	UE_API void OnActionDoubleClicked(const TArray< TSharedPtr<FEdGraphSchemaAction> >& InActions);
+	UE_API void ExecuteAction(TSharedPtr<FEdGraphSchemaAction> InAction);
+	UE_API TSharedPtr<SWidget> OnContextMenuOpening();
+
+	UE_API TSharedRef<SWidget> CreateAddNewMenuWidget();
+	UE_API void BuildAddNewMenu(FMenuBuilder& MenuBuilder);
+	UE_API TSharedRef<SWidget> CreateAddToSectionButton(int32 InSectionID, TWeakPtr<SWidget> WeakRowWidget, FText AddNewText, FName MetaDataTag);
+
+	UE_API void OnCategoryNameCommitted(const FText& InNewText, ETextCommit::Type InTextCommit, TWeakPtr< struct FGraphActionNode > InAction );
+	UE_API bool CanRequestRenameOnActionNode(TWeakPtr<struct FGraphActionNode> InSelectedNode) const;
+	UE_API FText OnGetSectionTitle( int32 InSectionID );
+	UE_API TSharedRef<SWidget> OnGetSectionWidget( TSharedRef<SWidget> RowWidget, int32 InSectionID );
+	UE_API EVisibility OnGetSectionTextVisibility(TWeakPtr<SWidget> RowWidget, int32 InSectionID) const;
+	UE_API TSharedRef<SWidget> OnGetFunctionListMenu();
+	UE_API void BuildOverridableFunctionsMenu(FMenuBuilder& MenuBuilder);
+	UE_API FReply OnAddButtonClickedOnSection(int32 InSectionID);
+	UE_API bool CanAddNewElementToSection(int32 InSectionID) const;
+
+	UE_API bool HandleActionMatchesName(FEdGraphSchemaAction* InAction, const FName& InName) const;
 
 	/** Support functions for checkbox to manage displaying user variables only */
-	bool IsShowingInheritedVariables() const;
-	void OnToggleShowInheritedVariables();
+	UE_API bool IsShowingInheritedVariables() const;
+	UE_API void OnToggleShowInheritedVariables();
 
 	/** Support functions for view options for Show Empty Sections */
-	void OnToggleShowEmptySections();
-	bool IsShowingEmptySections() const;
+	UE_API void OnToggleShowEmptySections();
+	UE_API bool IsShowingEmptySections() const;
 	
 	/** Support functions for view options for Show Replicated Variables only */
-	void OnToggleShowReplicatedVariablesOnly();
-	bool IsShowingReplicatedVariablesOnly() const;
+	UE_API void OnToggleShowReplicatedVariablesOnly();
+	UE_API bool IsShowingReplicatedVariablesOnly() const;
 
 	/** Support functions for view options for bAlwaysShowInterfacesInOverrides blueprint editor setting */
-	void OnToggleAlwaysShowInterfacesInOverrides();
-	bool GetAlwaysShowInterfacesInOverrides() const;
+	UE_API void OnToggleAlwaysShowInterfacesInOverrides();
+	UE_API bool GetAlwaysShowInterfacesInOverrides() const;
 
 	/** Support functions for view options for bShowParentClassInOverrides blueprint editor setting */
-	void OnToggleShowParentClassInOverrides();
-	bool GetShowParentClassInOverrides() const;
+	UE_API void OnToggleShowParentClassInOverrides();
+	UE_API bool GetShowParentClassInOverrides() const;
 
 	/** Support functions for view options for bShowAccessSpecifier blueprint editor setting */
-	void OnToggleShowAccessSpecifier();
-	bool GetShowAccessSpecifier() const;
+	UE_API void OnToggleShowAccessSpecifier();
+	UE_API bool GetShowAccessSpecifier() const;
 
 	/** Helper function to open the selected graph */
-	void OpenGraph(FDocumentTracker::EOpenDocumentCause InCause, bool bOpenExternalGraphInNewEditor = false);
+	UE_API void OpenGraph(FDocumentTracker::EOpenDocumentCause InCause, bool bOpenExternalGraphInNewEditor = false);
 
 	/**
 	* Check if the override of a given function is most likely desired as a blueprint function 
@@ -252,84 +254,84 @@ private:
 	* 
 	* @return	True if the function is desired as a function, false if desired as an event
 	*/
-	bool IsImplementationDesiredAsFunction(const UFunction* OverrideFunc) const;
+	UE_API bool IsImplementationDesiredAsFunction(const UFunction* OverrideFunc) const;
 
 	/** Callbacks for commands */
-	void OnOpenGraph();
-	void OnOpenGraphInNewTab();
-	void OnOpenExternalGraph();
-	bool CanOpenGraph() const;
-	bool CanOpenExternalGraph() const;
-	bool CanFocusOnNode() const;
-	void OnFocusNode();
-	void OnFocusNodeInNewTab();
-	void OnImplementFunction();
-	void ImplementFunction(TSharedPtr<FEdGraphSchemaAction_K2Graph> GraphAction);
-	void ImplementFunction(FEdGraphSchemaAction_K2Graph* GraphAction);
-	bool CanImplementFunction() const;
-	void OnFindReference(bool bSearchAllBlueprints, const EGetFindReferenceSearchStringFlags Flags);
-	bool CanFindReference() const;
-	void OnFindAndReplaceReference();
-	bool CanFindAndReplaceReference() const;
-	void OnDeleteEntry();
-	bool CanDeleteEntry() const;
-	FReply OnAddNewLocalVariable();
-	bool CanRequestRenameOnActionNode() const;
-	bool IsDuplicateActionVisible() const;
-	bool CanDuplicateAction() const;
-	void OnDuplicateAction();
-	void GotoNativeCodeVarDefinition();
-	bool IsNativeVariable() const;
-	void OnMoveToParent();
-	bool CanMoveVariableToParent() const;
-	bool CanMoveFunctionToParent() const;
-	void OnCopy();
-	bool CanCopy() const;
-	void OnCut();
-	bool CanCut() const;
-	void OnPasteVariable();
-	void OnPasteLocalVariable();
-	bool CanPasteVariable() const;
-	bool CanPasteLocalVariable() const;
-	void OnPasteFunction();
-	bool CanPasteFunction() const;
-	void OnPasteMacro();
-	bool CanPasteMacro() const;
+	UE_API void OnOpenGraph();
+	UE_API void OnOpenGraphInNewTab();
+	UE_API void OnOpenExternalGraph();
+	UE_API bool CanOpenGraph() const;
+	UE_API bool CanOpenExternalGraph() const;
+	UE_API bool CanFocusOnNode() const;
+	UE_API void OnFocusNode();
+	UE_API void OnFocusNodeInNewTab();
+	UE_API void OnImplementFunction();
+	UE_API void ImplementFunction(TSharedPtr<FEdGraphSchemaAction_K2Graph> GraphAction);
+	UE_API void ImplementFunction(FEdGraphSchemaAction_K2Graph* GraphAction);
+	UE_API bool CanImplementFunction() const;
+	UE_API void OnFindReference(bool bSearchAllBlueprints, const EGetFindReferenceSearchStringFlags Flags);
+	UE_API bool CanFindReference() const;
+	UE_API void OnFindAndReplaceReference();
+	UE_API bool CanFindAndReplaceReference() const;
+	UE_API void OnDeleteEntry();
+	UE_API bool CanDeleteEntry() const;
+	UE_API FReply OnAddNewLocalVariable();
+	UE_API bool CanRequestRenameOnActionNode() const;
+	UE_API bool IsDuplicateActionVisible() const;
+	UE_API bool CanDuplicateAction() const;
+	UE_API void OnDuplicateAction();
+	UE_API void GotoNativeCodeVarDefinition();
+	UE_API bool IsNativeVariable() const;
+	UE_API void OnMoveToParent();
+	UE_API bool CanMoveVariableToParent() const;
+	UE_API bool CanMoveFunctionToParent() const;
+	UE_API void OnCopy();
+	UE_API bool CanCopy() const;
+	UE_API void OnCut();
+	UE_API bool CanCut() const;
+	UE_API void OnPasteVariable();
+	UE_API void OnPasteLocalVariable();
+	UE_API bool CanPasteVariable() const;
+	UE_API bool CanPasteLocalVariable() const;
+	UE_API void OnPasteFunction();
+	UE_API bool CanPasteFunction() const;
+	UE_API void OnPasteMacro();
+	UE_API bool CanPasteMacro() const;
 
 	/** Gets the currently selected Category or returns default category name */
-	FText GetPasteCategory() const;
+	UE_API FText GetPasteCategory() const;
 
 	/** Callback when the filter is changed, forces the action tree(s) to filter */
-	void OnFilterTextChanged( const FText& InFilterText );
+	UE_API void OnFilterTextChanged( const FText& InFilterText );
 
 	/** Callback for the action trees to get the filter text */
-	FText GetFilterText() const;
+	UE_API FText GetFilterText() const;
 
 	/** Checks if the selected action has context menu */
-	bool SelectionHasContextMenu() const;
+	UE_API bool SelectionHasContextMenu() const;
 
 	/** Update Node Create Analytic */
-	void UpdateNodeCreation();
+	UE_API void UpdateNodeCreation();
 
 	/** Returns the displayed category, if any, of a graph */
-	FText GetGraphCategory(UEdGraph* InGraph) const;
+	UE_API FText GetGraphCategory(UEdGraph* InGraph) const;
 
 	/** Helper function to delete a graph in the MyBlueprint window */
-	void OnDeleteGraph(UEdGraph* InGraph, EEdGraphSchemaAction_K2Graph::Type);
+	UE_API void OnDeleteGraph(UEdGraph* InGraph, EEdGraphSchemaAction_K2Graph::Type);
 
 	/** Helper function to delete a delegate in the MyBlueprint window */
-	void OnDeleteDelegate(FEdGraphSchemaAction_K2Delegate* InDelegateAction);
+	UE_API void OnDeleteDelegate(FEdGraphSchemaAction_K2Delegate* InDelegateAction);
 
-	UEdGraph* GetFocusedGraph() const;
+	UE_API UEdGraph* GetFocusedGraph() const;
 
 	/** Delegate to hook us into non-structural Blueprint object post-change events */
-	void OnObjectPropertyChanged(UObject* InObject, FPropertyChangedEvent& InPropertyChangedEvent);
+	UE_API void OnObjectPropertyChanged(UObject* InObject, FPropertyChangedEvent& InPropertyChangedEvent);
 
 	/** Helper function indicating whehter we're in editing mode, and can modify the target blueprint */
-	bool IsEditingMode() const;
+	UE_API bool IsEditingMode() const;
 
 	/** Determine whether an FEdGraphSchemaAction is associated with an event */
-	static bool IsAnInterfaceEvent(FEdGraphSchemaAction_K2Graph* InAction);
+	static UE_API bool IsAnInterfaceEvent(FEdGraphSchemaAction_K2Graph* InAction);
 
 private:
 	/** List of UI Commands for this scope */
@@ -378,3 +380,5 @@ private:
 	/** If set we'll show only replicated variables (local to a particular blueprint view). */
 	bool bShowReplicatedVariablesOnly;
 };
+
+#undef UE_API
diff --git a/Engine/Source/Editor/Kismet/Public/SPinValueInspector.h b/Engine/Source/Editor/Kismet/Public/SPinValueInspector.h
--- a/Engine/Source/Editor/Kismet/Public/SPinValueInspector.h
+++ b/Engine/Source/Editor/Kismet/Public/SPinValueInspector.h
@@ -13,6 +13,8 @@
 #include "Widgets/DeclarativeSyntaxSupport.h"
 #include "Widgets/SCompoundWidget.h"
 
+#define UE_API KISMET_API
+
 class FText;
 class SSearchBox;
 class UBlueprint;
@@ -24,47 +26,47 @@ typedef TSharedPtr<struct FPinValueInspectorTreeViewNode> FPinValueInspectorTree
  * Inspects the referenced pin object's underlying property value and presents it within a tree view.
  * Compound properties (e.g. structs/containers) will be broken down into a hierarchy of child nodes.
  */
-class KISMET_API SPinValueInspector : public SCompoundWidget
+class SPinValueInspector : public SCompoundWidget
 {
 public:
 	SLATE_BEGIN_ARGS(SPinValueInspector)
 	{}
 	SLATE_END_ARGS()
 
-	void Construct(const FArguments& InArgs);
+	UE_API void Construct(const FArguments& InArgs);
 
 	/** Whether the search filter UI should be visible. */
-	bool ShouldShowSearchFilter() const;
+	UE_API bool ShouldShowSearchFilter() const;
 
 	/** Gets the current watched pin reference */
-	const FEdGraphPinReference& GetPinRef() const;
+	UE_API const FEdGraphPinReference& GetPinRef() const;
 
 	friend class FPinValueInspectorTooltip;
 
 protected:
 	/** @return Visibility of the search box filter widget. */
-	virtual EVisibility GetSearchFilterVisibility() const;
+	UE_API virtual EVisibility GetSearchFilterVisibility() const;
 
 	/** Passes SearchText through to tree view */
-	virtual void OnSearchTextChanged(const FText& InSearchText);
+	UE_API virtual void OnSearchTextChanged(const FText& InSearchText);
 
 	/** requests the constrained box be resized */
-	virtual void OnExpansionChanged(FDebugTreeItemPtr InItem, bool bItemIsExpanded);
+	UE_API virtual void OnExpansionChanged(FDebugTreeItemPtr InItem, bool bItemIsExpanded);
 
 	/** Adds a unique tree item to the TreeView. */
-	virtual void AddTreeItemUnique(const FDebugTreeItemPtr& Item);
+	UE_API virtual void AddTreeItemUnique(const FDebugTreeItemPtr& Item);
 
 	/** Adds relevant pins to the tree view */
-	virtual void PopulateTreeView();
+	UE_API virtual void PopulateTreeView();
 
 	/** Adds a single pin to the tree view */
-	void AddPinToTreeView(const UEdGraphPin* Pin, UBlueprint* Blueprint);
+	UE_API void AddPinToTreeView(const UEdGraphPin* Pin, UBlueprint* Blueprint);
 
 	/** Marks tree view for updating currently filtered items. */
-	void RequestUpdateFilteredItems();
+	UE_API void RequestUpdateFilteredItems();
 
 	/** Sets the current watched pin */
-	void SetPinRef(const FEdGraphPinReference& InPinRef);
+	UE_API void SetPinRef(const FEdGraphPinReference& InPinRef);
 
 private:
 	/** Holds a weak reference to the target pin. */
@@ -81,11 +83,11 @@ private:
 };
 
 /** class holding functions to spawn a pin value inspector tooltip */
-class KISMET_API FPinValueInspectorTooltip
+class FPinValueInspectorTooltip
 {
 public:
 	/** Moves the tooltip to the new location */
-	void MoveTooltip(const UE::Slate::FDeprecateVector2DParameter& InNewLocation);
+	UE_API void MoveTooltip(const UE::Slate::FDeprecateVector2DParameter& InNewLocation);
 
 	/** 
 	 * Dismisses the current tooltip, if it is not currently hovered
@@ -93,38 +95,40 @@ public:
 	 * 
 	 * @param bForceDismiss	true to dismiss regardless of hover & tooltip ownership 
 	 */
-	void TryDismissTooltip(bool bForceDismiss = false);
+	UE_API void TryDismissTooltip(bool bForceDismiss = false);
 
 	/** @returns whether or not the tooltip can close */
-	bool TooltipCanClose() const;
+	UE_API bool TooltipCanClose() const;
 
 	/** @returns whether this tooltip is the host for a context menu */
-	bool TooltipHostsMenu() const;
+	UE_API bool TooltipHostsMenu() const;
 
 private:
 	/** Dismisses the current tooltip (internal implementation) */
-	void DismissTooltip();
+	UE_API void DismissTooltip();
 
 public:
 	/** Summons a new tooltip in the shared window. If provided, uses a custom implementation of a pin value inspector widget. */
-	static TWeakPtr<FPinValueInspectorTooltip> SummonTooltip(FEdGraphPinReference InPinRef, TSharedPtr<SPinValueInspector> InContentWidgetOverride = { });
+	static UE_API TWeakPtr<FPinValueInspectorTooltip> SummonTooltip(FEdGraphPinReference InPinRef, TSharedPtr<SPinValueInspector> InContentWidgetOverride = { });
 
 	/** Release the references to the static widget shared pointers */
-	static void ShutdownTooltip();
+	static UE_API void ShutdownTooltip();
 
 	/** Default inspector widget in the tooltip */
-	static TSharedPtr<SPinValueInspector> ValueInspectorWidget;
+	static UE_API TSharedPtr<SPinValueInspector> ValueInspectorWidget;
 
 private:
 	/** Handles Creating a custom tooltip window for all PinValueInspector tooltips */
-	static void CreatePinValueTooltipWindow();
+	static UE_API void CreatePinValueTooltipWindow();
 
 	/** A reusable tooltip window for PinValueInspector */
-	static TSharedPtr<class SWindow> TooltipWindow;
+	static UE_API TSharedPtr<class SWindow> TooltipWindow;
 
 	/** Tooltip widget housed in the window */
-	static TSharedPtr<class SToolTip> TooltipWidget;
+	static UE_API TSharedPtr<class SToolTip> TooltipWidget;
 
 	/** The current "live" tooltip */
-	static TSharedPtr<FPinValueInspectorTooltip> Instance;
+	static UE_API TSharedPtr<FPinValueInspectorTooltip> Instance;
 };
+
+#undef UE_API
diff --git a/Engine/Source/Editor/Kismet/Public/SSCSEditor.h b/Engine/Source/Editor/Kismet/Public/SSCSEditor.h
--- a/Engine/Source/Editor/Kismet/Public/SSCSEditor.h
+++ b/Engine/Source/Editor/Kismet/Public/SSCSEditor.h
@@ -52,6 +52,8 @@
 #include "Widgets/Views/STableViewBase.h"
 #include "Widgets/Views/STreeView.h"
 
+#define UE_API KISMET_API
+
 class FDragDropEvent;
 class FExtender;
 class FMenuBuilder;
@@ -93,7 +95,7 @@ using FSCSEditorChildActorNodePtrType = TSharedPtr<class FSCSEditorTreeNodeChild
  *
  * Wrapper class for nodes displayed in the SCS (Simple Construction Script) editor tree widget.
  */
-class KISMET_API FSCSEditorTreeNode : public TSharedFromThis<FSCSEditorTreeNode>
+class FSCSEditorTreeNode : public TSharedFromThis<FSCSEditorTreeNode>
 {
 public:
 	/** Delegate for when the context menu requests a rename */
@@ -110,29 +112,29 @@ public:
 	/**
 	 * Constructs an empty tree node.
 	 */
-	FSCSEditorTreeNode(FSCSEditorTreeNode::ENodeType InNodeType);
+	UE_API FSCSEditorTreeNode(FSCSEditorTreeNode::ENodeType InNodeType);
 
 	/**
 	* @return The name to identify this node.
 	*/
-	virtual FName GetNodeID() const;
+	UE_API virtual FName GetNodeID() const;
 
 	/**
 	 * @return The name of the variable represented by this node.
 	 */
-	virtual FName GetVariableName() const;
+	UE_API virtual FName GetVariableName() const;
 	/**
 	 * @return The string to be used in the tree display.
 	 */
-	virtual FString GetDisplayString() const;
+	UE_API virtual FString GetDisplayString() const;
 	/**
 	* @return The name of this node in text.
 	*/
-	virtual FText GetDisplayName() const;
+	UE_API virtual FText GetDisplayName() const;
 	/**
 	 * @return The SCS node that is represented by this object, or NULL if there is no associated SCS node.
 	 */
-	virtual class USCS_Node* GetSCSNode() const;
+	UE_API virtual class USCS_Node* GetSCSNode() const;
 	/**
 	 * @param ActualEditedBlueprint currently edited blueprint
 	 * @note Derived classes should override GetOrCreateEditableObjectForBlueprint().
@@ -149,7 +151,7 @@ public:
 	 * @param InActor The Actor instance to use as the container object for finding the component instance.
 	 * @return The component instance represented by this node and contained within the given Actor instance, or NULL if not found.
 	 */
-	virtual UActorComponent* FindComponentInstanceInActor(const AActor* InActor) const;
+	UE_API virtual UActorComponent* FindComponentInstanceInActor(const AActor* InActor) const;
 	/**
 	 * @return This object's parent node (or an invalid reference if no parent is assigned).
 	 */
@@ -165,11 +167,11 @@ public:
 	/**
 	 * Sets the actor root to the given node for this node along with any children.
 	 */
-	void SetActorRootNode(FSCSEditorActorNodePtrType InActorNode);
+	UE_API void SetActorRootNode(FSCSEditorActorNodePtrType InActorNode);
 	/**
 	 * @return Type of node
 	 */
-	ENodeType GetNodeType() const;
+	UE_API ENodeType GetNodeType() const;
 	/**
 	 * @return True if this represents an actor node type
 	 */
@@ -236,7 +238,7 @@ public:
 	/**
 	 * @return Whether or not this node is a child (direct or indirect) of the given node.
 	 */
-	bool IsAttachedTo(FSCSEditorTreeNodePtrType InNodePtr) const;	
+	UE_API bool IsAttachedTo(FSCSEditorTreeNodePtrType InNodePtr) const;	
 
 	/**
 	 * Finds the closest ancestor node in the given node set.
@@ -244,14 +246,14 @@ public:
 	 * @param InNodes The given node set.
 	 * @return One of the nodes from the set, or an invalid node reference if the set does not contain any ancestor nodes.
 	 */
-	FSCSEditorTreeNodePtrType FindClosestParent(TArray<FSCSEditorTreeNodePtrType> InNodes);
+	UE_API FSCSEditorTreeNodePtrType FindClosestParent(TArray<FSCSEditorTreeNodePtrType> InNodes);
 
 	/**
 	 * Adds the given node as a child node.
 	 *
 	 * @param InChildNodePtr The node to add as a child node.
 	 */
-	virtual void AddChild(FSCSEditorTreeNodePtrType InChildNodePtr);
+	UE_API virtual void AddChild(FSCSEditorTreeNodePtrType InChildNodePtr);
 
 	/**
 	 * Adds a child node for the given SCS node.
@@ -260,7 +262,7 @@ public:
 	 * @param bIsInheritedSCS Whether or not the given SCS node is inherited from a parent.
 	 * @return A reference to the new child node or the existing child node if a match is found.
 	 */
-	FSCSEditorTreeNodePtrType AddChild(USCS_Node* InSCSNode, bool bIsInheritedSCS);
+	UE_API FSCSEditorTreeNodePtrType AddChild(USCS_Node* InSCSNode, bool bIsInheritedSCS);
 
 	/**
 	 * Adds a child node for the given component template.
@@ -268,7 +270,7 @@ public:
 	 * @param InComponentTemplate The component template for which to create a child node.
 	 * @return A reference to the new child node or the existing child node if a match is found.
 	 */
-	FSCSEditorTreeNodePtrType AddChildFromComponent(UActorComponent* InComponentTemplate);
+	UE_API FSCSEditorTreeNodePtrType AddChildFromComponent(UActorComponent* InComponentTemplate);
 
 	/**
 	 * Attempts to find a reference to the child node that matches the given SCS node.
@@ -278,7 +280,7 @@ public:
 	 * @param OutDepth If non-NULL, the depth of the child node will be returned in this parameter on success (default == NULL).
 	 * @return The child node that matches the given SCS node, or an invalid node reference if no match was found.
 	 */
-	FSCSEditorTreeNodePtrType FindChild(const USCS_Node* InSCSNode, bool bRecursiveSearch = false, uint32* OutDepth = NULL) const;
+	UE_API FSCSEditorTreeNodePtrType FindChild(const USCS_Node* InSCSNode, bool bRecursiveSearch = false, uint32* OutDepth = NULL) const;
 
 	/**
 	 * Attempts to find a reference to the child node that matches the given component template.
@@ -288,7 +290,7 @@ public:
 	 * @param OutDepth If non-NULL, the depth of the child node will be returned in this parameter on success (default == NULL).
 	 * @return The child node with a component template that matches the given component template instance, or an invalid node reference if no match was found.
 	 */
-	FSCSEditorTreeNodePtrType FindChild(const UActorComponent* InComponentTemplate, bool bRecursiveSearch = false, uint32* OutDepth = NULL) const;
+	UE_API FSCSEditorTreeNodePtrType FindChild(const UActorComponent* InComponentTemplate, bool bRecursiveSearch = false, uint32* OutDepth = NULL) const;
 
 	/**
 	 * Attempts to find a reference to the child node that matches the given component variable or instance name.
@@ -298,14 +300,14 @@ public:
 	 * @param OutDepth If non-NULL, the depth of the child node will be returned in this parameter on success (default == NULL).
 	 * @return The child node with a component variable or instance name that matches the given name, or an invalid node reference if no match was found.
 	 */
-	FSCSEditorTreeNodePtrType FindChild(const FName& InVariableOrInstanceName, bool bRecursiveSearch = false, uint32* OutDepth = NULL) const;
+	UE_API FSCSEditorTreeNodePtrType FindChild(const FName& InVariableOrInstanceName, bool bRecursiveSearch = false, uint32* OutDepth = NULL) const;
 
 	/**
 	 * Removes the given node from the list of child nodes.
 	 *
 	 * @param InChildNodePtr The child node to remove.
 	 */
-	virtual void RemoveChild(FSCSEditorTreeNodePtrType InChildNodePtr);
+	UE_API virtual void RemoveChild(FSCSEditorTreeNodePtrType InChildNodePtr);
 
 	bool IsSceneComponent() const
 	{
@@ -316,12 +318,12 @@ public:
 	virtual FSCSEditorChildActorNodePtrType GetChildActorNode() { return nullptr; }
 
 	// Tries to find a SCS node that was likely responsible for creating the specified instance component.  Note: This is not always possible to do!
-	static USCS_Node* FindSCSNodeForInstance(const UActorComponent* InstanceComponent, UClass* ClassToSearch);
+	static UE_API USCS_Node* FindSCSNodeForInstance(const UActorComponent* InstanceComponent, UClass* ClassToSearch);
 
 	/**
 	 * Creates the correct type of node based on the component (instanced or not, etc...)
 	 */
-	static FSCSEditorTreeNodePtrType FactoryNodeFromComponent(UActorComponent* InComponent);
+	static UE_API FSCSEditorTreeNodePtrType FactoryNodeFromComponent(UActorComponent* InComponent);
 
 	// Destructor
 	virtual ~FSCSEditorTreeNode() {}
@@ -331,7 +333,7 @@ public:
 	 * the node is in edition mode. When the text is not in edit mode anymore, the ongoing create transaction ends and the node
 	 * is considered fully created.
 	 */
-	void CloseOngoingCreateTransaction();
+	UE_API void CloseOngoingCreateTransaction();
 
 protected:
 	// Called when this node is being removed via a RemoveChild call
@@ -345,7 +347,7 @@ protected:
 	}
 
 	// Derived classes can override to create and/or return a reference to an alternate editable object.
-	virtual UObject* GetOrCreateEditableObjectForBlueprint(UBlueprint* InBlueprint) const;
+	UE_API virtual UObject* GetOrCreateEditableObjectForBlueprint(UBlueprint* InBlueprint) const;
 
 public:
 	/**
@@ -439,10 +441,10 @@ public:
 	 * Requests a rename on the node.
 	 * @param OngoingCreateTransaction The transaction scoping the node creation which will end once the node is named by the user or null if the rename is not part of a the creation process.
 	 */
-	void OnRequestRename(TUniquePtr<FScopedTransaction> OngoingCreateTransaction);
+	UE_API void OnRequestRename(TUniquePtr<FScopedTransaction> OngoingCreateTransaction);
 
 	/** Renames the object or variable represented by this node */
-	virtual void OnCompleteRename(const FText& InNewName);
+	UE_API virtual void OnCompleteRename(const FText& InNewName);
 
 	/** Sets up the delegate for a rename operation */
 	void SetRenameRequestedDelegate(FOnRenameRequested InRenameRequested) { RenameRequestedDelegate = InRenameRequested; }
@@ -455,18 +457,18 @@ public:
 	}
 
 	/** Returns whether the node will match the given type (for filtering) */
-	virtual bool MatchesFilterType(const UClass* InFilterType) const;
+	UE_API virtual bool MatchesFilterType(const UClass* InFilterType) const;
 
 	/** Refreshes this item's filtration state. Set bRecursive to 'true' to refresh any child nodes as well */
-	virtual bool RefreshFilteredState(const UClass* InFilterType, const TArray<FString>& InFilterTerms, bool bRecursive);
+	UE_API virtual bool RefreshFilteredState(const UClass* InFilterType, const TArray<FString>& InFilterTerms, bool bRecursive);
 
 protected:
 	/** Sets this item's filtration state. Use bUpdateParent to make sure the parent's EFilteredState::ChildMatches flag is properly updated based off the new state */
-	void SetCachedFilterState(bool bMatchesFilter, bool bUpdateParent);
+	UE_API void SetCachedFilterState(bool bMatchesFilter, bool bUpdateParent);
 	/** Updates the EFilteredState::ChildMatches flag, based off of children's current state */
-	void RefreshCachedChildFilterState(bool bUpdateParent);
+	UE_API void RefreshCachedChildFilterState(bool bUpdateParent);
 	/** Used to update the EFilteredState::ChildMatches flag for parent nodes, when this item's filtration state has changed */
-	void ApplyFilteredStateToParent();
+	UE_API void ApplyFilteredStateToParent();
 	
 	// Scope the creation of a node which ends when the initial 'name' is given/accepted by the user, which can be several frames after the node was actually created.
 	TUniquePtr<FScopedTransaction> OngoingCreateTransaction;
@@ -501,7 +503,7 @@ private:
 //////////////////////////////////////////////////////////////////////////
 //
 
-class KISMET_API FSCSEditorTreeNodeComponentBase : public FSCSEditorTreeNode
+class FSCSEditorTreeNodeComponentBase : public FSCSEditorTreeNode
 {
 protected:
 	FSCSEditorTreeNodeComponentBase()
@@ -511,14 +513,14 @@ protected:
 
 public:
 	// FSCSEditorTreeNode interface
-	virtual FName GetVariableName() const override;
-	virtual FString GetDisplayString() const override;
+	UE_API virtual FName GetVariableName() const override;
+	UE_API virtual FString GetDisplayString() const override;
 	virtual bool CanRename() const override { return !IsInheritedComponent() && !IsDefaultSceneRoot(); }
 	virtual bool CanDelete() const override { return !IsInheritedComponent() && !IsDefaultSceneRoot(); }
-	virtual bool CanReparent() const override;
-	virtual UBlueprint* GetBlueprint() const override;
-	virtual FSCSEditorChildActorNodePtrType GetChildActorNode();
-	virtual bool MatchesFilterType(const UClass* InFilterType) const override;
+	UE_API virtual bool CanReparent() const override;
+	UE_API virtual UBlueprint* GetBlueprint() const override;
+	UE_API virtual FSCSEditorChildActorNodePtrType GetChildActorNode();
+	UE_API virtual bool MatchesFilterType(const UClass* InFilterType) const override;
 	// End of FSCSEditorTreeNode interface
 
 private:
@@ -530,7 +532,7 @@ private:
 //////////////////////////////////////////////////////////////////////////
 // FSCSEditorTreeNodeInstancedInheritedComponent - A inherited component in the instanced case (either an inherited SCS node or an inherited native component)
 
-class KISMET_API FSCSEditorTreeNodeInstancedInheritedComponent : public FSCSEditorTreeNodeComponentBase
+class FSCSEditorTreeNodeInstancedInheritedComponent : public FSCSEditorTreeNodeComponentBase
 {
 public:
 	/**
@@ -538,17 +540,17 @@ public:
 	 *
 	 * @param InComponentTemplate The component template represented by this object.
 	 */
-	FSCSEditorTreeNodeInstancedInheritedComponent(AActor* Owner, UActorComponent* InComponentTemplate);
+	UE_API FSCSEditorTreeNodeInstancedInheritedComponent(AActor* Owner, UActorComponent* InComponentTemplate);
 
 	// FSCSEditorTreeNode public interface
-	virtual bool IsNativeComponent() const override;
-	virtual bool IsRootComponent() const override;
-	virtual bool IsInheritedSCSNode() const override;
+	UE_API virtual bool IsNativeComponent() const override;
+	UE_API virtual bool IsRootComponent() const override;
+	UE_API virtual bool IsInheritedSCSNode() const override;
 	virtual bool IsInstanced() const override { return true; }
 	virtual bool IsInheritedComponent() const override { return true; }
-	virtual bool IsDefaultSceneRoot() const override;
-	virtual bool CanEdit() const override;
-	virtual FText GetDisplayName() const override;
+	UE_API virtual bool IsDefaultSceneRoot() const override;
+	UE_API virtual bool CanEdit() const override;
+	UE_API virtual FText GetDisplayName() const override;
 	// End of FSCSEditorTreeNode public interface
 
 private:
@@ -558,7 +560,7 @@ private:
 //////////////////////////////////////////////////////////////////////////
 // FSCSEditorTreeNodeInstanceAddedComponent - A unique-to-this instance component
 
-class KISMET_API FSCSEditorTreeNodeInstanceAddedComponent : public FSCSEditorTreeNodeComponentBase
+class FSCSEditorTreeNodeInstanceAddedComponent : public FSCSEditorTreeNodeComponentBase
 {
 public:
 	/**
@@ -566,23 +568,23 @@ public:
 	*
 	* @param InComponentTemplate The component template represented by this object.
 	*/
-	FSCSEditorTreeNodeInstanceAddedComponent(AActor* Owner, UActorComponent* InComponentTemplate);
+	UE_API FSCSEditorTreeNodeInstanceAddedComponent(AActor* Owner, UActorComponent* InComponentTemplate);
 
 	// FSCSEditorTreeNode public interface
-	virtual bool IsRootComponent() const override;
+	UE_API virtual bool IsRootComponent() const override;
 	virtual bool IsInstanced() const override { return true; }
 	virtual bool IsUserInstancedComponent() const override { return true; }
-	virtual bool IsDefaultSceneRoot() const override;
+	UE_API virtual bool IsDefaultSceneRoot() const override;
 	virtual bool CanEdit() const override { return true; }
 	virtual FName GetVariableName() const override { return NAME_None; }
-	virtual FString GetDisplayString() const override;
-	virtual FText GetDisplayName() const override;
-	virtual void OnCompleteRename(const FText& InNewName) override;
+	UE_API virtual FString GetDisplayString() const override;
+	UE_API virtual FText GetDisplayName() const override;
+	UE_API virtual void OnCompleteRename(const FText& InNewName) override;
 	// End of FSCSEditorTreeNode public interface
 
 protected:
 	// FSCSEditorTreeNode protected interface
-	virtual void RemoveMeAsChild() override;
+	UE_API virtual void RemoveMeAsChild() override;
 	// End of FSCSEditorTreeNode protected interface
 
 private:
@@ -593,7 +595,7 @@ private:
 //////////////////////////////////////////////////////////////////////////
 // FSCSEditorTreeNodeComponent - A generic component in the non-instanced case (either a SCS node or an inherited native component)
 
-class KISMET_API FSCSEditorTreeNodeComponent : public FSCSEditorTreeNodeComponentBase
+class FSCSEditorTreeNodeComponent : public FSCSEditorTreeNodeComponentBase
 {
 public:
 	/**
@@ -602,35 +604,35 @@ public:
 	 * @param InSCSNode The SCS tree node represented by this object.
 	 * @param bInIsInherited Whether or not the SCS tree node is inherited from a parent Blueprint class.
 	 */
-	FSCSEditorTreeNodeComponent(class USCS_Node* InSCSNode, bool bInIsInherited = false);
+	UE_API FSCSEditorTreeNodeComponent(class USCS_Node* InSCSNode, bool bInIsInherited = false);
 
 	/**
 	 * Constructs a wrapper around a component template not contained within an SCS tree node (e.g. "native" components).
 	 *
 	 * @param InComponentTemplate The component template represented by this object.
 	 */
-	FSCSEditorTreeNodeComponent(UActorComponent* InComponentTemplate);
+	UE_API FSCSEditorTreeNodeComponent(UActorComponent* InComponentTemplate);
 
 
 	// FSCSEditorTreeNode public interface
-	virtual bool IsNativeComponent() const override;
-	virtual bool IsRootComponent() const override;
-	virtual bool IsInheritedSCSNode() const override;
-	virtual bool IsDefaultSceneRoot() const override;
-	virtual bool CanEdit() const override;
-	virtual FText GetDisplayName() const override;
-	virtual class USCS_Node* GetSCSNode() const override;
-	virtual void OnCompleteRename(const FText& InNewName) override;
+	UE_API virtual bool IsNativeComponent() const override;
+	UE_API virtual bool IsRootComponent() const override;
+	UE_API virtual bool IsInheritedSCSNode() const override;
+	UE_API virtual bool IsDefaultSceneRoot() const override;
+	UE_API virtual bool CanEdit() const override;
+	UE_API virtual FText GetDisplayName() const override;
+	UE_API virtual class USCS_Node* GetSCSNode() const override;
+	UE_API virtual void OnCompleteRename(const FText& InNewName) override;
 	// End of FSCSEditorTreeNode public interface
 
 protected:
 	// FSCSEditorTreeNode protected interface
-	virtual void RemoveMeAsChild() override;
-	virtual UObject* GetOrCreateEditableObjectForBlueprint(UBlueprint* InBlueprint) const override;
+	UE_API virtual void RemoveMeAsChild() override;
+	UE_API virtual UObject* GetOrCreateEditableObjectForBlueprint(UBlueprint* InBlueprint) const override;
 	// End of FSCSEditorTreeNode protected interface
 
 	/** Get overridden template component, specialized in given blueprint */
-	UActorComponent* INTERNAL_GetOverridenComponentTemplate(UBlueprint* Blueprint) const;
+	UE_API UActorComponent* INTERNAL_GetOverridenComponentTemplate(UBlueprint* Blueprint) const;
 
 private:
 	// Was this component inherited from a parent class or introduced in this class?
@@ -640,7 +642,7 @@ private:
 	TWeakObjectPtr<class USCS_Node> SCSNodePtr;
 };
 
-class KISMET_API FSCSEditorTreeNodeActorBase : public FSCSEditorTreeNode
+class FSCSEditorTreeNodeActorBase : public FSCSEditorTreeNode
 {
 public:
 	FSCSEditorTreeNodeActorBase(FSCSEditorTreeNode::ENodeType InNodeType, AActor* InActor)
@@ -649,22 +651,22 @@ public:
 		SetObject(InActor);
 	}
 
-	FSCSEditorTreeNodePtrType GetOwnerNode() const;
-	void SetOwnerNode(FSCSEditorTreeNodePtrType NewOwnerNode);
+	UE_API FSCSEditorTreeNodePtrType GetOwnerNode() const;
+	UE_API void SetOwnerNode(FSCSEditorTreeNodePtrType NewOwnerNode);
 
-	FSCSEditorTreeNodePtrType GetSceneRootNode() const;
-	void SetSceneRootNode(FSCSEditorTreeNodePtrType NewSceneRootNode);
+	UE_API FSCSEditorTreeNodePtrType GetSceneRootNode() const;
+	UE_API void SetSceneRootNode(FSCSEditorTreeNodePtrType NewSceneRootNode);
 
 	/** Returns the set of root nodes */
-	const TArray<FSCSEditorTreeNodePtrType>& GetComponentNodes() const;
+	UE_API const TArray<FSCSEditorTreeNodePtrType>& GetComponentNodes() const;
 
 	// FSCSEditorTreeNode public interface
-	virtual FName GetNodeID() const override;
-	virtual bool IsInstanced() const override;
+	UE_API virtual FName GetNodeID() const override;
+	UE_API virtual bool IsInstanced() const override;
 	virtual bool CanEdit() const override { return true; }
-	virtual void AddChild(FSCSEditorTreeNodePtrType InChildNodePtr) override;
-	virtual void RemoveChild(FSCSEditorTreeNodePtrType InChildNodePtr) override;
-	virtual UBlueprint* GetBlueprint() const override;
+	UE_API virtual void AddChild(FSCSEditorTreeNodePtrType InChildNodePtr) override;
+	UE_API virtual void RemoveChild(FSCSEditorTreeNodePtrType InChildNodePtr) override;
+	UE_API virtual UBlueprint* GetBlueprint() const override;
 	// End of FSCSEditorTreeNode public interface
 
 protected:
@@ -679,7 +681,7 @@ private:
 	TArray<FSCSEditorTreeNodePtrType> ComponentNodes;
 };
 
-class KISMET_API FSCSEditorTreeNodeRootActor : public FSCSEditorTreeNodeActorBase
+class FSCSEditorTreeNodeRootActor : public FSCSEditorTreeNodeActorBase
 {
 public:
 	FSCSEditorTreeNodeRootActor(AActor* InActor, bool bInAllowRename)
@@ -691,10 +693,10 @@ public:
 
 	// FSCSEditorTreeNode public interface
 	virtual bool CanRename() const override { return bAllowRename; }
-	virtual void OnCompleteRename(const FText& InNewName) override;
-	virtual void AddChild(FSCSEditorTreeNodePtrType InChildNodePtr) override;
-	virtual void RemoveChild(FSCSEditorTreeNodePtrType InChildNodePtr) override;
-	virtual bool RefreshFilteredState(const UClass* InFilterType, const TArray<FString>& InFilterTerms, bool bRecursive) override;
+	UE_API virtual void OnCompleteRename(const FText& InNewName) override;
+	UE_API virtual void AddChild(FSCSEditorTreeNodePtrType InChildNodePtr) override;
+	UE_API virtual void RemoveChild(FSCSEditorTreeNodePtrType InChildNodePtr) override;
+	UE_API virtual bool RefreshFilteredState(const UClass* InFilterType, const TArray<FString>& InFilterTerms, bool bRecursive) override;
 	// End of FSCSEditorTreeNode public interface
 
 private:
@@ -705,7 +707,7 @@ private:
 	TWeakPtr<class FSCSEditorTreeNodeSeparator> NonSceneComponentSeparatorNodePtr;
 };
 
-class KISMET_API FSCSEditorTreeNodeChildActor : public FSCSEditorTreeNodeActorBase
+class FSCSEditorTreeNodeChildActor : public FSCSEditorTreeNodeActorBase
 {
 public:
 	FSCSEditorTreeNodeChildActor(AActor* InActor)
@@ -714,13 +716,13 @@ public:
 	}
 
 	// FSCSEditorTreeNode public interface
-	virtual bool IsFlaggedForFiltration() const override;
+	UE_API virtual bool IsFlaggedForFiltration() const override;
 	// End of FSCSEditorTreeNode public interface
 
-	UChildActorComponent* GetChildActorComponent() const;
+	UE_API UChildActorComponent* GetChildActorComponent() const;
 };
 
-class KISMET_API FSCSEditorTreeNodeSeparator : public FSCSEditorTreeNode
+class FSCSEditorTreeNodeSeparator : public FSCSEditorTreeNode
 {
 public:
 	FSCSEditorTreeNodeSeparator()
@@ -730,11 +732,11 @@ public:
 
 	// FSCSEditorTreeNode public interface
 	virtual UBlueprint* GetBlueprint() const override { return nullptr; }
-	virtual bool MatchesFilterType(const UClass* InFilterType) const override;
+	UE_API virtual bool MatchesFilterType(const UClass* InFilterType) const override;
 	// End of FSCSEditorTreeNode public interface
 
 	/** If the given type matches the tree view filter, the separator will also be flagged for filtration. */
-	void AddFilteredComponentType(const TSubclassOf<UActorComponent>& InFilteredType);
+	UE_API void AddFilteredComponentType(const TSubclassOf<UActorComponent>& InFilteredType);
 
 private:
 	TArray<const UClass*> FilteredTypes;
@@ -1008,7 +1010,7 @@ namespace EComponentEditorMode
 
 class UE_DEPRECATED(5.0, "SSCSEditor has been deprecated, use a child class of SSubobjectEditor instead.") SSCSEditor;
 
-class KISMET_API SSCSEditor : public SCompoundWidget
+class SSCSEditor : public SCompoundWidget
 {
 public:
 	DECLARE_DELEGATE_RetVal_OneParam(class USCS_Node*, FOnAddNewComponent, class UClass*);
@@ -1043,29 +1045,29 @@ public:
 
 	SLATE_END_ARGS()
 
-	void Construct(const FArguments& InArgs);
+	UE_API void Construct(const FArguments& InArgs);
 
-	~SSCSEditor();
+	UE_API ~SSCSEditor();
 
 	/** SWidget interface */
-	virtual FReply OnKeyDown( const FGeometry& MyGeometry, const FKeyEvent& InKeyEvent );
+	UE_API virtual FReply OnKeyDown( const FGeometry& MyGeometry, const FKeyEvent& InKeyEvent );
 
 	/** Used by tree control - make a widget for a table row from a node */
-	TSharedRef<ITableRow> MakeTableRowWidget(FSCSEditorTreeNodePtrType InNodePtr, const TSharedRef<STableViewBase>& OwnerTable);
+	UE_API TSharedRef<ITableRow> MakeTableRowWidget(FSCSEditorTreeNodePtrType InNodePtr, const TSharedRef<STableViewBase>& OwnerTable);
 
 	/** Used by tree control - get children for a specified node */
-	void OnGetChildrenForTree(FSCSEditorTreeNodePtrType InNodePtr, TArray<FSCSEditorTreeNodePtrType>& OutChildren);
+	UE_API void OnGetChildrenForTree(FSCSEditorTreeNodePtrType InNodePtr, TArray<FSCSEditorTreeNodePtrType>& OutChildren);
 
 	/** Returns true if editing is allowed */
-	bool IsEditingAllowed() const;
+	UE_API bool IsEditingAllowed() const;
 
 	/** gets the actor root */
-	FSCSEditorActorNodePtrType GetActorNode() const;
+	UE_API FSCSEditorActorNodePtrType GetActorNode() const;
 
 	/** get the root scene node */
-	FSCSEditorTreeNodePtrType GetSceneRootNode() const;
+	UE_API FSCSEditorTreeNodePtrType GetSceneRootNode() const;
 
-	void SetSceneRootNode(FSCSEditorTreeNodePtrType NewSceneRootNode);
+	UE_API void SetSceneRootNode(FSCSEditorTreeNodePtrType NewSceneRootNode);
 
 	/** Adds a component to the SCS tree */
 	struct FAddNewComponentParams
@@ -1090,7 +1092,7 @@ public:
 	   @param Asset       					(In) Optional asset to assign to the component
 	   @param Params       					(In) Parameter block of optional behavior flags
 	 */
-	UActorComponent* AddNewComponent(UClass* NewComponentClass, UObject* Asset, const FAddNewComponentParams Params = FAddNewComponentParams());
+	UE_API UActorComponent* AddNewComponent(UClass* NewComponentClass, UObject* Asset, const FAddNewComponentParams Params = FAddNewComponentParams());
 
 	struct FAddedNodeDetails
 	{
@@ -1105,7 +1107,7 @@ public:
 	   @param Asset		(In) Optional asset to assign to the component
 	   @param bMarkBlueprintModified (In) Whether or not to mark the Blueprint as structurally modified
 	   @param bSetFocusToNewItem (In) Select the new item and activate the inline rename widget (default is true) */
-	void AddNewNode(FAddedNodeDetails& OutNodeDetails, TUniquePtr<FScopedTransaction> OngoingCreateTransaction, USCS_Node* NewNode, UObject* Asset, bool bMarkBlueprintModified, bool bSetFocusToNewItem = true);
+	UE_API void AddNewNode(FAddedNodeDetails& OutNodeDetails, TUniquePtr<FScopedTransaction> OngoingCreateTransaction, USCS_Node* NewNode, UObject* Asset, bool bMarkBlueprintModified, bool bSetFocusToNewItem = true);
 
 	/** Adds a new SCS Node to the component Table
 	   @param OngoingCreateTransaction (In) The transaction containing the creation of the node. The transaction will remain ongoing until the node gets its initial name from user.
@@ -1114,7 +1116,7 @@ public:
 	   @param bMarkBlueprintModified (In) Whether or not to mark the Blueprint as structurally modified
 	   @param bSetFocusToNewItem (In) Select the new item and activate the inline rename widget (default is true)
 	   @return The reference of the newly created ActorComponent */
-	UActorComponent* AddNewNode(TUniquePtr<FScopedTransaction> OngoingCreateTransaction, USCS_Node* NewNode, UObject* Asset, bool bMarkBlueprintModified, bool bSetFocusToNewItem = true);
+	UE_API UActorComponent* AddNewNode(TUniquePtr<FScopedTransaction> OngoingCreateTransaction, USCS_Node* NewNode, UObject* Asset, bool bMarkBlueprintModified, bool bSetFocusToNewItem = true);
 
 	/** Adds a new component instance node to the component Table
 		@param OngoingCreateTransaction (In) The transaction containing the creation of the node. The transaction will remain ongoing until the node gets its initial name from user.
@@ -1122,66 +1124,66 @@ public:
 		@param InParentNodePtr (In) The node this component will be added to
 		@param Asset (In) Optional asset to assign to the component
 		@param bSetFocusToNewItem (In) Select the new item and activate the inline rename widget (default is true) */
-	void AddNewNodeForInstancedComponent(TUniquePtr<FScopedTransaction> OngoingCreateTransaction, UActorComponent* NewInstanceComponent, FSCSEditorTreeNodePtrType InParentNodePtr, UObject* Asset, bool bSetFocusToNewItem = true);
+	UE_API void AddNewNodeForInstancedComponent(TUniquePtr<FScopedTransaction> OngoingCreateTransaction, UActorComponent* NewInstanceComponent, FSCSEditorTreeNodePtrType InParentNodePtr, UObject* Asset, bool bSetFocusToNewItem = true);
 	
 	/** Returns true if the specified component is currently selected */
-	bool IsComponentSelected(const UPrimitiveComponent* PrimComponent) const;
+	UE_API bool IsComponentSelected(const UPrimitiveComponent* PrimComponent) const;
 
 	/** Assigns a selection override delegate to the specified component */
-	void SetSelectionOverride(UPrimitiveComponent* PrimComponent) const;
+	UE_API void SetSelectionOverride(UPrimitiveComponent* PrimComponent) const;
 
 	/** Cut selected node(s) */
-	void CutSelectedNodes();
-	bool CanCutNodes() const;
+	UE_API void CutSelectedNodes();
+	UE_API bool CanCutNodes() const;
 
 	/** Copy selected node(s) */
-	void CopySelectedNodes();
-	bool CanCopyNodes() const;
+	UE_API void CopySelectedNodes();
+	UE_API bool CanCopyNodes() const;
 
 	/** Pastes previously copied node(s) */
-	void PasteNodes();
-	bool CanPasteNodes() const;
+	UE_API void PasteNodes();
+	UE_API bool CanPasteNodes() const;
 
 	/** Callbacks to duplicate the selected component */
-	bool CanDuplicateComponent() const;
-	void OnDuplicateComponent();
+	UE_API bool CanDuplicateComponent() const;
+	UE_API void OnDuplicateComponent();
 
 	/** Removes existing selected component nodes from the SCS */
-	void OnDeleteNodes();
-	bool CanDeleteNodes() const;
+	UE_API void OnDeleteNodes();
+	UE_API bool CanDeleteNodes() const;
 
 	/** Callbacks to find references of the selected component */
-	void OnFindReferences(bool bSearchAllBlueprints, const EGetFindReferenceSearchStringFlags Flags);
+	UE_API void OnFindReferences(bool bSearchAllBlueprints, const EGetFindReferenceSearchStringFlags Flags);
 
 	/** Removes an existing component node from the tree */
-	void RemoveComponentNode(FSCSEditorTreeNodePtrType InNodePtr);
+	UE_API void RemoveComponentNode(FSCSEditorTreeNodePtrType InNodePtr);
 
 	/** Called when selection in the tree changes */
-	void OnTreeSelectionChanged(FSCSEditorTreeNodePtrType InSelectedNodePtr, ESelectInfo::Type SelectInfo);
+	UE_API void OnTreeSelectionChanged(FSCSEditorTreeNodePtrType InSelectedNodePtr, ESelectInfo::Type SelectInfo);
 
 	/** Called when the Actor is selected. */
-	void OnActorSelected(const ECheckBoxState NewCheckedState);
+	UE_API void OnActorSelected(const ECheckBoxState NewCheckedState);
 
 	/** Called to determine if actor is selected. */
-	ECheckBoxState OnIsActorSelected() const;
+	UE_API ECheckBoxState OnIsActorSelected() const;
 
 	/** Update any associated selection (e.g. details view) from the passed in nodes */
-	void UpdateSelectionFromNodes(const TArray<FSCSEditorTreeNodePtrType> &SelectedNodes );
+	UE_API void UpdateSelectionFromNodes(const TArray<FSCSEditorTreeNodePtrType> &SelectedNodes );
 
 	/** Refresh the tree control to reflect changes in the SCS */
-	void UpdateTree(bool bRegenerateTreeNodes = true);
+	UE_API void UpdateTree(bool bRegenerateTreeNodes = true);
 
 	/** Dumps out the tree view contents to the log (used to assist with debugging widget hierarchy issues) */
-	void DumpTree();
+	UE_API void DumpTree();
 
 	/** Forces the details panel to refresh on the same objects */
-	void RefreshSelectionDetails();
+	UE_API void RefreshSelectionDetails();
 
 	/** Clears the current selection */
-	void ClearSelection();
+	UE_API void ClearSelection();
 
 	/** Get the currently selected tree nodes */
-	TArray<FSCSEditorTreeNodePtrType> GetSelectedNodes() const;
+	UE_API TArray<FSCSEditorTreeNodePtrType> GetSelectedNodes() const;
 
 	/** Get the number of currently selected tree nodes */
 	int32 GetNumSelectedNodes() const { return SCSTreeWidget->GetNumItemsSelected(); }
@@ -1194,7 +1196,7 @@ public:
 	 * @param CanExecuteActionDelegate			the delegate to query whether or not to execute the UI action
 	 * @param GetSelectedObjectsDelegate		the delegate to fill the currently select variables / components
 	 */
-	static void BuildMenuEventsSection( FMenuBuilder& Menu, UBlueprint* Blueprint, UClass* SelectedClass, FCanExecuteAction CanExecuteActionDelegate, FGetSelectedObjectsDelegate GetSelectedObjectsDelegate );
+	static UE_API void BuildMenuEventsSection( FMenuBuilder& Menu, UBlueprint* Blueprint, UClass* SelectedClass, FCanExecuteAction CanExecuteActionDelegate, FGetSelectedObjectsDelegate GetSelectedObjectsDelegate );
 
 	/**
 	 * Given an actor component, attempts to find an associated tree node.
@@ -1203,13 +1205,13 @@ public:
 	 * @param bIncludeAttachedComponents Whether or not to include components attached to each node in the search (default is true).
 	 * @return A shared pointer to a tree node. The pointer will be invalid if no match could be found.
 	 */
-	FSCSEditorTreeNodePtrType GetNodeFromActorComponent(const UActorComponent* ActorComponent, bool bIncludeAttachedComponents = true) const;
+	UE_API FSCSEditorTreeNodePtrType GetNodeFromActorComponent(const UActorComponent* ActorComponent, bool bIncludeAttachedComponents = true) const;
 
 	/** Select the root of the tree */
-	void SelectRoot();
+	UE_API void SelectRoot();
 
 	/** Select the given tree node */
-	void SelectNode(FSCSEditorTreeNodePtrType InNodeToSelect, bool IsCntrlDown);
+	UE_API void SelectNode(FSCSEditorTreeNodePtrType InNodeToSelect, bool IsCntrlDown);
 
 	/**
 	 * Set the expansion state of a node
@@ -1217,7 +1219,7 @@ public:
 	 * @param InNodeToChange	The node to be expanded/collapsed
 	 * @param bIsExpanded		True to expand the node, false to collapse it
 	 */
-	void SetNodeExpansionState(FSCSEditorTreeNodePtrType InNodeToChange, const bool bIsExpanded);
+	UE_API void SetNodeExpansionState(FSCSEditorTreeNodePtrType InNodeToChange, const bool bIsExpanded);
 
 	/**
 	 * Highlight a tree node and, optionally, a property with in it
@@ -1226,108 +1228,108 @@ public:
 	 * @param Property	The name of the property to be highlighted in the details view
 	 * @return True if the node was found in this Editor, otherwise false
 	 */
-	void HighlightTreeNode( FName TreeNodeName, const class FPropertyPath& Property );
+	UE_API void HighlightTreeNode( FName TreeNodeName, const class FPropertyPath& Property );
 	/**
 	 * Highlight a tree node and, optionally, a property with in it
 	 *
 	 * @param Node		 A Reference to the Node SCS_Node to be highlighted
 	 * @param Property	The name of the property to be highlighted in the details view
 	 */
-	void HighlightTreeNode( const USCS_Node* Node, FName Property );
+	UE_API void HighlightTreeNode( const USCS_Node* Node, FName Property );
 
 	/**
 	 * Function to save current state of SimpleConstructionScript and nodes associated with it.
 	 *
 	 * @param: SimpleContructionScript object reference.
 	 */
-	static void SaveSCSCurrentState( USimpleConstructionScript* SCSObj );
+	static UE_API void SaveSCSCurrentState( USimpleConstructionScript* SCSObj );
 
 	/**
 	 * Function to save the current state of SCS_Node and its children
 	 *
 	 * @param: Reference of the SCS_Node to be saved
 	 */
-	static void SaveSCSNode(USCS_Node* Node);
+	static UE_API void SaveSCSNode(USCS_Node* Node);
 
 	/** Is this node still used by the Simple Construction Script */
-	bool IsNodeInSimpleConstructionScript(USCS_Node* Node) const;
+	UE_API bool IsNodeInSimpleConstructionScript(USCS_Node* Node) const;
 
 	/**
 	 * Fills the supplied array with the currently selected objects
 	 * @param OutSelectedItems The array to fill.
 	 */
-	void GetSelectedItemsForContextMenu(TArray<FComponentEventConstructionData>& OutSelectedItems) const;
+	UE_API void GetSelectedItemsForContextMenu(TArray<FComponentEventConstructionData>& OutSelectedItems) const;
 
 	/** @return Array of the editable objects selected in the tree */
-	TArray<UObject*> GetSelectedEditableObjects() const;
+	UE_API TArray<UObject*> GetSelectedEditableObjects() const;
 
 	/** Provides access to the Blueprint context that's being edited */
-	class UBlueprint* GetBlueprint() const;
+	UE_API class UBlueprint* GetBlueprint() const;
 
 	/** @return The current editor mode (editing live actors or editing blueprints) */
 	EComponentEditorMode::Type GetEditorMode() const { return EditorMode; }
 
 	/** Try to handle a drag-drop operation */
-	FReply TryHandleAssetDragDropOperation(const FDragDropEvent& DragDropEvent);
+	UE_API FReply TryHandleAssetDragDropOperation(const FDragDropEvent& DragDropEvent);
 
 	/** Handler for recursively expanding/collapsing items */
-	void SetItemExpansionRecursive(FSCSEditorTreeNodePtrType Model, bool bInExpansionState);
+	UE_API void SetItemExpansionRecursive(FSCSEditorTreeNodePtrType Model, bool bInExpansionState);
 
 	/** Callback for the action trees to get the filter text */
-	FText GetFilterText() const;
+	UE_API FText GetFilterText() const;
 
 	/** Called at the end of each frame. */
-	void OnPostTick(float);
+	UE_API void OnPostTick(float);
 
 	/** Return the button widgets that can add components or create/edit blueprints */
-	TSharedPtr<SWidget> GetToolButtonsBox();
+	UE_API TSharedPtr<SWidget> GetToolButtonsBox();
 
 	/** Sets UI customizations of this SCSEditor. */
-	void SetUICustomization(TSharedPtr<ISCSEditorUICustomization> InUICustomization);
+	UE_API void SetUICustomization(TSharedPtr<ISCSEditorUICustomization> InUICustomization);
 
 protected:
-	FSCSEditorTreeNodePtrType FindOrCreateParentForExistingComponent(UActorComponent* InActorComponent, FSCSEditorActorNodePtrType ActorRootNode);
-	FSCSEditorTreeNodePtrType FindParentForNewComponent(UActorComponent* NewComponent) const;
-	FSCSEditorTreeNodePtrType FindParentForNewNode(USCS_Node* NewNode) const;
+	UE_API FSCSEditorTreeNodePtrType FindOrCreateParentForExistingComponent(UActorComponent* InActorComponent, FSCSEditorActorNodePtrType ActorRootNode);
+	UE_API FSCSEditorTreeNodePtrType FindParentForNewComponent(UActorComponent* NewComponent) const;
+	UE_API FSCSEditorTreeNodePtrType FindParentForNewNode(USCS_Node* NewNode) const;
 
 	/** Add a component from the selection in the combo box */
-	UActorComponent* PerformComboAddClass(TSubclassOf<UActorComponent> ComponentClass, EComponentCreateAction::Type ComponentCreateAction, UObject* AssetOverride);
+	UE_API UActorComponent* PerformComboAddClass(TSubclassOf<UActorComponent> ComponentClass, EComponentCreateAction::Type ComponentCreateAction, UObject* AssetOverride);
 
 	/** Called to display context menu when right clicking on the widget */
-	TSharedPtr< SWidget > CreateContextMenu();
+	UE_API TSharedPtr< SWidget > CreateContextMenu();
 
 	/** Registers context menu by name for later access */
-	void RegisterContextMenu();
+	UE_API void RegisterContextMenu();
 
 	/** Populate context menu on the fly */
-	void PopulateContextMenu(UToolMenu* InMenu);
+	UE_API void PopulateContextMenu(UToolMenu* InMenu);
 
 	/** Called when the level editor requests a component to be renamed. */
-	void OnLevelComponentRequestRename(const UActorComponent* InComponent);
+	UE_API void OnLevelComponentRequestRename(const UActorComponent* InComponent);
 
 	/** Checks to see if renaming is allowed on the selected component */
-	bool CanRenameComponent() const;
+	UE_API bool CanRenameComponent() const;
 
 	/**
 	 * Requests a rename on the selected component just after creation so that the user can provide the initial
 	 * component name (overwriting the default generated one), which is considered part of the creation process.
 	 * @param OngoingCreateTransaction The ongoing transaction started when the component was created.
 	 */
-	void OnRenameComponent(TUniquePtr<FScopedTransaction> OngoingCreateTransaction);
+	UE_API void OnRenameComponent(TUniquePtr<FScopedTransaction> OngoingCreateTransaction);
 
 	/**
 	 * Requests a rename on the selected component.
 	 */
-	void OnRenameComponent();
+	UE_API void OnRenameComponent();
 
 	/** Called when component objects are replaced following construction script execution */
-	void OnObjectsReplaced(const TMap<UObject*, UObject*>& OldToNewInstanceMap);
+	UE_API void OnObjectsReplaced(const TMap<UObject*, UObject*>& OldToNewInstanceMap);
 
 	/** Helper method to update component pointers held by the given actor node's subtree */
-	void ReplaceComponentReferencesInTree(FSCSEditorActorNodePtrType InActorNode, const TMap<UObject*, UObject*>& OldToNewInstanceMap, bool& OutHasChanges);
+	UE_API void ReplaceComponentReferencesInTree(FSCSEditorActorNodePtrType InActorNode, const TMap<UObject*, UObject*>& OldToNewInstanceMap, bool& OutHasChanges);
 
 	/** Update component pointers held by tree nodes if components have been replaced following construction script execution */
-	void ReplaceComponentReferencesInTree(const TArray<FSCSEditorTreeNodePtrType>& Nodes, const TMap<UObject*, UObject*>& OldToNewInstanceMap, bool& OutHasChanges);
+	UE_API void ReplaceComponentReferencesInTree(const TArray<FSCSEditorTreeNodePtrType>& Nodes, const TMap<UObject*, UObject*>& OldToNewInstanceMap, bool& OutHasChanges);
 
 	/**
 	 * Function to create events for the current selection
@@ -1336,7 +1338,7 @@ protected:
 	 * @param GetSelectedObjectsDelegate	the delegate to gather information about current selection
 	 * @param NodeIndex						an index to a specified node to add event for or < 0 for all selected nodes.
 	 */
-	static void CreateEventsForSelection(UBlueprint* Blueprint, FName EventName, FGetSelectedObjectsDelegate GetSelectedObjectsDelegate);
+	static UE_API void CreateEventsForSelection(UBlueprint* Blueprint, FName EventName, FGetSelectedObjectsDelegate GetSelectedObjectsDelegate);
 
 	/**
 	 * Function to construct an event for a node
@@ -1344,7 +1346,7 @@ protected:
 	 * @param EventName						the event to add
 	 * @param EventData						the event data structure describing the node
 	 */
-	static void ConstructEvent(UBlueprint* Blueprint, const FName EventName, const FComponentEventConstructionData EventData);
+	static UE_API void ConstructEvent(UBlueprint* Blueprint, const FName EventName, const FComponentEventConstructionData EventData);
 
 	/**
 	 * Function to view an event for a node
@@ -1352,76 +1354,76 @@ protected:
 	 * @param EventName						the event to view
 	 * @param EventData						the event data structure describing the node
 	 */
-	static void ViewEvent(UBlueprint* Blueprint, const FName EventName, const FComponentEventConstructionData EventData);
+	static UE_API void ViewEvent(UBlueprint* Blueprint, const FName EventName, const FComponentEventConstructionData EventData);
 
 	/** Helper method to add a tree node for the given SCS node */
-	FSCSEditorTreeNodePtrType AddTreeNode(USCS_Node* InSCSNode, FSCSEditorTreeNodePtrType InParentNodePtr, const bool bIsInheritedSCS);
+	UE_API FSCSEditorTreeNodePtrType AddTreeNode(USCS_Node* InSCSNode, FSCSEditorTreeNodePtrType InParentNodePtr, const bool bIsInheritedSCS);
 
 	/** Helper method to add a tree node for the given actor component */
-	FSCSEditorTreeNodePtrType AddTreeNodeFromComponent(UActorComponent* InSceneComponent, FSCSEditorTreeNodePtrType InParentTreeNode = FSCSEditorTreeNodePtrType());
+	UE_API FSCSEditorTreeNodePtrType AddTreeNodeFromComponent(UActorComponent* InSceneComponent, FSCSEditorTreeNodePtrType InParentTreeNode = FSCSEditorTreeNodePtrType());
 	
 	/** Helper method to add a tree node for the given node's child actor node, if present */
-	FSCSEditorTreeNodePtrType AddTreeNodeFromChildActor(FSCSEditorTreeNodePtrType InNodePtr);
+	UE_API FSCSEditorTreeNodePtrType AddTreeNodeFromChildActor(FSCSEditorTreeNodePtrType InNodePtr);
 	
 	/** Helper method to recursively find a tree node for the given SCS node starting at the given tree node */
-	FSCSEditorTreeNodePtrType FindTreeNode(const USCS_Node* InSCSNode, FSCSEditorTreeNodePtrType InStartNodePtr = FSCSEditorTreeNodePtrType()) const;
+	UE_API FSCSEditorTreeNodePtrType FindTreeNode(const USCS_Node* InSCSNode, FSCSEditorTreeNodePtrType InStartNodePtr = FSCSEditorTreeNodePtrType()) const;
 
 	/** Helper method to recursively find a tree node for the given scene component starting at the given tree node */
-	FSCSEditorTreeNodePtrType FindTreeNode(const UActorComponent* InComponent, FSCSEditorTreeNodePtrType InStartNodePtr = FSCSEditorTreeNodePtrType()) const;
+	UE_API FSCSEditorTreeNodePtrType FindTreeNode(const UActorComponent* InComponent, FSCSEditorTreeNodePtrType InStartNodePtr = FSCSEditorTreeNodePtrType()) const;
 
 	/** Helper method to recursively find a tree node for the given variable or instance name starting at the given tree node */
-	FSCSEditorTreeNodePtrType FindTreeNode(const FName& InVariableOrInstanceName, FSCSEditorTreeNodePtrType InStartNodePtr = FSCSEditorTreeNodePtrType()) const;
+	UE_API FSCSEditorTreeNodePtrType FindTreeNode(const FName& InVariableOrInstanceName, FSCSEditorTreeNodePtrType InStartNodePtr = FSCSEditorTreeNodePtrType()) const;
 
 	/** Callback when a component item is scrolled into view */
-	void OnItemScrolledIntoView( FSCSEditorTreeNodePtrType InItem, const TSharedPtr<ITableRow>& InWidget);
+	UE_API void OnItemScrolledIntoView( FSCSEditorTreeNodePtrType InItem, const TSharedPtr<ITableRow>& InWidget);
 
 	/** Callback when a component item is double clicked. */
-	void HandleItemDoubleClicked(FSCSEditorTreeNodePtrType InItem);
+	UE_API void HandleItemDoubleClicked(FSCSEditorTreeNodePtrType InItem);
 
 	/** Recursively visits the given node + its children and invokes the given function for each. */
-	void DepthFirstTraversal(const FSCSEditorTreeNodePtrType& InNodePtr, TSet<FSCSEditorTreeNodePtrType>& OutVisitedNodes, const TFunctionRef<void(const FSCSEditorTreeNodePtrType&)> InFunction) const;
+	UE_API void DepthFirstTraversal(const FSCSEditorTreeNodePtrType& InNodePtr, TSet<FSCSEditorTreeNodePtrType>& OutVisitedNodes, const TFunctionRef<void(const FSCSEditorTreeNodePtrType&)> InFunction) const;
 
 	/** Returns the set of expandable nodes that are currently collapsed in the UI */
-	void GetCollapsedNodes(const FSCSEditorTreeNodePtrType& InNodePtr, TSet<FSCSEditorTreeNodePtrType>& OutCollapsedNodes) const;
+	UE_API void GetCollapsedNodes(const FSCSEditorTreeNodePtrType& InNodePtr, TSet<FSCSEditorTreeNodePtrType>& OutCollapsedNodes) const;
 
 	/** @return The visibility of the promote to blueprint button (only visible with an actor instance that is not created from a blueprint)*/
-	EVisibility GetPromoteToBlueprintButtonVisibility() const;
+	UE_API EVisibility GetPromoteToBlueprintButtonVisibility() const;
 
 	/** @return The visibility of the Edit Blueprint button (only visible with an actor instance that is created from a blueprint)*/
-	EVisibility GetEditBlueprintButtonVisibility() const;
+	UE_API EVisibility GetEditBlueprintButtonVisibility() const;
 
 	/** @return The visibility of the Add Component combo button */
-	EVisibility GetComponentClassComboButtonVisibility() const;
+	UE_API EVisibility GetComponentClassComboButtonVisibility() const;
 
 	/** @return The visibility of the components tree */
-	EVisibility GetComponentsTreeVisibility() const;
+	UE_API EVisibility GetComponentsTreeVisibility() const;
 
 	/** @return The visibility of the components filter box */
-	EVisibility GetComponentsFilterBoxVisibility() const;
+	UE_API EVisibility GetComponentsFilterBoxVisibility() const;
 
 	/** @return the tooltip describing how many properties will be applied to the blueprint */
-	FText OnGetApplyChangesToBlueprintTooltip() const;
+	UE_API FText OnGetApplyChangesToBlueprintTooltip() const;
 
 	/** @return the tooltip describing how many properties will be reset to the blueprint default*/
-	FText OnGetResetToBlueprintDefaultsTooltip() const;
+	UE_API FText OnGetResetToBlueprintDefaultsTooltip() const;
 
 	/** Opens the blueprint editor for the blueprint being viewed by the scseditor */
-	void OnOpenBlueprintEditor(bool bForceCodeEditing) const;
+	UE_API void OnOpenBlueprintEditor(bool bForceCodeEditing) const;
 
 	/** Propagates instance changes to the blueprint */
-	void OnApplyChangesToBlueprint() const;
+	UE_API void OnApplyChangesToBlueprint() const;
 
 	/** Resets instance changes to the blueprint default */
-	void OnResetToBlueprintDefaults();
+	UE_API void OnResetToBlueprintDefaults();
 
 	/** Converts the current actor instance to a blueprint */
-	void PromoteToBlueprint() const;
+	UE_API void PromoteToBlueprint() const;
 
 	/** Called when the promote to blueprint button is clicked */
-	FReply OnPromoteToBlueprintClicked();
+	UE_API FReply OnPromoteToBlueprintClicked();
 
 	/** gets a root nodes of the tree */
-	const TArray<FSCSEditorTreeNodePtrType>& GetRootNodes() const;
+	UE_API const TArray<FSCSEditorTreeNodePtrType>& GetRootNodes() const;
 
 	/**
 	 * Creates a new C++ component from the specified class type
@@ -1429,7 +1431,7 @@ protected:
 	 *
 	 * @return The new class that was created
 	 */
-	UClass* CreateNewCPPComponent(TSubclassOf<UActorComponent> ComponentClass);
+	UE_API UClass* CreateNewCPPComponent(TSubclassOf<UActorComponent> ComponentClass);
 	
 	/**
 	 * Creates a new Blueprint component from the specified class type
@@ -1437,23 +1439,23 @@ protected:
 	 *
 	 * @return The new class that was created
 	 */
-	UClass* CreateNewBPComponent(TSubclassOf<UActorComponent> ComponentClass);
+	UE_API UClass* CreateNewBPComponent(TSubclassOf<UActorComponent> ComponentClass);
 
 	/** Recursively updates the filtered state for each component item */
-	void OnFilterTextChanged(const FText& InFilterText);
+	UE_API void OnFilterTextChanged(const FText& InFilterText);
 
 	/** 
 	 * Compares the filter bar's text with the item's component name. Use 
 	 * bRecursive to refresh the state of child nodes as well. Returns true if 
 	 * the node is set to be filtered out 
 	 */
-	bool RefreshFilteredState(FSCSEditorTreeNodePtrType TreeNode, bool bRecursive);
+	UE_API bool RefreshFilteredState(FSCSEditorTreeNodePtrType TreeNode, bool bRecursive);
 
 	/** Helper method to construct the subtree for the given actor (root) node. */
-	void BuildSubTreeForActorNode(FSCSEditorActorNodePtrType InActorNode);
+	UE_API void BuildSubTreeForActorNode(FSCSEditorActorNodePtrType InActorNode);
 
 	/** @return Type of component to filter the tree view with or nullptr if there's no filter. */
-	TSubclassOf<UActorComponent> GetComponentTypeFilterToApply() const;
+	UE_API TSubclassOf<UActorComponent> GetComponentTypeFilterToApply() const;
 
 public:
 	/** Tree widget */
@@ -1499,7 +1501,7 @@ public:
 	FSimpleDelegate OnObjectReplaced;
 
 	/** Returns the Actor context for which we are viewing/editing the SCS.  Can return null.  Should not be cached as it may change from frame to frame. */
-	class AActor* GetActorContext() const;
+	UE_API class AActor* GetActorContext() const;
 
 private:
 	/** Indicates which editor mode we're in. */
@@ -1534,3 +1536,5 @@ private:
 
 	/** SCSEditor UI extension */
 	TSharedPtr<class SExtensionPanel> ExtensionPanel; };
+
+#undef UE_API
diff --git a/Engine/Source/Editor/Kismet/Public/SSCSEditorMenuContext.h b/Engine/Source/Editor/Kismet/Public/SSCSEditorMenuContext.h
--- a/Engine/Source/Editor/Kismet/Public/SSCSEditorMenuContext.h
+++ b/Engine/Source/Editor/Kismet/Public/SSCSEditorMenuContext.h
@@ -10,8 +10,8 @@ class SSCSEditor;
 class SSubobjectEditor;
 
 class UE_DEPRECATED(5.0, "USSCSEditorMenuContext has been deprecated, use USubobjectEditorMenuContext instead.") USSCSEditorMenuContext;
-UCLASS()
-class KISMET_API USSCSEditorMenuContext : public UObject
+UCLASS(MinimalAPI)
+class USSCSEditorMenuContext : public UObject
 {
 	GENERATED_BODY()
 public:
diff --git a/Engine/Source/Editor/Kismet/Public/SubobjectEditorExtensionContext.h b/Engine/Source/Editor/Kismet/Public/SubobjectEditorExtensionContext.h
--- a/Engine/Source/Editor/Kismet/Public/SubobjectEditorExtensionContext.h
+++ b/Engine/Source/Editor/Kismet/Public/SubobjectEditorExtensionContext.h
@@ -14,8 +14,8 @@ class SSubobjectBlueprintEditor;
 class SSubobjectEditor;
 class SSubobjectInstanceEditor;
 
-UCLASS()
-class KISMET_API USubobjectEditorExtensionContext : public UObject
+UCLASS(MinimalAPI)
+class USubobjectEditorExtensionContext : public UObject
 {
 	GENERATED_BODY()
 
diff --git a/Engine/Source/Editor/KismetWidgets/Public/CreateBlueprintFromActorDialog.h b/Engine/Source/Editor/KismetWidgets/Public/CreateBlueprintFromActorDialog.h
--- a/Engine/Source/Editor/KismetWidgets/Public/CreateBlueprintFromActorDialog.h
+++ b/Engine/Source/Editor/KismetWidgets/Public/CreateBlueprintFromActorDialog.h
@@ -37,7 +37,7 @@ public:
 	static KISMETWIDGETS_API void OpenDialog(ECreateBlueprintFromActorMode CreateMode, AActor* InActorOverride = nullptr, bool bInReplaceActors = true);
 
 	UE_DEPRECATED(4.25, "Use version of OpenDialog that takes the CreateMode enum")
-	static KISMETWIDGETS_API void OpenDialog(bool bInHarvest, AActor* InActorOverride = nullptr)
+	static void OpenDialog(bool bInHarvest, AActor* InActorOverride = nullptr)
 	{
 		OpenDialog(bInHarvest ? ECreateBlueprintFromActorMode::Harvest : ECreateBlueprintFromActorMode::Subclass, InActorOverride);
 	}
diff --git a/Engine/Source/Editor/KismetWidgets/Public/SKismetLinearExpression.h b/Engine/Source/Editor/KismetWidgets/Public/SKismetLinearExpression.h
--- a/Engine/Source/Editor/KismetWidgets/Public/SKismetLinearExpression.h
+++ b/Engine/Source/Editor/KismetWidgets/Public/SKismetLinearExpression.h
@@ -10,6 +10,8 @@
 #include "Widgets/SCompoundWidget.h"
 #include "Widgets/SWidget.h"
 
+#define UE_API KISMETWIDGETS_API
+
 class SWidget;
 class UEdGraphNode;
 class UEdGraphPin;
@@ -17,7 +19,7 @@ class UEdGraphPin;
 //////////////////////////////////////////////////////////////////////////
 // SKismetLinearExpression
 
-class KISMETWIDGETS_API SKismetLinearExpression : public SCompoundWidget
+class SKismetLinearExpression : public SCompoundWidget
 {
 public:
 	SLATE_BEGIN_ARGS( SKismetLinearExpression )
@@ -27,14 +29,16 @@ public:
 		SLATE_ATTRIBUTE( bool, IsEditable )
 	SLATE_END_ARGS()
 public:
-	void Construct(const FArguments& InArgs, UEdGraphPin* InitialInputPin);
+	UE_API void Construct(const FArguments& InArgs, UEdGraphPin* InitialInputPin);
 
-	void SetExpressionRoot(UEdGraphPin* InputPin);
+	UE_API void SetExpressionRoot(UEdGraphPin* InputPin);
 protected:
-	TSharedRef<SWidget> MakeNodeWidget(const UEdGraphNode* Node, const UEdGraphPin* FromPin);
-	TSharedRef<SWidget> MakePinWidget(const UEdGraphPin* Pin);
+	UE_API TSharedRef<SWidget> MakeNodeWidget(const UEdGraphNode* Node, const UEdGraphPin* FromPin);
+	UE_API TSharedRef<SWidget> MakePinWidget(const UEdGraphPin* Pin);
 
 protected:
 	TAttribute<bool> IsEditable;
 	TSet<const UEdGraphNode*> VisitedNodes;
 };
+
+#undef UE_API
diff --git a/Engine/Source/Editor/KismetWidgets/Public/SLevelOfDetailBranchNode.h b/Engine/Source/Editor/KismetWidgets/Public/SLevelOfDetailBranchNode.h
--- a/Engine/Source/Editor/KismetWidgets/Public/SLevelOfDetailBranchNode.h
+++ b/Engine/Source/Editor/KismetWidgets/Public/SLevelOfDetailBranchNode.h
@@ -10,6 +10,8 @@
 #include "Widgets/SCompoundWidget.h"
 #include "Widgets/SWidget.h"
 
+#define UE_API KISMETWIDGETS_API
+
 class SWidget;
 struct FGeometry;
 
@@ -18,7 +20,7 @@ struct FGeometry;
 
 DECLARE_DELEGATE_RetVal_OneParam(TSharedRef<SWidget>, FOnGetActiveDetailSlotContent, bool);
 
-class KISMETWIDGETS_API SLevelOfDetailBranchNode : public SCompoundWidget
+class SLevelOfDetailBranchNode : public SCompoundWidget
 {
 	SLATE_BEGIN_ARGS(SLevelOfDetailBranchNode)
 		: _UseLowDetailSlot(false)
@@ -37,16 +39,16 @@ class KISMETWIDGETS_API SLevelOfDetailBranchNode : public SCompoundWidget
 	SLATE_END_ARGS()
 
 public:
-	SLevelOfDetailBranchNode();
+	UE_API SLevelOfDetailBranchNode();
 
-	void Construct(const FArguments& InArgs);
+	UE_API void Construct(const FArguments& InArgs);
 
 	// SWidget interface
-	virtual void Tick( const FGeometry& AllottedGeometry, const double InCurrentTime, const float InDeltaTime ) override;
+	UE_API virtual void Tick( const FGeometry& AllottedGeometry, const double InCurrentTime, const float InDeltaTime ) override;
 	// End of SWidget interface
 
 	// Determine whether we need to show the low-detail slot or high-detail slot
-	void RefreshLODSlotContent();
+	UE_API void RefreshLODSlotContent();
 
 protected:
 	// What kind of slot was shown last frame
@@ -64,3 +66,5 @@ protected:
 	FOnGetActiveDetailSlotContent OnGetActiveDetailSlotContent;
 
 };
+
+#undef UE_API
