[Blueprint]
- Fix missing preload that might look into unserialized redirectors

diff --git a/Engine/Source/Editor/Kismet/Private/BlueprintCompilationManager.cpp b/Engine/Source/Editor/Kismet/Private/BlueprintCompilationManager.cpp
--- a/Engine/Source/Editor/Kismet/Private/BlueprintCompilationManager.cpp
+++ b/Engine/Source/Editor/Kismet/Private/BlueprintCompilationManager.cpp
@@ -3268,6 +3268,8 @@ void FBlueprintCompilationManagerImpl::ReinstanceBatch(TArray<FReinstancingJob>&
 	// we may want to search the entire graph some day, but that would be expensive:
 	for (TObjectIterator<UObjectRedirector> Itr; Itr; ++Itr)
 	{
+		Itr->ConditionalPreload();
+
 		if (Itr->DestinationObject && 
 			Itr->DestinationObject->HasAnyFlags(RF_ClassDefaultObject|RF_ArchetypeObject))
 		{
