IncrementalCook: Revert to ignoring soft dependencies in blueprints meaning the presence of a soft reference does not trigger a recook

#rnx

diff --git a/Engine/Source/Editor/Kismet/Private/Blueprints/BlueprintDependencies.cpp b/Engine/Source/Editor/Kismet/Private/Blueprints/BlueprintDependencies.cpp
--- a/Engine/Source/Editor/Kismet/Private/Blueprints/BlueprintDependencies.cpp
+++ b/Engine/Source/Editor/Kismet/Private/Blueprints/BlueprintDependencies.cpp
@@ -106,11 +106,10 @@ static void UE::Private::GetAllImportedObjects(TConstArrayView<const UObject*> R
 	// (e.g. I depend on myself or I depend on my inner) are not reported - only external objects:
 	struct FExternalReferenceFinder : public FArchiveUObject
 	{
-		FExternalReferenceFinder(const UObject* Obj, TSet<const UObject*>& InReferences, TSet<FName>& InSoftReferences, const TSet<const UObject*>& InRoots)
+		FExternalReferenceFinder(const UObject* Obj, TSet<const UObject*>& InReferences, const TSet<const UObject*>& InRoots)
 			: FArchiveUObject()
 			, SearchRoots(InRoots)
 			, References(InReferences)
-			, SoftReferences(InSoftReferences)
 		{
 			SetIsSaving(true);
 			SetShouldSkipCompilingAssets(true);
@@ -143,23 +142,6 @@ static void UE::Private::GetAllImportedObjects(TConstArrayView<const UObject*> R
 
 		virtual FArchive& operator<<(FLazyObjectPtr& Value) override { return ShouldTraverseProperty() ? FArchiveUObject::operator<<(Value) : *this; }
 		virtual FArchive& operator<<(FObjectPtr& Value) override { return ShouldTraverseProperty() ? FArchiveUObject::operator<<(Value) : *this; }
-		virtual FArchive& operator<<(FSoftObjectPtr& Value) override 
-		{ 
-			const FSoftObjectPath& Path = Value.ToSoftObjectPath();
-			if( !Path.IsNull() && ShouldTraverseProperty() )
-			{
-				SoftReferences.Add(Path.GetLongPackageFName());
-			}
-			return *this; 
-		}
-		virtual FArchive& operator<<(FSoftObjectPath& Value) override 
-		{ 
-			if( !Value.IsNull() && ShouldTraverseProperty() )
-			{
-				SoftReferences.Add(Value.GetLongPackageFName());
-			}
-			return *this; 
-		}
 		virtual FArchive& operator<<(FWeakObjectPtr& Value) override { return ShouldTraverseProperty() ? FArchiveUObject::operator<<(Value) : *this; }
 
 		virtual FArchive& operator<<(UObject*& ObjRef) override
@@ -194,7 +176,6 @@ static void UE::Private::GetAllImportedObjects(TConstArrayView<const UObject*> R
 
 		const TSet<const UObject*>& SearchRoots;
 		TSet<const UObject*>& References;
-		TSet<FName>& SoftReferences;
 	};
 
 	TSet<const UObject*> RootSet(Roots);
@@ -221,7 +202,7 @@ static void UE::Private::GetAllImportedObjects(TConstArrayView<const UObject*> R
 
 	for(const UObject* Obj: Objects)
 	{
-		FExternalReferenceFinder References(Obj, Result, OutSoftImports, RootSet);
+		FExternalReferenceFinder References(Obj, Result, RootSet);
 	}
 	
 	OutImports = Result.Array();
