Preserve some of the overridable state when reinstancing archetypes and CDOs

diff --git a/Engine/Source/Editor/Kismet/Private/BlueprintCompilationManager.cpp b/Engine/Source/Editor/Kismet/Private/BlueprintCompilationManager.cpp
--- a/Engine/Source/Editor/Kismet/Private/BlueprintCompilationManager.cpp
+++ b/Engine/Source/Editor/Kismet/Private/BlueprintCompilationManager.cpp
@@ -2876,6 +2876,7 @@ void FBlueprintCompilationManagerImpl::ReinstanceBatch(TArray<FReinstancingJob>&
 				{
 					FBlueprintCompileReinstancer::CopyPropertiesForUnrelatedObjects(Pair.Key, Pair.Value, /*bClearExternalReferences*/true, bUseDeltaSerialization, /*bOnlyHandleDirectSubObjects*/true, &OldToNewInstanceMap, &InOutOldToNewClassMap);
 				}
+				FOverridableManager::Get().RestoreOverrideState(OldCDO, NewCDO);
 
 				if (ReinstancingJob.Compiler.IsValid())
 				{
@@ -3205,6 +3206,7 @@ void FBlueprintCompilationManagerImpl::ReinstanceBatch(TArray<FReinstancingJob>&
 					{
 						FBlueprintCompileReinstancer::CopyPropertiesForUnrelatedObjects(Pair.Key, Pair.Value, /*bClearExternalReferences*/false, bUseDeltaSerialization, /*bOnlyHandleDirectSubObjects*/true, &OldToNewInstanceMap, &InOutOldToNewClassMap);
 					}
+					FOverridableManager::Get().RestoreOverrideState(OldInstance, NewArchetype);
 
 					if (OldToNewTemplates && !OldToNewInstanceMap.IsEmpty())
 					{
