Aggressively clear K2Node_FunctionEntry's variable cache before skeleton classes are regenerated

[RN] Fixed rare crash when cooking blueprints that have local variables

diff --git a/Engine/Source/Editor/Kismet/Private/BlueprintCompilationManager.cpp b/Engine/Source/Editor/Kismet/Private/BlueprintCompilationManager.cpp
--- a/Engine/Source/Editor/Kismet/Private/BlueprintCompilationManager.cpp
+++ b/Engine/Source/Editor/Kismet/Private/BlueprintCompilationManager.cpp
@@ -189,6 +189,7 @@ struct FBlueprintCompilationManagerImpl : public FGCObject
 	void VerifyNoQueuedRequests(const TArray<FCompilerData>& CurrentlyCompilingBPs);
 	void SaveConstructionHashes(const TArray<FCompilerData>& CurrentlyCompilingBPs);
 	void EnsureArchetypesArePostLoaded();
+	void DiscardFunctionEntryCaches(const TArray<FCompilerData>& CurrentlyCompilingBPs);
 
 	static void GatherArchetypesRequiringPostload(TConstArrayView<FBPCompileRequest> FromRequests, TArray<UObject*>& OutArchetypesRequiringPostload);
 	static void ReparentHierarchies(const TMap<UClass*, UClass*>& OldClassToNewClass, EReparentClassOptions Options);
@@ -1022,10 +1023,11 @@ void FBlueprintCompilationManagerImpl::FlushCompilationQueueImpl(bool bSuppressB
 
 		SlowTask.EnterProgressFrame();
 
-		// STAGE VII: safely throw away old skeleton CDOs:
+		// STAGE VII: safely throw away old skeleton CDOs and function entry caches::
 		{
 			using namespace UE::Kismet::BlueprintCompilationManager;
 
+			DiscardFunctionEntryCaches(CurrentlyCompilingBPs);
 			TMap<UClass*, UClass*> NewSkeletonToOldSkeleton;
 			for (FCompilerData& CompilerData : CurrentlyCompilingBPs)
 			{
@@ -2270,6 +2272,26 @@ void FBlueprintCompilationManagerImpl::EnsureArchetypesArePostLoaded()
 	// upon the mutating archetype
 }
 
+void FBlueprintCompilationManagerImpl::DiscardFunctionEntryCaches(const TArray<FCompilerData>& CurrentlyCompilingBPs)
+{
+	for (const FCompilerData& CompilerData : CurrentlyCompilingBPs)
+	{
+		UBlueprint* BP = CompilerData.BP;
+		TArray<UEdGraph*> AllGraphs;
+		BP->GetAllGraphs(AllGraphs);
+		TArray<UK2Node_FunctionEntry*> EntryNodes;
+		for (UEdGraph* Graph : AllGraphs)
+		{
+			Graph->GetNodesOfClass(EntryNodes);
+			for (UK2Node_FunctionEntry* EntryNode : EntryNodes)
+			{
+				EntryNode->ClearCachedBlueprintData(BP);
+			}
+			EntryNodes.Reset(); // typically 1 or 0 so just avoiding repeated small allocations
+		}
+	}
+}
+
 void FBlueprintCompilationManagerImpl::GatherArchetypesRequiringPostload(TConstArrayView<FBPCompileRequest> FromRequests, TArray<UObject*>& OutArchetypesRequiringPostload)
 {
 	for (const FBPCompileRequest& Request : FromRequests)
