Tag our reference finder as persistent to avoid picking up some transient containers like UBlueprint::CachedDependents - which is unstable. The dependency key must be stable.
[RN] Fixed bug that caused blueprints to have an unstable dependency key used by incremental cook jobs - this would pessimize the number of assets processed by incremental cook jobs

diff --git a/Engine/Source/Editor/Kismet/Private/Blueprints/BlueprintDependencies.cpp b/Engine/Source/Editor/Kismet/Private/Blueprints/BlueprintDependencies.cpp
--- a/Engine/Source/Editor/Kismet/Private/Blueprints/BlueprintDependencies.cpp
+++ b/Engine/Source/Editor/Kismet/Private/Blueprints/BlueprintDependencies.cpp
@@ -111,6 +111,7 @@ static void UE::Private::GetAllImportedObjects(TConstArrayView<const UObject*> R
 			, SearchRoots(InRoots)
 			, References(InReferences)
 		{
+			SetIsPersistent(true); // to avoid transient properties - these often contain unstable caches which cannot be used in key calculation!
 			SetIsSaving(true);
 			SetShouldSkipCompilingAssets(true);
 			SetWantBinaryPropertySerialization(true);
@@ -146,7 +147,8 @@ static void UE::Private::GetAllImportedObjects(TConstArrayView<const UObject*> R
 
 		virtual FArchive& operator<<(UObject*& ObjRef) override
 		{
-			if (ObjRef != nullptr &&
+			if (ShouldTraverseProperty() &&
+				ObjRef != nullptr &&
 				// uht objects are tagged as RF_Transient but they 'are native' 
 				// so still referencable:
 				(!ObjRef->HasAnyFlags(RF_Transient) || ObjRef->IsNative()) &&
