[UE] Fix a regression that leads to an ensure() in the editor when reloading a data-only Blueprint asset that inherits from a non-native parent type.

diff --git a/Engine/Source/Editor/Kismet/Private/BlueprintCompilationManager.cpp b/Engine/Source/Editor/Kismet/Private/BlueprintCompilationManager.cpp
--- a/Engine/Source/Editor/Kismet/Private/BlueprintCompilationManager.cpp
+++ b/Engine/Source/Editor/Kismet/Private/BlueprintCompilationManager.cpp
@@ -1511,20 +1511,22 @@ void FBlueprintCompilationManagerImpl::FlushCompilationQueueImpl(bool bSuppressB
 					{
 						CompilerData.Reinstancer->PropagateSparseClassDataToNewClass(BPGC);
 					}
-					// relink, generate CDO:
-					BPGC->bLayoutChanging = false;
-					BPGC->Bind();
-					BPGC->StaticLink(true);
 
 					// When compiling skeleton only, it's possible that a non-skeleton BPGC
 					// has not been compiled yet and has uninitialized class flags. Some class
 					// flags must be set in order for the CDO construction to work correctly,
 					// like for instancing subobjects. Make sure that BPGC has at least its
 					// inherited and property based class flags.
+					// Note: This will clear flags that are otherwise set at Link() time.
 					FBlueprintClassFlagUtils::PropagateParentClassFlags_CompileClassLayout(BPGC);
 					FBlueprintClassFlagUtils::PropagateParentClassFlags_FinishCompilingClass(BPGC);
 					FBlueprintClassFlagUtils::AppendPropertyBasedClassFlags(BPGC);
 
+					// relink, generate CDO:
+					BPGC->bLayoutChanging = false;
+					BPGC->Bind();
+					BPGC->StaticLink(true);
+
 					// Force recreate the CDO
 					BPGC->SetDefaultObject(nullptr);
 					BPGC->GetDefaultObject(true);
