Avoid tracking CDOs for classes labeled UCLASS(Transient) which are often used as runtime caches and thus have unstable identity which if included in a dependency key will often change

#rnx

diff --git a/Engine/Source/Editor/Kismet/Private/Blueprints/BlueprintDependencies.cpp b/Engine/Source/Editor/Kismet/Private/Blueprints/BlueprintDependencies.cpp
--- a/Engine/Source/Editor/Kismet/Private/Blueprints/BlueprintDependencies.cpp
+++ b/Engine/Source/Editor/Kismet/Private/Blueprints/BlueprintDependencies.cpp
@@ -393,8 +393,18 @@ static FString UE::Private::HashNativeClass( const UClass* NativeClass)
 	// the other core level facilities. Exporting a class to text using ExportText has
 	// always produced an empty object, which is not meaningful. JsonObjectGraph::Stringify
 	// will provide identity for the UClass:
+	TArray<const UObject*, TInlineAllocator<2>> RootObjects;
+	if(NativeClass->HasAnyClassFlags(CLASS_Transient))
+	{
+		RootObjects = {NativeClass};
+	}
+	else
+	{
+		RootObjects = {NativeClass, NativeClass->GetDefaultObject(false)};
+	}
+
 	FUtf8String StringifiedClass = UE::JsonObjectGraph::Stringify(
-		{NativeClass, NativeClass->GetDefaultObject(false)},
+		RootObjects,
 		FJsonStringifyOptions(EJsonStringifyFlags::FilterEditorOnlyData));
 	FBlake3 Hash;
 	Hash.Update(StringifiedClass.GetCharArray().GetData(), StringifiedClass.GetCharArray().Num());
