Honor IgnoreForMemberInitializationTest when canonizing structs for dependency tracking, this excludes members which are always unstable from the struct identity calculation

#rnx

diff --git a/Engine/Source/Editor/Kismet/Private/Blueprints/BlueprintDependencies.cpp b/Engine/Source/Editor/Kismet/Private/Blueprints/BlueprintDependencies.cpp
--- a/Engine/Source/Editor/Kismet/Private/Blueprints/BlueprintDependencies.cpp
+++ b/Engine/Source/Editor/Kismet/Private/Blueprints/BlueprintDependencies.cpp
@@ -88,6 +88,9 @@ private:
 	TMap<const UClass*, FBPDependencyCacheEntry> ClassToHash;
 	TMap<const UScriptStruct*, FBPDependencyCacheEntry> StructToHash;
 	TMap<FName, FBPDependencyCacheEntry> PackageToHash;
+	// when disabled we clear and return the scratch entry so that callers can always populate the cache..
+	FBPDependencyCacheEntry ScratchEntry; 
+	bool bCacheDisabled = false;
 };
 
 // Filesystem helper:
@@ -157,7 +160,6 @@ static void UE::Private::GetAllImportedObjects(TConstArrayView<const UObject*> R
 				!ObjRef->IsA<UPackage>() &&
 				!IsInRoots(ObjRef))
 			{
-				// Set to null any pointer to an external asset
 				References.Add(ObjRef);
 			}
 
@@ -364,9 +366,29 @@ static FString UE::Private::HashNativeStruct(const UScriptStruct* Struct)
 	// so use another json routine to get identity for default values
 	// shortcomings here are that we aren't getting much info about 
 	// native serialize overrides (inc. custom version bumps)
+
+	// Filter out IgnoreForMemberInitializationTest as we know those values
+	// will change every time the structure is reconstructed - without this
+	// then types that have side effects in their ctor will always report
+	// as changed:
+	FJsonObjectConverter::CustomExportCallback ExportCb = FJsonObjectConverter::CustomExportCallback::CreateLambda(
+		[](FProperty* Prop, const void* Data) -> TSharedPtr<FJsonValue>
+		{
+			static const FName NAME_IgnoreForMemberInitializationTest(TEXT("IgnoreForMemberInitializationTest"));
+			if (Prop->HasMetaData(NAME_IgnoreForMemberInitializationTest))
+			{
+				return MakeShared<FJsonValueString>(TEXT("")); // export unstable members as the empty string
+			}
+
+			return {};
+		});
 	FStructOnScope Defaults(Struct);
 	FString DefaultValues;
-	FJsonObjectConverter::UStructToJsonObjectString(Struct, Defaults.GetStructMemory(), DefaultValues);
+	FJsonObjectConverter::UStructToJsonObjectString(Struct, Defaults.GetStructMemory(), DefaultValues,
+		/*CheckFlags = */ 0, 
+		/*SkipFlags =*/ 0, 
+		/*int32 Indent =*/ 0, 
+		&ExportCb);
 	StringifiedStruct.Append(DefaultValues);
 
 	FBlake3 Hash;
@@ -445,6 +467,11 @@ static FString UE::Private::HashPackageFile(FName PackageName)
 UE::Private::FBPDependencyCacheEntry* UE::Private::FBPDependencyCache::LookupCache(const UClass* Class)
 {
 	FBPDependencyCache& Cache = GetCache();
+	if(Cache.bCacheDisabled)
+	{
+		Cache.ScratchEntry = {};
+		return &Cache.ScratchEntry;
+	}
 	FBPDependencyCacheEntry& Value = Cache.ClassToHash.FindOrAdd(Class);
 	return &Value;
 }
@@ -452,6 +479,11 @@ UE::Private::FBPDependencyCacheEntry* UE::Private::FBPDependencyCache::LookupCac
 UE::Private::FBPDependencyCacheEntry* UE::Private::FBPDependencyCache::LookupCache(const UScriptStruct* Struct)
 {
 	FBPDependencyCache& Cache = GetCache();
+	if(Cache.bCacheDisabled)
+	{
+		Cache.ScratchEntry = {};
+		return &Cache.ScratchEntry;
+	}
 	FBPDependencyCacheEntry& Value = Cache.StructToHash.FindOrAdd(Struct);
 	return &Value;
 }
@@ -459,6 +491,11 @@ UE::Private::FBPDependencyCacheEntry* UE::Private::FBPDependencyCache::LookupCac
 UE::Private::FBPDependencyCacheEntry* UE::Private::FBPDependencyCache::LookupCache(FName PackageName)
 {
 	FBPDependencyCache& Cache = GetCache();
+	if(Cache.bCacheDisabled)
+	{
+		Cache.ScratchEntry = {};
+		return &Cache.ScratchEntry;
+	}
 	FBPDependencyCacheEntry& Value = Cache.PackageToHash.FindOrAdd(PackageName);
 	return &Value;
 }
@@ -469,6 +506,13 @@ UE::Private::FBPDependencyCache& UE::Private::FBPDependencyCache::GetCache()
 	if(!GBPDependencyCache)
 	{
 		GBPDependencyCache = MakeUnique<FBPDependencyCache>();
+		// increase testability - avoid caching native dependencies when not running 
+		// the cook commandlet. This allows us to find e.g. unstable native
+		// constructors:
+		if(!IsRunningCookCommandlet())
+		{
+			GBPDependencyCache->bCacheDisabled = true;
+		}
 	}
 	return *GBPDependencyCache;
 }
