Including objects created in PreCreateSubobjectsForReinstantiation for reference fixup in archetype reinstancing

[RN] Fixed cause of loss of blueprint subobjects when a blueprint that is contained in another blueprint is compiled

diff --git a/Engine/Source/Editor/Kismet/Private/BlueprintCompilationManager.cpp b/Engine/Source/Editor/Kismet/Private/BlueprintCompilationManager.cpp
--- a/Engine/Source/Editor/Kismet/Private/BlueprintCompilationManager.cpp
+++ b/Engine/Source/Editor/Kismet/Private/BlueprintCompilationManager.cpp
@@ -3246,6 +3246,10 @@ void FBlueprintCompilationManagerImpl::ReinstanceBatch(TArray<FReinstancingJob>&
 						{
 							*NewArchetypeMapping = Pair.Value;
 						}
+						else
+						{
+							OldArchetypeToNewArchetype.Add(Pair.Key, Pair.Value);
+						}
 					}
 
 					OldInstance->RemoveFromRoot();
