Prevent Pin Value Inspector Tooltips from using global invalidation, as it can cause broken tooltips or crashes in some scenarios

diff --git a/Engine/Source/Editor/Kismet/Private/SPinValueInspector.cpp b/Engine/Source/Editor/Kismet/Private/SPinValueInspector.cpp
--- a/Engine/Source/Editor/Kismet/Private/SPinValueInspector.cpp
+++ b/Engine/Source/Editor/Kismet/Private/SPinValueInspector.cpp
@@ -310,7 +310,6 @@ TWeakPtr<FPinValueInspectorTooltip> FPinValueInspectorTooltip::SummonTooltip(FEd
 			}
 
 			TooltipWindow->ShowWindow();
-			TooltipWindow->SetAllowFastUpdate(true);
 			return Instance;
 		}
 	}
@@ -355,7 +354,6 @@ void FPinValueInspectorTooltip::DismissTooltip()
 	ValueInspectorWidget->SetPinRef(FEdGraphPinReference());
 
 	TooltipWindow->HideWindow();
-	TooltipWindow->SetAllowFastUpdate(false);
 	if (TooltipWidget->GetContentWidget() != ValueInspectorWidget)
 	{
 		TooltipWidget->SetContentWidget(ValueInspectorWidget->AsShared());
