Blueprint dependencies : Ignore package tracking in hash calculation.

During the hash calculation for blueprints, some packages can be accessed. The package access tracker will detect it and create a package dependency for each access. The problem is the hashes are cached, so the calculation is only done for the first blueprint asking for the hash. And so the tracker adds package dependencies only to this first blueprint. The first blueprint is not always the same so dependencies are non-deterministic.

#rnx

diff --git a/Engine/Source/Editor/Kismet/Private/Blueprints/BlueprintDependencies.cpp b/Engine/Source/Editor/Kismet/Private/Blueprints/BlueprintDependencies.cpp
--- a/Engine/Source/Editor/Kismet/Private/Blueprints/BlueprintDependencies.cpp
+++ b/Engine/Source/Editor/Kismet/Private/Blueprints/BlueprintDependencies.cpp
@@ -11,6 +11,7 @@
 #include "JsonObjectGraph/Stringify.h"
 #include "Kismet2/BlueprintEditorUtils.h"
 #include "Misc/Optional.h"
+#include "Misc/PackageAccessTracking.h"
 #include "Misc/ScopeExit.h"
 #include "Misc/StringOutputDevice.h"
 #include "Serialization/CompactBinary.h"
@@ -350,6 +351,14 @@ static void UE::Private::GetAllDependencies(
 
 static FString UE::Private::HashNativeStruct(const UScriptStruct* Struct)
 {
+	// The hash calculated in this function is cached. So only the first package asking for it calculates it
+	// Other packages will retrieve the cached value. To calculate the hash we might go through pointers
+	// and this will be detected by the package tracker and create package dependencies. But only the first package
+	// calculating the hash will have those dependencies and since the first package doing the calculation is not always
+	// the same we end up with non deterministic dependencies.
+	// If accessing the dereferenced object is relevant, the hash should change so tracking those dependencies is useless anyway.
+	UE_COOK_DISABLE_PACKAGE_ACCESS_TRACKING_SCOPED();
+
 	TRACE_CPUPROFILER_EVENT_SCOPE(BlueprintDependencies::HashNativeStruct);
 	FBPDependencyCacheEntry* CachedValue = FBPDependencyCache::LookupCache(Struct);
 	if(CachedValue->Hash.Len() > 0)
@@ -404,6 +413,14 @@ static FString UE::Private::HashNativeStruct(const UScriptStruct* Struct)
 
 static FString UE::Private::HashNativeClass( const UClass* NativeClass)
 {
+	// The hash calculated in this function is cached. So only the first package asking for it calculates it
+	// Other packages will retrieve the cached value. To calculate the hash we might go through pointers
+	// and this will be detected by the package tracker and create package dependencies. But only the first package
+	// calculating the hash will have those dependencies and since the first package doing the calculation is not always
+	// the same we end up with non deterministic dependencies.
+	// If accessing the dereferenced object is relevant, the hash should change so tracking those dependencies is useless anyway.
+	UE_COOK_DISABLE_PACKAGE_ACCESS_TRACKING_SCOPED();
+
 	TRACE_CPUPROFILER_EVENT_SCOPE(BlueprintDependencies::HashNativeClass);
 	FBPDependencyCacheEntry* CachedValue = FBPDependencyCache::LookupCache(NativeClass);
 	if(CachedValue->Hash.Len() > 0)
diff --git a/Engine/Source/Runtime/CoreUObject/Public/Misc/PackageAccessTracking.h b/Engine/Source/Runtime/CoreUObject/Public/Misc/PackageAccessTracking.h
--- a/Engine/Source/Runtime/CoreUObject/Public/Misc/PackageAccessTracking.h
+++ b/Engine/Source/Runtime/CoreUObject/Public/Misc/PackageAccessTracking.h
@@ -43,8 +43,16 @@ class ITargetPlatform;
 #define UE_COOK_RESULTPROJECTION_SCOPED(ProjectionName) \
 	PackageAccessTracking_Private::FPackageAccessRefScope ANONYMOUS_VARIABLE(PackageAccessTracker_)( \
 		PackageAccessTracking_Private::FPackageAccessRefScope::CookResultProjectionType, ProjectionName)
+
+/** 
+ * Disable the package tracker in the current scope. It means when the code will dereference a TObjectPtr, no
+ * package dependency will be created for the currently cooked package.  
+ */
+#define UE_COOK_DISABLE_PACKAGE_ACCESS_TRACKING_SCOPED() UE_COOK_RESULTPROJECTION_SCOPED(UE::Cook::ResultProjection::None)
+
 #else
 #define UE_COOK_RESULTPROJECTION_SCOPED(ProjectionName)
+#define UE_COOK_DISABLE_PACKAGE_ACCESS_TRACKING_SCOPED()
 #endif //WITH_EDITOR && UE_WITH_PACKAGE_ACCESS_TRACKING
 
 ///////////////////////////////////////////////////////////////////////////////////////////////////////////////////////
