Added 'Open in New Tab' functionality and context menu entries for collapsed graphs, functions and macros

diff --git a/Engine/Plugins/Runtime/RigVM/Source/RigVMEditor/Private/Editor/RigVMEditor.cpp b/Engine/Plugins/Runtime/RigVM/Source/RigVMEditor/Private/Editor/RigVMEditor.cpp
--- a/Engine/Plugins/Runtime/RigVM/Source/RigVMEditor/Private/Editor/RigVMEditor.cpp
+++ b/Engine/Plugins/Runtime/RigVM/Source/RigVMEditor/Private/Editor/RigVMEditor.cpp
@@ -44,6 +44,7 @@
 #include "Overrides/OverrideStatusDetailsObjectFilter.h"
 #include "Overrides/OverrideStatusDetailsWidgetBuilder.h"
 #include "Overrides/OverrideStatusDetailsWidgetBuilder.h"
+#include "GraphEditorActions.h"
 
 #if WITH_RIGVMLEGACYEDITOR
 #include "SKismetInspector.h"
@@ -1079,6 +1080,14 @@ void FRigVMEditorBase::CreateDefaultCommandsImpl()
 
 void FRigVMEditorBase::OnCreateGraphEditorCommands(TSharedPtr<FUICommandList> GraphEditorCommandsList)
 {
+	if (GraphEditorCommandsList.IsValid())
+	{
+		GraphEditorCommandsList->MapAction(FGraphEditorCommands::Get().OpenInNewTab,
+			FExecuteAction::CreateSP(StaticCastSharedRef<FRigVMEditorBase>(SharedRef()), &FRigVMEditorBase::OnOpenSelectedNodesInNewTab),
+			FCanExecuteAction::CreateSP(StaticCastSharedRef<FRigVMEditorBase>(SharedRef()), &FRigVMEditorBase::CanOpenSelectedNodesInNewTab),
+			FIsActionButtonVisible::CreateSP(StaticCastSharedRef<FRigVMEditorBase>(SharedRef()), &FRigVMEditorBase::CanOpenSelectedNodesInNewTab)
+		);
+	}
 }
 
 void FRigVMEditorBase::Compile()
@@ -3967,6 +3976,62 @@ void FRigVMEditorBase::SwapAssetReferences()
 	SwapFunctionsDialog->ShowNormal();
 }
 
+void FRigVMEditorBase::OnOpenSelectedNodesInNewTab()
+{
+	const TRigVMAssetInterfacePtr RigVMAsset = GetRigVMAsset();
+	if (RigVMAsset == nullptr)
+	{
+		return;
+	}
+
+	const FGraphPanelSelectionSet SelectedNodes = GetSelectedNodes();
+	for (FGraphPanelSelectionSet::TConstIterator NodeIt(SelectedNodes); NodeIt; ++NodeIt)
+	{
+		if (const URigVMEdGraphNode* EdGraphNode = Cast<URigVMEdGraphNode>(*NodeIt))
+		{
+			if (const URigVMLibraryNode* LibraryNode = Cast<URigVMLibraryNode>(EdGraphNode->GetModelNode()))
+			{
+				UEdGraph* FoundEdGraph = nullptr;
+				if (const URigVMGraph* ContainedGraph = LibraryNode->GetContainedGraph())
+				{
+					FoundEdGraph = RigVMAsset->GetEdGraph(ContainedGraph);
+				}
+				else if (const URigVMFunctionReferenceNode* FunctionReferenceNode = Cast<URigVMFunctionReferenceNode>(LibraryNode))
+				{
+					if (const URigVMLibraryNode* ReferencedNode = FunctionReferenceNode->LoadReferencedNode())
+					{
+						const URigVMFunctionEntryNode* EntryNode = ReferencedNode->GetEntryNode();
+						if (EntryNode != nullptr && EntryNode->GetGraph() != nullptr)
+						{
+							FoundEdGraph = RigVMAsset->GetEdGraph(EntryNode->GetGraph());
+						}
+					}
+				}
+
+				if (FoundEdGraph)
+				{
+					OpenDocument(FoundEdGraph, FDocumentTracker::ForceOpenNewDocument);
+				}
+			}
+		}
+	}
+}
+
+bool FRigVMEditorBase::CanOpenSelectedNodesInNewTab() const
+{
+	const FGraphPanelSelectionSet SelectedNodes = GetSelectedNodes();
+	for (FGraphPanelSelectionSet::TConstIterator NodeIt(SelectedNodes); NodeIt; ++NodeIt)
+	{
+		if (const URigVMEdGraphNode* EDGraphNode = Cast<URigVMEdGraphNode>(*NodeIt))
+		{
+			return Cast<URigVMLibraryNode>(EDGraphNode->GetModelNode()) != nullptr;
+		}
+	}
+
+	return false;
+}
+
+
 void FRigVMEditorBase::OnGraphNodeDropToPerform(TSharedPtr<FDragDropOperation> InDragDropOp, UEdGraph* InGraph, const FVector2f& InNodePosition, const FVector2f& InScreenPosition)
 {
 	DECLARE_SCOPE_HIERARCHICAL_COUNTER_FUNC()
diff --git a/Engine/Plugins/Runtime/RigVM/Source/RigVMEditor/Private/RigVMEditorModule.cpp b/Engine/Plugins/Runtime/RigVM/Source/RigVMEditor/Private/RigVMEditorModule.cpp
--- a/Engine/Plugins/Runtime/RigVM/Source/RigVMEditor/Private/RigVMEditorModule.cpp
+++ b/Engine/Plugins/Runtime/RigVM/Source/RigVMEditor/Private/RigVMEditorModule.cpp
@@ -1067,11 +1067,17 @@ void FRigVMEditorModule::GetNodeOrganizationContextMenuActions(TScriptInterface<
 	URigVMGraph* Model = ModelNode->GetGraph();
 	URigVMController* Controller = RigVMClientHost->GetRigVMClient()->GetController(Model);
 
+	FToolMenuSection& NavigationSection = Menu->AddSection("RigVMEditorContextMenuNavigation", LOCTEXT("NavigationHeader", "Navigation"));
 	FToolMenuSection& OrganizationSection = Menu->AddSection("RigVMEditorContextMenuOrganization", LOCTEXT("OrganizationHeader", "Organization"));
 
 	check(EdGraphNode);
 	check(ModelNode);
 
+	if (URigVMLibraryNode* LibraryNode = Cast<URigVMLibraryNode>(ModelNode))
+	{
+		NavigationSection.AddMenuEntry(FGraphEditorCommands::Get().OpenInNewTab);
+	}
+
 	if(ModelNode->SupportsRenaming())
 	{
 		OrganizationSection.AddMenuEntry(
diff --git a/Engine/Plugins/Runtime/RigVM/Source/RigVMEditor/Public/Editor/RigVMEditor.h b/Engine/Plugins/Runtime/RigVM/Source/RigVMEditor/Public/Editor/RigVMEditor.h
--- a/Engine/Plugins/Runtime/RigVM/Source/RigVMEditor/Public/Editor/RigVMEditor.h
+++ b/Engine/Plugins/Runtime/RigVM/Source/RigVMEditor/Public/Editor/RigVMEditor.h
@@ -561,6 +561,8 @@ protected:
 	UE_API void SwapFunctionForAssets(const TArray<FAssetData>& InAssets, bool bSetupUndo);
 	UE_API void SwapAssetReferences();
 
+	UE_API void OnOpenSelectedNodesInNewTab();
+	UE_API bool CanOpenSelectedNodesInNewTab() const;
 
 	/** Once the log is collected update the graph */
 	UE_API void UpdateGraphCompilerErrors();
diff --git a/Engine/Source/Editor/BlueprintGraph/Private/EdGraphSchema_K2.cpp b/Engine/Source/Editor/BlueprintGraph/Private/EdGraphSchema_K2.cpp
--- a/Engine/Source/Editor/BlueprintGraph/Private/EdGraphSchema_K2.cpp
+++ b/Engine/Source/Editor/BlueprintGraph/Private/EdGraphSchema_K2.cpp
@@ -1694,6 +1694,12 @@ void UEdGraphSchema_K2::GetContextMenuActions(UToolMenu* Menu, UGraphNodeContext
 				Section.AddMenuEntry(FGraphEditorCommands::Get().SelectReferenceInLevel);
 			}
 
+			{
+				FToolMenuSection& Section = Menu->AddSection("EdGraphSchemaNodeNavigation", LOCTEXT("NodeNavigationMenuHeader", "Navigation"));
+				// Opens a function, macro or collapsed graph in a new tab
+				Section.AddMenuEntry(FGraphEditorCommands::Get().OpenInNewTab);
+			}
+
 			if (!bIsDebugging)
 			{
 				// Collapse/expand nodes
diff --git a/Engine/Source/Editor/GraphEditor/Private/GraphEditorActions.cpp b/Engine/Source/Editor/GraphEditor/Private/GraphEditorActions.cpp
--- a/Engine/Source/Editor/GraphEditor/Private/GraphEditorActions.cpp
+++ b/Engine/Source/Editor/GraphEditor/Private/GraphEditorActions.cpp
@@ -91,6 +91,8 @@ void FGraphEditorCommandsImpl::RegisterCommands()
 	
 	UI_COMMAND( GoToDefinition, "Goto Definition", "Jumps to the definition of the selected node if available, e.g., C++ code for a native function or the graph for a Blueprint function.", EUserInterfaceActionType::Button, FInputChord(EKeys::G, EModifierKey::Alt) )
 
+	UI_COMMAND( OpenInNewTab, "Open in New Tab", "Opens this function, macro or collapsed graph in a new tab.", EUserInterfaceActionType::Button, FInputChord() )
+
 	UI_COMMAND( BreakThisLink, "Break This Link", "Breaks the selected pin's only link.", EUserInterfaceActionType::Button, FInputChord() )
 	UI_COMMAND( BreakPinLinks, "Break All Pin Links", "Breaks all of the selected pin's links", EUserInterfaceActionType::Button, FInputChord() )
 	UI_COMMAND( PromoteToVariable, "Promote to Variable", "Promotes something to a variable", EUserInterfaceActionType::Button, FInputChord() )
diff --git a/Engine/Source/Editor/GraphEditor/Public/GraphEditorActions.h b/Engine/Source/Editor/GraphEditor/Public/GraphEditorActions.h
--- a/Engine/Source/Editor/GraphEditor/Public/GraphEditorActions.h
+++ b/Engine/Source/Editor/GraphEditor/Public/GraphEditorActions.h
@@ -124,6 +124,9 @@ public:
 	// Jumps to the definition of the selected node (or otherwise focuses something interesting about that node, e.g., the inner graph for a collapsed graph)
 	TSharedPtr< FUICommandInfo > GoToDefinition;
 
+	// Opens collapsed graphs, functions or macros in a new tab
+	TSharedPtr< FUICommandInfo > OpenInNewTab;
+
 	// Pin-specific actions
 	TSharedPtr< FUICommandInfo > BreakThisLink;
 	TSharedPtr< FUICommandInfo > BreakPinLinks;
diff --git a/Engine/Source/Editor/Kismet/Private/BlueprintEditor.cpp b/Engine/Source/Editor/Kismet/Private/BlueprintEditor.cpp
--- a/Engine/Source/Editor/Kismet/Private/BlueprintEditor.cpp
+++ b/Engine/Source/Editor/Kismet/Private/BlueprintEditor.cpp
@@ -1699,6 +1699,12 @@ TSharedRef<SGraphEditor> FBlueprintEditor::CreateGraphEditorWidget(TSharedRef<FT
 				FCanExecuteAction::CreateSP(this, &FBlueprintEditor::CanGoToDocumentation)
 				);
 
+			GraphEditorCommands->MapAction(FGraphEditorCommands::Get().OpenInNewTab,
+				FExecuteAction::CreateSP(this, &FBlueprintEditor::OnOpenSelectedNodesInNewTab),
+				FCanExecuteAction::CreateSP(this, &FBlueprintEditor::CanOpenSelectedNodesInNewTab),
+				FIsActionButtonVisible::CreateSP(this, &FBlueprintEditor::CanOpenSelectedNodesInNewTab)
+			);
+
 			GraphEditorCommands->MapAction(FGraphEditorCommands::Get().EnableNodes,
 				FExecuteAction(),
 				FCanExecuteAction(),
@@ -6331,6 +6337,71 @@ bool FBlueprintEditor::CanExpandNodes() const
 	return false;
 }
 
+void FBlueprintEditor::OnOpenSelectedNodesInNewTab()
+{
+	const FGraphPanelSelectionSet SelectedNodes = GetSelectedNodes();
+	for (FGraphPanelSelectionSet::TConstIterator NodeIt(SelectedNodes); NodeIt; ++NodeIt)
+	{
+		if (const UK2Node_Composite* CompositeNode = Cast<UK2Node_Composite>(*NodeIt))
+		{
+			if (CompositeNode->BoundGraph)
+			{
+				OpenDocument(CompositeNode->BoundGraph, FDocumentTracker::ForceOpenNewDocument);
+			}
+		}
+		else if (const UK2Node_MacroInstance* SelectedMacroInstanceNode = Cast<UK2Node_MacroInstance>(*NodeIt))
+		{
+			if (SelectedMacroInstanceNode->GetMacroGraph())
+			{
+				OpenDocument(SelectedMacroInstanceNode->GetMacroGraph(), FDocumentTracker::ForceOpenNewDocument);
+			}
+		}
+		else if (const UK2Node_CallFunction* SelectedCallFunctionNode = Cast<UK2Node_CallFunction>(*NodeIt))
+		{
+			const UEdGraphNode* ResultEventNode = nullptr;
+			const UEdGraph* FunctionGraph = SelectedCallFunctionNode->GetFunctionGraph(ResultEventNode);
+
+			if (FunctionGraph && !ResultEventNode)
+			{
+				OpenDocument(FunctionGraph, FDocumentTracker::ForceOpenNewDocument);
+			}
+		}
+	}
+}
+
+bool FBlueprintEditor::CanOpenSelectedNodesInNewTab() const
+{
+	const FGraphPanelSelectionSet SelectedNodes = GetSelectedNodes();
+	for (FGraphPanelSelectionSet::TConstIterator NodeIt(SelectedNodes); NodeIt; ++NodeIt)
+	{
+		if (const UK2Node_Composite* CompositeNode = Cast<UK2Node_Composite>(*NodeIt))
+		{
+			if (CompositeNode->BoundGraph)
+			{
+				return true;
+			}
+		}
+		else if (const UK2Node_MacroInstance* SelectedMacroInstanceNode = Cast<UK2Node_MacroInstance>(*NodeIt))
+		{
+			if (SelectedMacroInstanceNode->GetMacroGraph())
+			{
+				return true;
+			}
+		}
+		else if (const UK2Node_CallFunction* SelectedCallFunctionNode = Cast<UK2Node_CallFunction>(*NodeIt))
+		{
+			const UEdGraphNode* ResultEventNode = nullptr;
+			const UEdGraph* FunctionGraph = SelectedCallFunctionNode->GetFunctionGraph(ResultEventNode);
+
+			// Only allow opening a new tab if the call function isn't sourced by a custom event call,
+			// because in that case it just opens the current graph in a new tab which can be confusing and doesn't add any benefit to the user 
+			return FunctionGraph != nullptr && ResultEventNode == nullptr;
+		}
+	}
+
+	return false;
+}
+
 void FBlueprintEditor::MoveNodesToAveragePos(TSet<UEdGraphNode*>& AverageNodes, FVector2D SourcePos, bool bExpandedNodesNeedUniqueGuid /* = false */) const
 {
 	if (AverageNodes.Num() > 0)
diff --git a/Engine/Source/Editor/Kismet/Public/BlueprintEditor.h b/Engine/Source/Editor/Kismet/Public/BlueprintEditor.h
--- a/Engine/Source/Editor/Kismet/Public/BlueprintEditor.h
+++ b/Engine/Source/Editor/Kismet/Public/BlueprintEditor.h
@@ -1026,6 +1026,9 @@ protected:
 	UE_API void OnExpandNodes();
 	UE_API bool CanExpandNodes() const;
 
+	UE_API void OnOpenSelectedNodesInNewTab();
+	UE_API bool CanOpenSelectedNodesInNewTab() const;
+
 	/**
 	* Move the given set of nodes to an average spot near the Source position
 	* 
diff --git a/Engine/Source/Editor/MaterialEditor/Private/MaterialEditor.cpp b/Engine/Source/Editor/MaterialEditor/Private/MaterialEditor.cpp
--- a/Engine/Source/Editor/MaterialEditor/Private/MaterialEditor.cpp
+++ b/Engine/Source/Editor/MaterialEditor/Private/MaterialEditor.cpp
@@ -7201,6 +7201,38 @@ bool FMaterialEditor::CanExpandNodes() const
 	return false;
 }
 
+void FMaterialEditor::OnOpenSelectedNodesInNewTab()
+{
+	const FGraphPanelSelectionSet SelectedNodes = GetSelectedNodes();
+	for (FGraphPanelSelectionSet::TConstIterator NodeIt(SelectedNodes); NodeIt; ++NodeIt)
+	{
+		if (const UMaterialGraphNode_Composite* CompositeNode = Cast<UMaterialGraphNode_Composite>(*NodeIt))
+		{
+			if (CompositeNode->BoundGraph)
+			{
+				OpenDocument(CompositeNode->BoundGraph, FDocumentTracker::ForceOpenNewDocument);
+			}
+		}
+	}
+}
+
+bool FMaterialEditor::CanOpenSelectedNodesInNewTab() const
+{
+	const FGraphPanelSelectionSet SelectedNodes = GetSelectedNodes();
+	for (FGraphPanelSelectionSet::TConstIterator NodeIt(SelectedNodes); NodeIt; ++NodeIt)
+	{
+		if (const UMaterialGraphNode_Composite* CompositeNode = Cast<UMaterialGraphNode_Composite>(*NodeIt))
+		{
+			if (CompositeNode->BoundGraph)
+			{
+				return true;
+			}
+		}
+	}
+
+	return false;
+}
+
 void FMaterialEditor::OnAlignTop()
 {
 	if (FocusedGraphEdPtr.IsValid())
@@ -7795,6 +7827,12 @@ TSharedRef<SGraphEditor> FMaterialEditor::CreateGraphEditorWidget(TSharedRef<cla
 			FCanExecuteAction::CreateSP(this, &FMaterialEditor::CanGoToDocumentation)
 		);
 
+		GraphEditorCommands->MapAction(FGraphEditorCommands::Get().OpenInNewTab,
+			FExecuteAction::CreateSP(this, &FMaterialEditor::OnOpenSelectedNodesInNewTab),
+			FCanExecuteAction::CreateSP(this, &FMaterialEditor::CanOpenSelectedNodesInNewTab),
+			FIsActionButtonVisible::CreateSP(this, &FMaterialEditor::CanOpenSelectedNodesInNewTab)
+		);
+
 		// Collapse Node Commands
 		GraphEditorCommands->MapAction( FGraphEditorCommands::Get().CollapseNodes,
 			FExecuteAction::CreateSP( this, &FMaterialEditor::OnCollapseNodes ),
diff --git a/Engine/Source/Editor/MaterialEditor/Private/MaterialEditor.h b/Engine/Source/Editor/MaterialEditor/Private/MaterialEditor.h
--- a/Engine/Source/Editor/MaterialEditor/Private/MaterialEditor.h
+++ b/Engine/Source/Editor/MaterialEditor/Private/MaterialEditor.h
@@ -595,6 +595,11 @@ protected:
 	/** Check if node can be expanded */
 	bool CanExpandNodes() const;
 
+	/** Opens collapsed nodes in a new tab */
+	void OnOpenSelectedNodesInNewTab();
+	/** Check if selected nodes can be opened in a new tab*/
+	bool CanOpenSelectedNodesInNewTab() const;
+
 	void OnAlignTop();
 	void OnAlignMiddle();
 	void OnAlignBottom();
diff --git a/Engine/Source/Editor/UnrealEd/Private/MaterialGraphNode.cpp b/Engine/Source/Editor/UnrealEd/Private/MaterialGraphNode.cpp
--- a/Engine/Source/Editor/UnrealEd/Private/MaterialGraphNode.cpp
+++ b/Engine/Source/Editor/UnrealEd/Private/MaterialGraphNode.cpp
@@ -500,6 +500,11 @@ void UMaterialGraphNode::GetNodeContextMenuActions(UToolMenu* Menu, UGraphNodeCo
 			}
 		}
 
+		{
+			FToolMenuSection& Section = Menu->AddSection("MaterialSchemaNodeNavigation", LOCTEXT("NodeNavigationMenuHeader", "Navigation"));
+			Section.AddMenuEntry(FGraphEditorCommands::Get().OpenInNewTab);
+		}
+		
 		// Separate the above frequently used options from the below less frequently used common options
 		
 		{
