Blueprint Editor: Clicking empty space in the viewport clears selection.

[RN]

diff --git a/Engine/Source/Editor/Kismet/Private/SCSEditorViewportClient.cpp b/Engine/Source/Editor/Kismet/Private/SCSEditorViewportClient.cpp
--- a/Engine/Source/Editor/Kismet/Private/SCSEditorViewportClient.cpp
+++ b/Engine/Source/Editor/Kismet/Private/SCSEditorViewportClient.cpp
@@ -390,107 +390,120 @@ void FSCSEditorViewportClient::ProcessClick(class FSceneView& View, class HHitPr
 {
 	const FViewportClick Click(&View, this, Key, Event, HitX, HitY);
 
-	if (HitProxy)
+	if (!HitProxy)
 	{
-		if (HitProxy->IsA(HInstancedStaticMeshInstance::StaticGetType()))
+		if (const TSharedPtr<FBlueprintEditor> BlueprintEditor = BlueprintEditorPtr.Pin())
 		{
-			HInstancedStaticMeshInstance* InstancedStaticMeshInstanceProxy = ((HInstancedStaticMeshInstance*)HitProxy);
-
-			TSharedPtr<ISCSEditorCustomization> Customization = BlueprintEditorPtr.Pin()->CustomizeSubobjectEditor(InstancedStaticMeshInstanceProxy->Component);
-			if (Customization.IsValid() && Customization->HandleViewportClick(AsShared(), View, HitProxy, Key, Event, HitX, HitY))
+			if (const TSharedPtr<SSubobjectEditor> SubobjectEditor = BlueprintEditor->GetSubobjectEditor())
 			{
-				Invalidate();
+				// If the user clicked on empty space, clear the selection
+				SubobjectEditor->ClearSelection();
 			}
-
-			return;
 		}
-		else if (HitProxy->IsA(HWidgetAxis::StaticGetType()))
-		{
-			const bool bOldModeWidgets1 = EngineShowFlags.ModeWidgets;
-			const bool bOldModeWidgets2 = View.Family->EngineShowFlags.ModeWidgets;
-
-			EngineShowFlags.SetModeWidgets(false);
-			FSceneViewFamily* SceneViewFamily = const_cast<FSceneViewFamily*>(View.Family);
-			SceneViewFamily->EngineShowFlags.SetModeWidgets(false);
-			bool bWasWidgetDragging = Widget->IsDragging();
-			Widget->SetDragging(false);
-
-			// Invalidate the hit proxy map so it will be rendered out again when GetHitProxy
-			// is called
-			Viewport->InvalidateHitProxy();
-
-			// This will actually re-render the viewport's hit proxies!
-			HHitProxy* HitProxyWithoutAxisWidgets = Viewport->GetHitProxy(HitX, HitY);
-			if (HitProxyWithoutAxisWidgets != NULL && !HitProxyWithoutAxisWidgets->IsA(HWidgetAxis::StaticGetType()))
-			{
-				// Try this again, but without the widget this time!
-				ProcessClick(View, HitProxyWithoutAxisWidgets, Key, Event, HitX, HitY);
-			}
-
-			// Undo the evil
-			EngineShowFlags.SetModeWidgets(bOldModeWidgets1);
-			SceneViewFamily->EngineShowFlags.SetModeWidgets(bOldModeWidgets2);
+		return;
+	}
 
-			Widget->SetDragging(bWasWidgetDragging);
+	if (HitProxy->IsA(HInstancedStaticMeshInstance::StaticGetType()))
+	{
+		HInstancedStaticMeshInstance* InstancedStaticMeshInstanceProxy = ((HInstancedStaticMeshInstance*)HitProxy);
 
-			// Invalidate the hit proxy map again so that it'll be refreshed with the original
-			// scene contents if we need it again later.
-			Viewport->InvalidateHitProxy();
-			return;
+		TSharedPtr<ISCSEditorCustomization> Customization = BlueprintEditorPtr.Pin()->CustomizeSubobjectEditor(InstancedStaticMeshInstanceProxy->Component);
+		if (Customization.IsValid() && Customization->HandleViewportClick(AsShared(), View, HitProxy, Key, Event, HitX, HitY))
+		{
+			Invalidate();
 		}
-		else if(GUnrealEd->ComponentVisManager.HandleClick(this, HitProxy, Click))
+
+		return;
+	}
+
+	if (HitProxy->IsA(HWidgetAxis::StaticGetType()))
+	{
+		const bool bOldModeWidgets1 = EngineShowFlags.ModeWidgets;
+		const bool bOldModeWidgets2 = View.Family->EngineShowFlags.ModeWidgets;
+
+		EngineShowFlags.SetModeWidgets(false);
+		FSceneViewFamily* SceneViewFamily = const_cast<FSceneViewFamily*>(View.Family);
+		SceneViewFamily->EngineShowFlags.SetModeWidgets(false);
+		bool bWasWidgetDragging = Widget->IsDragging();
+		Widget->SetDragging(false);
+
+		// Invalidate the hit proxy map so it will be rendered out again when GetHitProxy
+		// is called
+		Viewport->InvalidateHitProxy();
+
+		// This will actually re-render the viewport's hit proxies!
+		HHitProxy* HitProxyWithoutAxisWidgets = Viewport->GetHitProxy(HitX, HitY);
+		if (HitProxyWithoutAxisWidgets != NULL && !HitProxyWithoutAxisWidgets->IsA(HWidgetAxis::StaticGetType()))
 		{
-			// Component Vis Manager handled this click, no need to do anything
+			// Try this again, but without the widget this time!
+			ProcessClick(View, HitProxyWithoutAxisWidgets, Key, Event, HitX, HitY);
 		}
-		else if (HitProxy->IsA(HActor::StaticGetType()))
+
+		// Undo the evil
+		EngineShowFlags.SetModeWidgets(bOldModeWidgets1);
+		SceneViewFamily->EngineShowFlags.SetModeWidgets(bOldModeWidgets2);
+
+		Widget->SetDragging(bWasWidgetDragging);
+
+		// Invalidate the hit proxy map again so that it'll be refreshed with the original
+		// scene contents if we need it again later.
+		Viewport->InvalidateHitProxy();
+		return;
+	}
+
+	if(GUnrealEd->ComponentVisManager.HandleClick(this, HitProxy, Click))
+	{
+		// Component Vis Manager handled this click, no need to do anything
+		return;
+	}
+
+	if (HitProxy->IsA(HActor::StaticGetType()))
+	{
+		HActor* ActorProxy = (HActor*)HitProxy;
+		AActor* PreviewActor = GetPreviewActor();
+		if (ActorProxy && ActorProxy->Actor && ActorProxy->PrimComponent)
 		{
-			HActor* ActorProxy = (HActor*)HitProxy;
-			AActor* PreviewActor = GetPreviewActor();
-			if (ActorProxy && ActorProxy->Actor && ActorProxy->PrimComponent)
-			{
-				const USceneComponent* SelectedCompInstance = nullptr;
+			const USceneComponent* SelectedCompInstance = nullptr;
 
-				if (ActorProxy->Actor == PreviewActor)
+			if (ActorProxy->Actor == PreviewActor)
+			{
+				const UPrimitiveComponent* TestComponent = ActorProxy->PrimComponent;
+				if (ActorProxy->Actor->GetComponents().Contains(TestComponent)
+					&& IsObjectSelectableInViewport(const_cast<UPrimitiveComponent*>(TestComponent)))
 				{
-					const UPrimitiveComponent* TestComponent = ActorProxy->PrimComponent;
-					if (ActorProxy->Actor->GetComponents().Contains(TestComponent)
-						&& IsObjectSelectableInViewport(const_cast<UPrimitiveComponent*>(TestComponent)))
-					{
-						SelectedCompInstance = TestComponent;
-					}
+					SelectedCompInstance = TestComponent;
 				}
-				else if (ActorProxy->Actor->IsChildActor())
+			}
+			else if (ActorProxy->Actor->IsChildActor())
+			{
+				AActor* TestActor = ActorProxy->Actor;
+				while (TestActor->GetParentActor()->IsChildActor())
 				{
-					AActor* TestActor = ActorProxy->Actor;
-					while (TestActor->GetParentActor()->IsChildActor())
-					{
-						TestActor = TestActor->GetParentActor();
-					}
+					TestActor = TestActor->GetParentActor();
+				}
 
-					if (TestActor->GetParentActor() == PreviewActor)
-					{
-						SelectedCompInstance = TestActor->GetParentComponent();
-					}
+				if (TestActor->GetParentActor() == PreviewActor)
+				{
+					SelectedCompInstance = TestActor->GetParentComponent();
 				}
+			}
 
-				if (SelectedCompInstance)
+			if (SelectedCompInstance)
+			{
+				TSharedPtr<ISCSEditorCustomization> Customization = BlueprintEditorPtr.Pin()->CustomizeSubobjectEditor(SelectedCompInstance);
+				if (!(Customization.IsValid() && Customization->HandleViewportClick(AsShared(), View, HitProxy, Key, Event, HitX, HitY)))
 				{
-					TSharedPtr<ISCSEditorCustomization> Customization = BlueprintEditorPtr.Pin()->CustomizeSubobjectEditor(SelectedCompInstance);
-					if (!(Customization.IsValid() && Customization->HandleViewportClick(AsShared(), View, HitProxy, Key, Event, HitX, HitY)))
+					const bool bIsCtrlKeyDown = Viewport->KeyState(EKeys::LeftControl) || Viewport->KeyState(EKeys::RightControl);
+					if (BlueprintEditorPtr.IsValid())
 					{
-						const bool bIsCtrlKeyDown = Viewport->KeyState(EKeys::LeftControl) || Viewport->KeyState(EKeys::RightControl);
-						if (BlueprintEditorPtr.IsValid())
-						{
-							// Note: This will find and select any node associated with the component instance that's attached to the proxy (including visualizers)
-							BlueprintEditorPtr.Pin()->FindAndSelectSubobjectEditorTreeNode(SelectedCompInstance, bIsCtrlKeyDown);
-						}
+						// Note: This will find and select any node associated with the component instance that's attached to the proxy (including visualizers)
+						BlueprintEditorPtr.Pin()->FindAndSelectSubobjectEditorTreeNode(SelectedCompInstance, bIsCtrlKeyDown);
 					}
 				}
 			}
-
-			Invalidate();
-			return;
 		}
+
+		Invalidate();
 	}
 }
 
