BlueprintEditor selection notifier modular feature. This allows our modeless tools system to listen in on blueprint editor selection changes to activate tools/tool groups and show UI if required. Follow-up CL 44756770 for integration with our modeless tools system.

diff --git a/Engine/Source/Editor/Kismet/Private/BlueprintEditor.cpp b/Engine/Source/Editor/Kismet/Private/BlueprintEditor.cpp
--- a/Engine/Source/Editor/Kismet/Private/BlueprintEditor.cpp
+++ b/Engine/Source/Editor/Kismet/Private/BlueprintEditor.cpp
@@ -182,6 +182,7 @@
 #include "Algo/MinElement.h"
 #include "Editor/EditorEngine.h"
 #include "EditorViewportSelectabilityBridge.h"
+#include "IBlueprintSelectionNotifier.h"
 
 DEFINE_LOG_CATEGORY_STATIC(LogBlueprintEditor, Log, All);
 
@@ -1299,6 +1300,36 @@ void FBlueprintEditor::OnSelectionUpdated(const TArray<TSharedPtr<FSubobjectEdit
 		Options.bShowComponents = bShowComponents;
 		Inspector->ShowDetailsForObjects(InspectorObjects, Options);
 	}
+
+	// Get all registered Blueprint Selection Notifiers
+	TArray<IBlueprintSelectionNotifier*> Notifiers = IModularFeatures::Get().GetModularFeatureImplementations<IBlueprintSelectionNotifier>(
+		IBlueprintSelectionNotifier::GetModularFeatureName());
+
+	// Notify each registered implementation
+	TArray<FTypedElementHandle> SelectedElements;
+	for (const FSubobjectEditorTreeNodePtrType& SelectedNode : SelectedNodes)
+	{
+		if (const FSubobjectData* const Data = SelectedNode->GetDataSource())
+		{
+			if (const UActorComponent* Component = Data->FindComponentInstanceInActor(GetPreviewActor()))
+			{
+				// Convert component to TypedElementHandle
+				FTypedElementHandle ElementHandle = UEngineElementsLibrary::AcquireEditorComponentElementHandle(Component);
+				if (ElementHandle.IsSet())
+				{
+					SelectedElements.Add(ElementHandle);
+				}
+			}
+		}
+	}
+
+	for (IBlueprintSelectionNotifier* Notifier : Notifiers)
+	{
+		if (Notifier)
+		{
+			Notifier->OnBlueprintSelectionUpdated(this, Blueprint, SelectedElements);
+		}
+	}
 }
 
 void FBlueprintEditor::OnComponentDoubleClicked(TSharedPtr<FSubobjectEditorTreeNode> Node)
diff --git a/Engine/Source/Editor/Kismet/Public/IBlueprintSelectionNotifier.h b/Engine/Source/Editor/Kismet/Public/IBlueprintSelectionNotifier.h
new file mode 100644
--- /dev/null
+++ b/Engine/Source/Editor/Kismet/Public/IBlueprintSelectionNotifier.h
@@ -0,0 +1,32 @@
+﻿// Copyright Epic Games, Inc. All Rights Reserved.
+
+#pragma once
+
+#include "CoreMinimal.h"
+#include "Features/IModularFeature.h"
+
+class UBlueprint;
+struct FTypedElementHandle;
+class IAssetEditorInstance;
+
+/**
+ * Modular feature interface for Blueprint Editor selection notifications
+ * Implemented by plugins that want to be notified when Blueprint Editor selection changes
+ */
+class UE_INTERNAL IBlueprintSelectionNotifier : public IModularFeature
+{
+public:
+    static FName GetModularFeatureName()
+    {
+        static FName ModularFeatureName = FName(TEXT("BlueprintSelectionNotifier"));
+        return ModularFeatureName;
+    }
+
+    /**
+     * Called when Blueprint Editor selection changes
+     * @param AssetEditorInstance The asset editor instance that triggered the selection change
+     * @param Blueprint The blueprint being edited
+     * @param SelectedElements The currently selected elements in the blueprint
+     */
+    virtual void OnBlueprintSelectionUpdated(IAssetEditorInstance* AssetEditorInstance, UBlueprint* Blueprint, const TArray<FTypedElementHandle>& SelectedElements) = 0;
+};
\ No newline at end of file
