Use the full object name for the warning in archetype sorting.

The default short name fails to provide enough context when it's tripped.

diff --git a/Engine/Source/Editor/Kismet/Private/BlueprintCompilationManager.cpp b/Engine/Source/Editor/Kismet/Private/BlueprintCompilationManager.cpp
--- a/Engine/Source/Editor/Kismet/Private/BlueprintCompilationManager.cpp
+++ b/Engine/Source/Editor/Kismet/Private/BlueprintCompilationManager.cpp
@@ -3116,13 +3116,13 @@ void FBlueprintCompilationManagerImpl::ReinstanceBatch(TArray<FReinstancingJob>&
 				TArray<UObject*> Dependencies;
 				if (OldArchetypeInfo.ArchetypeTemplate && !OldArchetypeInfo.ArchetypeTemplate->HasAnyFlags(RF_ClassDefaultObject))
 				{
-					if(ArchetypeObjects.Contains(OldArchetypeInfo.ArchetypeTemplate))
+					if (ArchetypeObjects.Contains(OldArchetypeInfo.ArchetypeTemplate))
 					{
 						Dependencies.Add(OldArchetypeInfo.ArchetypeTemplate);
 					}
 					else
 					{
-						UE_LOG(LogBlueprint, Warning, TEXT("Expecting the template object (%s) of archetype (%s) to already be in the list of archetypes"), *GetNameSafe(OldArchetypeInfo.ArchetypeTemplate), *GetNameSafe(OldArchetypeInfo.Archetype));
+						UE_LOG(LogBlueprint, Warning, TEXT("Expecting the template object (%s) of archetype (%s) to already be in the list of archetypes"), *GetFullNameSafe(OldArchetypeInfo.ArchetypeTemplate), *GetFullNameSafe(OldArchetypeInfo.Archetype));
 					}
 				}
 
