Removed erroneous warning from Blueprint re-instancing.

(This is loosely tied to UE-224532.)

For a given archetype object, there's no guarantee that its template is the same class. It could also be that the template is another subobject with a different _base_ type. In this case, the template won't be in the ArchetypeObjects array.

Currently, we don't see this because there's a bug in GetArchetype where we incorrectly return a subobject's class CDO as the template, which always shows up in the ArchetypeObjects array.

#rnx

diff --git a/Engine/Source/Editor/Kismet/Private/BlueprintCompilationManager.cpp b/Engine/Source/Editor/Kismet/Private/BlueprintCompilationManager.cpp
--- a/Engine/Source/Editor/Kismet/Private/BlueprintCompilationManager.cpp
+++ b/Engine/Source/Editor/Kismet/Private/BlueprintCompilationManager.cpp
@@ -3120,10 +3120,6 @@ void FBlueprintCompilationManagerImpl::ReinstanceBatch(TArray<FReinstancingJob>&
 					{
 						Dependencies.Add(OldArchetypeInfo.ArchetypeTemplate);
 					}
-					else
-					{
-						UE_LOG(LogBlueprint, Warning, TEXT("Expecting the template object (%s) of archetype (%s) to already be in the list of archetypes"), *GetFullNameSafe(OldArchetypeInfo.ArchetypeTemplate), *GetFullNameSafe(OldArchetypeInfo.Archetype));
-					}
 				}
 
 				return Dependencies;
