[BlueprintEditor] FBlueprintEditor::GetPreviewActor now tests if the PreviewScene is valid before returning its result.

diff --git a/Engine/Source/Editor/Kismet/Private/BlueprintEditor.cpp b/Engine/Source/Editor/Kismet/Private/BlueprintEditor.cpp
--- a/Engine/Source/Editor/Kismet/Private/BlueprintEditor.cpp
+++ b/Engine/Source/Editor/Kismet/Private/BlueprintEditor.cpp
@@ -10410,6 +10410,21 @@ bool FBlueprintEditor::CanFindReferences()
 
 AActor* FBlueprintEditor::GetPreviewActor() const
 {
+	if (!PreviewScene.IsInitialized())
+	{
+		// On editor shutdown the PreviewScene might uninitialize before the BlueprintEditor.
+		// The PreviewScene nulls its world and starts the gc, which would theoretically collect the PreviewActor too.
+		// In practice however, the actor might be kept alive through references (e.g. on Linux, see UE-296301).
+		// Thus, we test manually for the state of the PreviewScene, to not return an invalid actor.
+		if (AActor* Actor = PreviewActorPtr.Get())
+		{
+			UE_LOG(LogBlueprint, Warning, TEXT("Expected preview actor '%s' to be garbage collected, but it was not:"), *Actor->GetName());
+			GEngine->Exec(Actor->GetWorld(), *FString::Printf(TEXT("obj refs name=\"%s\""), *Actor->GetPathName()));
+		}
+		PreviewActorPtr.Reset();
+		return nullptr;
+	}
+
 	UBlueprint* PreviewBlueprint = GetBlueprintObj();
 
 	// Note: The weak ptr can become stale if the actor is reinstanced due to a Blueprint change, etc. In that 
diff --git a/Engine/Source/Runtime/Engine/Public/PreviewScene.h b/Engine/Source/Runtime/Engine/Public/PreviewScene.h
--- a/Engine/Source/Runtime/Engine/Public/PreviewScene.h
+++ b/Engine/Source/Runtime/Engine/Public/PreviewScene.h
@@ -88,6 +88,7 @@ public:
 	// Accessors.
 	UWorld* GetWorld() const { return PreviewWorld; }
 	FSceneInterface* GetScene() const { return PreviewWorld->Scene; }
+	bool IsInitialized() const { return PreviewWorld != nullptr; }
 
 	/** Access to line drawing */
 	class ULineBatchComponent* GetLineBatcher() const { return LineBatcher; }
