Update ABP interface icon in the 'My Blueprint' tab to show as blue when the function has been implemented

diff --git a/Engine/Source/Editor/Kismet/Private/SBlueprintPalette.cpp b/Engine/Source/Editor/Kismet/Private/SBlueprintPalette.cpp
--- a/Engine/Source/Editor/Kismet/Private/SBlueprintPalette.cpp
+++ b/Engine/Source/Editor/Kismet/Private/SBlueprintPalette.cpp
@@ -326,6 +326,18 @@ static void GetSubGraphIcon(FEdGraphSchemaAction_K2Graph const* const ActionIn,
 					ColorOut = FLinearColor::Yellow;
 				}
 			}
+			if (const UAnimBlueprint* AnimBlueprint = Cast<UAnimBlueprint>(BlueprintIn))
+			{
+				IAnimClassInterface* AnimClassInterface = IAnimClassInterface::GetFromClass(AnimBlueprint->GetAnimBlueprintGeneratedClass());
+				if (const FAnimBlueprintFunction* Function = IAnimClassInterface::FindAnimBlueprintFunction(AnimClassInterface, ActionIn->FuncName))
+				{
+					if (Function->bImplemented)
+					{
+						ToolTipOut = ToolTipOut = FText::Format(LOCTEXT("FunctionFromInterface_Tooltip", "Function (from Interface '{InterfaceName}') with implementation"), Args);
+						ColorOut = FLinearColor(0.0f, 1.0f, 1.0f);
+					}
+				}
+			}
 		}
 		break;
 	case EEdGraphSchemaAction_K2Graph::Function:
