Fix LOCTEXT key conflict

#rnx

diff --git a/Engine/Source/Editor/Kismet/Private/SBlueprintPalette.cpp b/Engine/Source/Editor/Kismet/Private/SBlueprintPalette.cpp
--- a/Engine/Source/Editor/Kismet/Private/SBlueprintPalette.cpp
+++ b/Engine/Source/Editor/Kismet/Private/SBlueprintPalette.cpp
@@ -333,7 +333,7 @@ static void GetSubGraphIcon(FEdGraphSchemaAction_K2Graph const* const ActionIn,
 				{
 					if (Function->bImplemented)
 					{
-						ToolTipOut = ToolTipOut = FText::Format(LOCTEXT("FunctionFromInterface_Tooltip", "Function (from Interface '{InterfaceName}') with implementation"), Args);
+						ToolTipOut = ToolTipOut = FText::Format(LOCTEXT("FunctionFromInterfaceImplemented_Tooltip", "Function (from Interface '{InterfaceName}') with implementation"), Args);
 						ColorOut = FLinearColor(0.0f, 1.0f, 1.0f);
 					}
 				}
