New command line argument to dump BP native class serialization.

#rnx

diff --git a/Engine/Source/Editor/Kismet/Internal/Blueprints/BlueprintDependencies.h b/Engine/Source/Editor/Kismet/Internal/Blueprints/BlueprintDependencies.h
--- a/Engine/Source/Editor/Kismet/Internal/Blueprints/BlueprintDependencies.h
+++ b/Engine/Source/Editor/Kismet/Internal/Blueprints/BlueprintDependencies.h
@@ -6,6 +6,8 @@
 
 #include "BlueprintDependencies.generated.h"
 
+DECLARE_LOG_CATEGORY_EXTERN(LogCookBlueprint, Display, All);
+
 UENUM()
 enum class EBPDependencyType
 {
diff --git a/Engine/Source/Editor/Kismet/Private/Blueprints/BlueprintDependencies.cpp b/Engine/Source/Editor/Kismet/Private/Blueprints/BlueprintDependencies.cpp
--- a/Engine/Source/Editor/Kismet/Private/Blueprints/BlueprintDependencies.cpp
+++ b/Engine/Source/Editor/Kismet/Private/Blueprints/BlueprintDependencies.cpp
@@ -11,6 +11,7 @@
 #include "JsonObjectConverter.h"
 #include "JsonObjectGraph/Stringify.h"
 #include "Kismet2/BlueprintEditorUtils.h"
+#include "Misc/CommandLine.h"
 #include "Misc/Optional.h"
 #include "Misc/PackageAccessTracking.h"
 #include "Misc/ScopeExit.h"
@@ -19,6 +20,8 @@
 #include "Serialization/CompactBinaryWriter.h"
 #include "UObject/StructOnScope.h"
 
+DEFINE_LOG_CATEGORY(LogCookBlueprint);
+
 #include UE_INLINE_GENERATED_CPP_BY_NAME(BlueprintDependencies)
 
 /*
@@ -446,6 +449,19 @@ static FString UE::Private::HashNativeClass( const UClass* NativeClass)
 	FUtf8String StringifiedClass = UE::JsonObjectGraph::Stringify(
 		RootObjects,
 		FJsonStringifyOptions(EJsonStringifyFlags::FilterEditorOnlyData));
+
+	static FString ClassToDump;
+	static bool bDumpNativeClass = FParse::Value(FCommandLine::Get(), TEXT("-DumpBPNativeClassSerialization="), ClassToDump);
+	static FName ClassNameToDump(ClassToDump);
+
+	if (bDumpNativeClass && NativeClass->GetFName() == ClassNameToDump)
+	{
+		UE_LOG(LogCookBlueprint, Display, TEXT("---Start Blueprint native class serialization---"));
+		UE_LOG(LogCookBlueprint, Display, TEXT("%s"), *ClassToDump);
+		UE_LOG(LogCookBlueprint, Display, TEXT("%hs"), *StringifiedClass);
+		UE_LOG(LogCookBlueprint, Display, TEXT("---End---"));
+	}
+
 	FBlake3 Hash;
 	Hash.Update(StringifiedClass.GetCharArray().GetData(), StringifiedClass.GetCharArray().Num());
 	CachedValue->Hash = LexToString(Hash.Finalize());
