UE-316040 - Crash when converting a function graph to an event graph

ConvertFunctionToEvent incorrectly copies over getter nodes that represent function input parameters. Since the event graph doesn't support these, we need to replace these nodes with direct links from the corresponding parameter pins on the new event node.

#rnx

diff --git a/Engine/Source/Editor/Kismet/Private/BlueprintEditor.cpp b/Engine/Source/Editor/Kismet/Private/BlueprintEditor.cpp
--- a/Engine/Source/Editor/Kismet/Private/BlueprintEditor.cpp
+++ b/Engine/Source/Editor/Kismet/Private/BlueprintEditor.cpp
@@ -6571,7 +6571,7 @@ void FBlueprintEditor::ConvertFunctionToEvent(UK2Node_FunctionEntry* SelectedCal
 		EventGraph->Modify();
 
 		// Find all return nodes and break their connections
-		TArray< UK2Node_FunctionResult*> FunctionReturnNodes;
+		TArray<UK2Node_FunctionResult*> FunctionReturnNodes;
 		FunctionGraph->GetNodesOfClass(FunctionReturnNodes);
 		for (UK2Node_FunctionResult* Node : FunctionReturnNodes)
 		{
@@ -6582,8 +6582,8 @@ void FBlueprintEditor::ConvertFunctionToEvent(UK2Node_FunctionEntry* SelectedCal
 		}
 
 		// Keep track of the old connections from the entry node
-		TMap<FString, TSet<UEdGraphPin*>> PinConnections;
-		FEdGraphUtilities::GetPinConnectionMap(SelectedCallFunctionNode, PinConnections);
+		TMap<FString, TSet<UEdGraphPin*>> EntryNodePinConnections;
+		FEdGraphUtilities::GetPinConnectionMap(SelectedCallFunctionNode, EntryNodePinConnections);
 
 		FName EventName = Func->GetFName();
 		UClass* const OverrideFuncClass = CastChecked<UClass>(Func->GetOuter())->GetAuthoritativeClass();
@@ -6616,7 +6616,6 @@ void FBlueprintEditor::ConvertFunctionToEvent(UK2Node_FunctionEntry* SelectedCal
 			// Spawn a new node, we have to do it this way to avoid renaming on creation of the custom event
 			UEdGraphNode* NewNode = nullptr;
 			NewNode = NewObject<UEdGraphNode>(EventGraph, UK2Node_CustomEvent::StaticClass());
-			check(NewNode != nullptr);
 			NewNode->CreateNewGuid();
 
 			NewNode->NodePosX = static_cast<int32>(SpawnPos.X);
@@ -6641,14 +6640,54 @@ void FBlueprintEditor::ConvertFunctionToEvent(UK2Node_FunctionEntry* SelectedCal
 
 			K2Schema->ReconstructNode(*NewEventNode);
 		}
+
+		check(NewEventNode);
 		
+		TArray<UK2Node_Variable*> AllVarNodes;
+		TArray<UK2Node_Variable*> RemainingVarNodes;
+		FunctionGraph->GetNodesOfClass<UK2Node_Variable>(AllVarNodes);
+
+		// Parameter getter nodes aren't supported in the event graph, so they'll be destroyed.
+		// In order to preserve their connections, we'll add their linked pins to the entry node's pin connection map.
+		for (UK2Node_Variable* VarNode : AllVarNodes)
+		{
+			check(VarNode);
+			const FProperty* Property = VarNode->GetPropertyForVariable();
+			const bool bIsInputParameter =
+				Property &&
+				Property->HasAnyPropertyFlags(CPF_Parm) &&
+				!Property->HasAnyPropertyFlags(CPF_ReturnParm)
+			;
+
+			if (bIsInputParameter)
+			{
+				TSet<UEdGraphPin*>* ConnectionsEntry = EntryNodePinConnections.Find(Property->GetName());
+				if (ensure(ConnectionsEntry))
+				{
+					UEdGraphPin* ValuePin = VarNode->GetValuePin();
+					if (ensure(ValuePin))
+					{
+						for (UEdGraphPin* ConnectedPin : ValuePin->LinkedTo)
+						{
+							if (ensure(ConnectedPin))
+							{
+								ConnectionsEntry->Add(ConnectedPin);
+							}
+						}
+					}
+				}
+
+				VarNode->DestroyNode();
+			}
+			else
+			{
+				RemainingVarNodes.Add(VarNode);
+			}
+		}
+
 		// If this function had any local scope variables, we need to convert them to global scope
-		if(SelectedCallFunctionNode->LocalVariables.Num() > 0)
+		if (SelectedCallFunctionNode->LocalVariables.Num() > 0)
 		{
-			// Find any UK2Node_Variable's that may reference a local variable
-			TArray<UK2Node_Variable*> VarNodes;
-			FunctionGraph->GetNodesOfClass<UK2Node_Variable>(VarNodes);
-
 			// Make a globally scoped version of any local variables
 			for (const FBPVariableDescription& LocalVar : SelectedCallFunctionNode->LocalVariables)
 			{
@@ -6662,7 +6701,7 @@ void FBlueprintEditor::ConvertFunctionToEvent(UK2Node_FunctionEntry* SelectedCal
 					// Upon success, update the variable node's reference members
 					const FGuid NewVarGuid = FBlueprintEditorUtils::FindMemberVariableGuidByName(NodeBP, NewVarName);
 
-					for (UK2Node_Variable* VarNode : VarNodes)
+					for (UK2Node_Variable* VarNode : RemainingVarNodes)
 					{
 						if (VarNode->GetVarName() == LocalVar.VarName)
 						{
@@ -6694,13 +6733,11 @@ void FBlueprintEditor::ConvertFunctionToEvent(UK2Node_FunctionEntry* SelectedCal
 			Result->DestroyNode();
 		}
 
-		// Connect any pins that need to be set from the old function
-		if (NewEventNode)
-		{
-			// Link the nodes from the original function entry node to the new event node
-			FEdGraphUtilities::ReconnectPinMap(NewEventNode, PinConnections);
-			FEdGraphUtilities::CopyPinDefaults(SelectedCallFunctionNode, NewEventNode);
-		}
+		// Link the nodes from the original function entry node to the new event node
+		FEdGraphUtilities::ReconnectPinMap(NewEventNode, EntryNodePinConnections);
+
+		// Copy over parameter default values to the new event node
+		FEdGraphUtilities::CopyPinDefaults(SelectedCallFunctionNode, NewEventNode);
 
 		// Remove the old function graph
 		FBlueprintEditorUtils::RemoveGraph(NodeBP, FunctionGraph, EGraphRemoveFlags::Recompile);
@@ -6708,10 +6745,7 @@ void FBlueprintEditor::ConvertFunctionToEvent(UK2Node_FunctionEntry* SelectedCal
 		FBlueprintEditorUtils::MarkBlueprintAsStructurallyModified(NodeBP);
 
 		// Do this AFTER removing the function graph so that it's not opened into the existing function graph document tab.
-		if (NewEventNode)
-		{
-			FKismetEditorUtilities::BringKismetToFocusAttentionOnObject(NewEventNode, false);
-		}
+		FKismetEditorUtilities::BringKismetToFocusAttentionOnObject(NewEventNode, false);
 	}
 }
 
