Control Rig: Blueprint saving even when generating errors and Save on Success Only is enabled
#rnx

diff --git a/Engine/Source/Editor/Kismet/Private/BlueprintCompilationManager.cpp b/Engine/Source/Editor/Kismet/Private/BlueprintCompilationManager.cpp
--- a/Engine/Source/Editor/Kismet/Private/BlueprintCompilationManager.cpp
+++ b/Engine/Source/Editor/Kismet/Private/BlueprintCompilationManager.cpp
@@ -445,10 +445,18 @@ void FBlueprintCompilationManagerImpl::CompileSynchronouslyImpl(const FBPCompile
 	{
 		if (!bSkipSaving)
 		{
+			UBlueprintEditorSettings* Settings = GetMutableDefault<UBlueprintEditorSettings>();
+			const bool bSaveBlueprintsAfterCompile = Settings->SaveOnCompile == SoC_Always;
+            const bool bSaveBlueprintAfterCompileSucceeded = Settings->SaveOnCompile == SoC_SuccessOnly;
+			
 			TArray<UPackage*> PackagesToSave;
 			for (UBlueprint* BP : CompiledBlueprintsToSave)
 			{
-				PackagesToSave.Add(BP->GetOutermost());
+				// The status of the blueprint might have changed during PostCDOCompiled
+				if(bSaveBlueprintsAfterCompile || (bSaveBlueprintAfterCompileSucceeded && BP->Status == BS_UpToDate))
+				{
+					PackagesToSave.Add(BP->GetOutermost());
+				}
 			}
 
 			FEditorFileUtils::PromptForCheckoutAndSave(PackagesToSave, /*bCheckDirty =*/true, /*bPromptToSave =*/false);
