Blueprint Editor: The tree view should now poperly sort by parentage, type (scene components, then actors, then whatever else), then display name (alphabetically).

[RN]

diff --git a/Engine/Source/Editor/Kismet/Private/BlueprintEditor.cpp b/Engine/Source/Editor/Kismet/Private/BlueprintEditor.cpp
--- a/Engine/Source/Editor/Kismet/Private/BlueprintEditor.cpp
+++ b/Engine/Source/Editor/Kismet/Private/BlueprintEditor.cpp
@@ -179,6 +179,7 @@
 #include "AnimationBlendSpaceSampleGraph.h"
 #include "SSubobjectEditor.h"
 #include "BlueprintActionDatabase.h"
+#include "BlueprintEditorSCSEditorUICustomization.h"
 #include "Algo/MinElement.h"
 #include "Editor/EditorEngine.h"
 #include "EditorViewportSelectabilityBridge.h"
@@ -2421,10 +2422,13 @@ void FBlueprintEditor::InitBlueprintEditor(
 	const EToolkitMode::Type Mode,
 	const TSharedPtr< IToolkitHost >& InitToolkitHost,
 	const TArray<UBlueprint*>& InBlueprints,
-	bool bShouldOpenInDefaultsMode)
+	bool bShouldOpenInDefaultsMode, 
+	TSharedPtr<ISCSEditorUICustomization> InSCSEditorUICustomization)
 {
 	check(InBlueprints.Num() == 1 || bShouldOpenInDefaultsMode);
 
+	SCSEditorUICustomization = InSCSEditorUICustomization;
+
 	// TRUE if a single Blueprint is being opened and is marked as newly created
 	bool bNewlyCreated = InBlueprints.Num() == 1 && InBlueprints[0]->bIsNewlyCreated;
 
@@ -3078,8 +3082,9 @@ void FBlueprintEditor::CreateSubobjectEditors()
 		.AllowEditing(this, &FBlueprintEditor::InEditingMode)
 		.OnSelectionUpdated(this, &FBlueprintEditor::OnSelectionUpdated)
 		.OnItemDoubleClicked(this, &FBlueprintEditor::OnComponentDoubleClicked)
-		.SubobjectClassListFilters(ClassFilters);
-	
+		.SubobjectClassListFilters(ClassFilters)
+		.SCSEditorUICustomization(SCSEditorUICustomization);
+
 	LLM_SCOPE_BYNAME(TEXT("BPCreateSubobjectEditorViewportClient"));
 
 	SubobjectViewport = SAssignNew(SubobjectViewport, SSCSEditorViewport)
@@ -8015,8 +8020,10 @@ void FBlueprintEditor::SetDetailsCustomization(TSharedPtr<FDetailsViewObjectFilt
 	}
 }
 
-void FBlueprintEditor::SetSubobjectEditorUICustomization(TSharedPtr<ISCSEditorUICustomization> SCSEditorUICustomization)
+void FBlueprintEditor::SetSubobjectEditorUICustomization(TSharedPtr<ISCSEditorUICustomization> InSCSEditorUICustomization)
 {
+	SCSEditorUICustomization = InSCSEditorUICustomization;
+
 	if(SubobjectEditor.IsValid())
 	{
 		SubobjectEditor->SetUICustomization(SCSEditorUICustomization);
diff --git a/Engine/Source/Editor/Kismet/Private/BlueprintEditorModule.cpp b/Engine/Source/Editor/Kismet/Private/BlueprintEditorModule.cpp
--- a/Engine/Source/Editor/Kismet/Private/BlueprintEditorModule.cpp
+++ b/Engine/Source/Editor/Kismet/Private/BlueprintEditorModule.cpp
@@ -5,6 +5,7 @@
 
 #include "BlueprintDebugger.h"
 #include "BlueprintEditor.h"
+#include "BlueprintEditorSCSEditorUICustomization.h"
 #include "BlueprintGraphPanelPinFactory.h"
 #include "BlueprintNamespaceRegistry.h"
 #include "Containers/ContainerAllocationPolicies.h"
@@ -253,6 +254,8 @@ void FBlueprintEditorModule::StartupModule()
 	// Add a page for pre-loading of the editor
 	MessageLogModule.GetLogListing("BlueprintLog")->NewPage(LOCTEXT("PreloadLogPageLabel", "Editor Load"));
 
+	SCSEditorUICustomization = FFBlueprintEditorSCSEditorUICustomization::GetInstance();
+
 	// Register internal SCS editor customizations
 	RegisterSCSEditorCustomization("InstancedStaticMeshComponent", FSCSEditorCustomizationBuilder::CreateStatic(&FInstancedStaticMeshSCSEditorCustomization::MakeInstance));
 	RegisterSCSEditorCustomization("HierarchicalInstancedStaticMeshComponent", FSCSEditorCustomizationBuilder::CreateStatic(&FInstancedStaticMeshSCSEditorCustomization::MakeInstance));
@@ -313,10 +316,9 @@ TSharedRef<IBlueprintEditor> FBlueprintEditorModule::CreateBlueprintEditor(const
 {
 	TSharedRef< FBlueprintEditor > NewBlueprintEditor( new FBlueprintEditor() );
 
-	NewBlueprintEditor->InitBlueprintEditor(Mode, InitToolkitHost, BlueprintsToEdit, bShouldOpenInDefaultsMode);
+	NewBlueprintEditor->InitBlueprintEditor(Mode, InitToolkitHost, BlueprintsToEdit, bShouldOpenInDefaultsMode, SCSEditorUICustomization);
 
 	NewBlueprintEditor->SetDetailsCustomization(DetailsObjectFilter, DetailsRootCustomization);
-	NewBlueprintEditor->SetSubobjectEditorUICustomization(SCSEditorUICustomization);
 
 	for(auto It(SCSEditorCustomizations.CreateConstIterator()); It; ++It)
 	{
@@ -387,7 +389,7 @@ void FBlueprintEditorModule::SetDetailsCustomization(TSharedPtr<FDetailsViewObje
 
 void FBlueprintEditorModule::SetSubobjectEditorUICustomization(TSharedPtr<ISCSEditorUICustomization> InSCSEditorUICustomization)
 {
-	SCSEditorUICustomization = InSCSEditorUICustomization;
+	AddCustomization(InSCSEditorUICustomization);
 
 	for (const TSharedRef<IBlueprintEditor>& BlueprintEditor : GetBlueprintEditors())
 	{
@@ -395,6 +397,22 @@ void FBlueprintEditorModule::SetSubobjectEditorUICustomization(TSharedPtr<ISCSEd
 	}
 }
 
+void FBlueprintEditorModule::AddCustomization(TSharedPtr<ISCSEditorUICustomization> InSCSEditorUICustomization)
+{
+	if (SCSEditorUICustomization)
+	{
+		SCSEditorUICustomization->AddCustomization(InSCSEditorUICustomization);
+	}
+}
+
+void FBlueprintEditorModule::RemoveCustomization(TSharedPtr<ISCSEditorUICustomization> InSCSEditorUICustomization)
+{
+	if (SCSEditorUICustomization)
+	{
+		SCSEditorUICustomization->RemoveCustomization(InSCSEditorUICustomization);
+	}
+}
+
 void FBlueprintEditorModule::RegisterSCSEditorCustomization(const FName& InComponentName, FSCSEditorCustomizationBuilder InCustomizationBuilder)
 {
 	SCSEditorCustomizations.Add(InComponentName, InCustomizationBuilder);
diff --git a/Engine/Source/Editor/Kismet/Private/BlueprintEditorSCSEditorUICustomization.cpp b/Engine/Source/Editor/Kismet/Private/BlueprintEditorSCSEditorUICustomization.cpp
new file mode 100644
--- /dev/null
+++ b/Engine/Source/Editor/Kismet/Private/BlueprintEditorSCSEditorUICustomization.cpp
@@ -0,0 +1,206 @@
+// Copyright Epic Games, Inc. All Rights Reserved.
+#include "BlueprintEditorSCSEditorUICustomization.h"
+
+#include "SubobjectData.h"
+#include "SubobjectDataHandle.h"
+
+TSharedPtr<FFBlueprintEditorSCSEditorUICustomization> FFBlueprintEditorSCSEditorUICustomization::Instance;
+
+TSharedPtr<FFBlueprintEditorSCSEditorUICustomization> FFBlueprintEditorSCSEditorUICustomization::GetInstance()
+{
+	if (!Instance)
+	{
+		Instance = MakeShareable(new FFBlueprintEditorSCSEditorUICustomization());
+	}
+
+	return Instance;
+}
+
+bool FFBlueprintEditorSCSEditorUICustomization::HideComponentsTree(TArrayView<UObject*> Context) const
+{
+	for (const TSharedPtr<ISCSEditorUICustomization>& Customization : Customizations)
+	{
+		if (Customization->HideComponentsTree(Context))
+		{
+			return true;
+		}
+	}
+
+	return ISCSEditorUICustomization::HideComponentsTree(Context);
+}
+
+bool FFBlueprintEditorSCSEditorUICustomization::HideComponentsFilterBox(TArrayView<UObject*> Context) const
+{
+	for (const TSharedPtr<ISCSEditorUICustomization>& Customization : Customizations)
+	{
+		if (Customization->HideComponentsFilterBox(Context))
+		{
+			return true;
+		}
+	}
+
+	return ISCSEditorUICustomization::HideComponentsFilterBox(Context);
+}
+
+bool FFBlueprintEditorSCSEditorUICustomization::HideAddComponentButton(TArrayView<UObject*> Context) const
+{
+	for (const TSharedPtr<ISCSEditorUICustomization>& Customization : Customizations)
+	{
+		if (Customization->HideAddComponentButton(Context))
+		{
+			return true;
+		}
+	}
+
+	return ISCSEditorUICustomization::HideAddComponentButton(Context);
+}
+
+bool FFBlueprintEditorSCSEditorUICustomization::HideBlueprintButtons(TArrayView<UObject*> Context) const
+{
+	for (const TSharedPtr<ISCSEditorUICustomization>& Customization : Customizations)
+	{
+		if (Customization->HideBlueprintButtons(Context))
+		{
+			return true;
+		}
+	}
+
+	return ISCSEditorUICustomization::HideBlueprintButtons(Context);
+}
+
+const FSlateBrush* FFBlueprintEditorSCSEditorUICustomization::GetIconBrush(const FSubobjectData& SubobjectData) const
+{
+	for (const TSharedPtr<ISCSEditorUICustomization>& Customization : Customizations)
+	{
+		if (const FSlateBrush* IconBrush = Customization->GetIconBrush(SubobjectData))
+		{
+			return IconBrush;
+		}
+	}
+
+	return ISCSEditorUICustomization::GetIconBrush(SubobjectData);
+}
+
+TSharedPtr<SWidget> FFBlueprintEditorSCSEditorUICustomization::GetControlsWidget(const FSubobjectData& SubobjectData) const
+{
+	for (const TSharedPtr<ISCSEditorUICustomization>& Customization : Customizations)
+	{
+		if (TSharedPtr<SWidget> ControlsWidget = Customization->GetControlsWidget(SubobjectData))
+		{
+			return ControlsWidget;
+		}
+	}
+
+	return ISCSEditorUICustomization::GetControlsWidget(SubobjectData);
+}
+
+EChildActorComponentTreeViewVisualizationMode FFBlueprintEditorSCSEditorUICustomization::GetChildActorVisualizationMode() const
+{
+	for (const TSharedPtr<ISCSEditorUICustomization>& Customization : Customizations)
+	{
+		EChildActorComponentTreeViewVisualizationMode VisualizationMode = Customization->GetChildActorVisualizationMode();
+		if (VisualizationMode != EChildActorComponentTreeViewVisualizationMode::UseDefault)
+		{
+			return VisualizationMode;
+		}
+	}
+
+	return ISCSEditorUICustomization::GetChildActorVisualizationMode();
+}
+
+TSubclassOf<UActorComponent> FFBlueprintEditorSCSEditorUICustomization::GetComponentTypeFilter(TArrayView<UObject*> Context) const
+{
+	for (const TSharedPtr<ISCSEditorUICustomization>& Customization : Customizations)
+	{
+		if(TSubclassOf<UActorComponent> SubClass = Customization->GetComponentTypeFilter(Context))
+		{
+			return SubClass;
+		}
+	}
+
+	return ISCSEditorUICustomization::GetComponentTypeFilter(Context);
+}
+
+void FFBlueprintEditorSCSEditorUICustomization::AddCustomization(TSharedPtr<ISCSEditorUICustomization> Customization)
+{
+	Customizations.AddUnique(Customization);
+}
+
+void FFBlueprintEditorSCSEditorUICustomization::RemoveCustomization(TSharedPtr<ISCSEditorUICustomization> Customization)
+{
+	Customizations.Remove(Customization);
+}
+
+bool FFBlueprintEditorSCSEditorUICustomization::SortSubobjectData(TArray<FSubobjectDataHandle>& SubobjectData)
+{
+	for (const TSharedPtr<ISCSEditorUICustomization>& Customization : Customizations)
+	{
+		if (Customization->SortSubobjectData(SubobjectData))
+		{
+			return true;
+		}
+	}
+
+	return ISCSEditorUICustomization::SortSubobjectData(SubobjectData);
+}
+
+FFBlueprintEditorSCSEditorUICustomization::FFBlueprintEditorSCSEditorUICustomization()
+{
+	Customizations.Add(FFBlueprintEditorDefaultSortUICustomization::GetInstance());
+}
+
+TSharedPtr<FFBlueprintEditorDefaultSortUICustomization> FFBlueprintEditorDefaultSortUICustomization::Instance;
+
+TSharedPtr<FFBlueprintEditorDefaultSortUICustomization> FFBlueprintEditorDefaultSortUICustomization::GetInstance()
+{
+	if (!Instance)
+	{
+		Instance = MakeShareable(new FFBlueprintEditorDefaultSortUICustomization());
+	}
+
+	return Instance;
+}
+
+bool FFBlueprintEditorDefaultSortUICustomization::SortSubobjectData(TArray<FSubobjectDataHandle>& SubobjectData)
+{
+	// Sort subobjects by parentage, type (always put scene components first in the tree), then display name
+	SubobjectData.Sort([](const FSubobjectDataHandle& A, const FSubobjectDataHandle& B)
+	{
+		// If A is a child of B, sort B ahead
+		if (A.GetData()->GetParentData() == B.GetData())
+		{
+			return false;
+		}
+
+		// Scene components should sort ahead of non-scene components
+		if (A.GetData()->IsSceneComponent())
+		{
+			// If both are scene components, sort by display name
+			if (B.GetData()->IsSceneComponent())
+			{
+				return A.GetData()->GetDisplayName().CompareTo(B.GetData()->GetDisplayName()) < 0;
+			}
+
+			return !B.GetData()->IsActor();
+		}
+
+		// Actors should go next
+		if (A.GetData()->IsActor())
+		{
+			// If both are actors, sort by display name
+			if (B.GetData()->IsActor())
+			{
+				return A.GetData()->GetDisplayName().CompareTo(B.GetData()->GetDisplayName()) < 0;
+			}
+
+			return true;
+		}
+
+		// Otherwise just sort by display name
+		return A.GetData()->GetDisplayName().CompareTo(B.GetData()->GetDisplayName()) < 0;
+	});
+
+	return true;
+}
+
+
diff --git a/Engine/Source/Editor/Kismet/Private/BlueprintEditorSCSEditorUICustomization.h b/Engine/Source/Editor/Kismet/Private/BlueprintEditorSCSEditorUICustomization.h
new file mode 100644
--- /dev/null
+++ b/Engine/Source/Editor/Kismet/Private/BlueprintEditorSCSEditorUICustomization.h
@@ -0,0 +1,54 @@
+// Copyright Epic Games, Inc. All Rights Reserved.
+
+#pragma once
+
+#include "ISCSEditorUICustomization.h"
+#include "Containers/Array.h"
+#include "Templates/SharedPointer.h"
+
+class FFBlueprintEditorSCSEditorUICustomization : public ISCSEditorUICustomization
+{
+public:
+	static TSharedPtr<FFBlueprintEditorSCSEditorUICustomization> GetInstance();
+	
+	virtual ~FFBlueprintEditorSCSEditorUICustomization() override = default;
+
+	virtual bool HideComponentsTree(TArrayView<UObject*> Context) const override;
+
+	virtual bool HideComponentsFilterBox(TArrayView<UObject*> Context) const override;
+
+	virtual bool HideAddComponentButton(TArrayView<UObject*> Context) const override;
+
+	virtual bool HideBlueprintButtons(TArrayView<UObject*> Context) const override;
+
+	virtual const FSlateBrush* GetIconBrush(const FSubobjectData&) const override;
+	virtual TSharedPtr<SWidget> GetControlsWidget(const FSubobjectData&) const override;
+
+	virtual EChildActorComponentTreeViewVisualizationMode GetChildActorVisualizationMode() const override;
+
+	virtual TSubclassOf<UActorComponent> GetComponentTypeFilter(TArrayView<UObject*> Context) const override;
+
+	void AddCustomization(TSharedPtr<ISCSEditorUICustomization> Customization);
+	void RemoveCustomization(TSharedPtr<ISCSEditorUICustomization> Customization);
+
+	virtual bool SortSubobjectData(TArray<FSubobjectDataHandle>& SubobjectData) override;
+
+private:
+	FFBlueprintEditorSCSEditorUICustomization();
+
+	static TSharedPtr<FFBlueprintEditorSCSEditorUICustomization> Instance;
+	TArray<TSharedPtr<ISCSEditorUICustomization>> Customizations;
+};
+
+class FFBlueprintEditorDefaultSortUICustomization : public ISCSEditorUICustomization
+{
+public:
+	static TSharedPtr<FFBlueprintEditorDefaultSortUICustomization> GetInstance();
+
+	virtual bool SortSubobjectData(TArray<FSubobjectDataHandle>& SubobjectData) override;
+
+private:
+	FFBlueprintEditorDefaultSortUICustomization() = default;
+
+	static TSharedPtr<FFBlueprintEditorDefaultSortUICustomization> Instance;
+};
\ No newline at end of file
diff --git a/Engine/Source/Editor/Kismet/Public/BlueprintEditor.h b/Engine/Source/Editor/Kismet/Public/BlueprintEditor.h
--- a/Engine/Source/Editor/Kismet/Public/BlueprintEditor.h
+++ b/Engine/Source/Editor/Kismet/Public/BlueprintEditor.h
@@ -246,8 +246,9 @@ public:
 	 * @param	InitToolkitHost			When Mode is WorldCentric, this is the level editor instance to spawn this editor within
 	 * @param	InBlueprints			The blueprints to edit
 	 * @param	bShouldOpenInDefaultsMode	If true, the editor will open in defaults editing mode
+	 * @param	InSCSEditorUICustomization	UI customizations for the SCS editor inside the blueprint editor
 	 */
-	UE_API void InitBlueprintEditor(const EToolkitMode::Type Mode, const TSharedPtr< class IToolkitHost >& InitToolkitHost, const TArray<class UBlueprint*>& InBlueprints, bool bShouldOpenInDefaultsMode);
+	UE_API void InitBlueprintEditor(const EToolkitMode::Type Mode, const TSharedPtr<IToolkitHost >& InitToolkitHost, const TArray<UBlueprint*>& InBlueprints, bool bShouldOpenInDefaultsMode, TSharedPtr<ISCSEditorUICustomization> InSCSEditorUICustomization = nullptr);
 
 public:
 	//~ Begin FAssetEditorToolkit Interface
@@ -684,12 +685,12 @@ public:
 	UE_API void SetDetailsCustomization(TSharedPtr<class FDetailsViewObjectFilter> DetailsObjectFilter, TSharedPtr<class IDetailRootObjectCustomization> DetailsRootCustomization);
 
 	/** Sets subobject editor UI customization */
-	UE_API void SetSubobjectEditorUICustomization(TSharedPtr<class ISCSEditorUICustomization> SCSEditorUICustomization);
+	UE_API void SetSubobjectEditorUICustomization(TSharedPtr<class ISCSEditorUICustomization> InSCSEditorUICustomization);
 
 	UE_DEPRECATED(5.0, "SetSCSEditorUICustomization has been deprecated. Use SetSubobjectEditorUICustomization instead.")
-	void SetSCSEditorUICustomization(TSharedPtr<class ISCSEditorUICustomization> SCSEditorUICustomization)
+	void SetSCSEditorUICustomization(TSharedPtr<class ISCSEditorUICustomization> InSCSEditorUICustomization)
 	{
-		return SetSubobjectEditorUICustomization(SCSEditorUICustomization);
+		return SetSubobjectEditorUICustomization(InSCSEditorUICustomization);
 	}
 
 	/**
@@ -1622,6 +1623,9 @@ private:
 	/** analytics statistics for the Editor */ 
 	FAnalyticsStatistics AnalyticsStats;
 
+	/** UI customizations for the SCS editor inside the blueprint editor */
+	TSharedPtr<ISCSEditorUICustomization> SCSEditorUICustomization;
+
 	/** Customizations for the Subobject editor */
 	TMap<FName, TSharedPtr<ISCSEditorCustomization>> SubobjectEditorCustomizations;
 
diff --git a/Engine/Source/Editor/Kismet/Public/BlueprintEditorModule.h b/Engine/Source/Editor/Kismet/Public/BlueprintEditorModule.h
--- a/Engine/Source/Editor/Kismet/Public/BlueprintEditorModule.h
+++ b/Engine/Source/Editor/Kismet/Public/BlueprintEditorModule.h
@@ -24,6 +24,8 @@
 
 #define UE_API KISMET_API
 
+class FFBlueprintEditorSCSEditorUICustomization;
+class ISCSEditorUICustomization;
 class FBlueprintEditor;
 class FExtender;
 class FFieldClass;
@@ -230,7 +232,14 @@ public:
 	virtual void SetDetailsCustomization(TSharedPtr<class FDetailsViewObjectFilter> InDetailsObjectFilter, TSharedPtr<class IDetailRootObjectCustomization> InDetailsRootCustomization);
 
 	/** Sets SCS editor UI customization */
-	virtual void SetSubobjectEditorUICustomization(TSharedPtr<class ISCSEditorUICustomization> InSCSEditorUICustomization);
+	UE_DEPRECATED(5.7, "Please use AddCustomization instead.")
+	virtual void SetSubobjectEditorUICustomization(TSharedPtr<ISCSEditorUICustomization> InSCSEditorUICustomization);
+
+	/** Adds the SCS editor UI customization to the array of customizations to be invoked */
+	virtual void AddCustomization(TSharedPtr<ISCSEditorUICustomization> InSCSEditorUICustomization);
+
+	/** Removes the SCS editor UI customization from the array of customizations to be invoked */
+	virtual void RemoveCustomization(TSharedPtr<ISCSEditorUICustomization> InSCSEditorUICustomization);
 
 	/**
 	 * Register a customization for interacting with the SCS editor 
@@ -392,7 +401,7 @@ private:
 	TSharedPtr<class FDetailsViewObjectFilter> DetailsObjectFilter;
 
 	/** UI customizations for the SCS editor inside the blueprint editor */
-	TSharedPtr<class ISCSEditorUICustomization> SCSEditorUICustomization;
+	TSharedPtr<FFBlueprintEditorSCSEditorUICustomization> SCSEditorUICustomization;
 
 	/** 
 	 * A command list that can be passed around and isn't bound to an instance 
diff --git a/Engine/Source/Editor/Kismet/Public/ISCSEditorUICustomization.h b/Engine/Source/Editor/Kismet/Public/ISCSEditorUICustomization.h
--- a/Engine/Source/Editor/Kismet/Public/ISCSEditorUICustomization.h
+++ b/Engine/Source/Editor/Kismet/Public/ISCSEditorUICustomization.h
@@ -9,6 +9,7 @@
 
 class UActorComponent;
 struct FSubobjectData;
+struct FSubobjectDataHandle;
 
 /** SCSEditor UI customization */
 class ISCSEditorUICustomization
@@ -68,4 +69,8 @@ public:
 
 	/** @return A component type that limits visible nodes when filtering the tree view */
 	virtual TSubclassOf<UActorComponent> GetComponentTypeFilter(TArrayView<UObject*> Context) const { return nullptr; }
+
+	/** Optionally sorts the given SubobjectDataHandles for the tree view as appropriate for this editor.
+	 * @return Whether any sorting was attempted */
+	virtual bool SortSubobjectData(TArray<FSubobjectDataHandle>& SubobjectData) { return false; }
 };
diff --git a/Engine/Source/Editor/LevelEditor/Private/SActorDetails.cpp b/Engine/Source/Editor/LevelEditor/Private/SActorDetails.cpp
--- a/Engine/Source/Editor/LevelEditor/Private/SActorDetails.cpp
+++ b/Engine/Source/Editor/LevelEditor/Private/SActorDetails.cpp
@@ -2,6 +2,7 @@
 
 #include "SActorDetails.h"
 
+#include "ActorDetailsSCSEditorUICustomization.h"
 #include "Components/ActorComponent.h"
 #include "Components/SceneComponent.h"
 #include "DetailsViewObjectFilter.h"
@@ -278,7 +279,8 @@ void SActorDetails::Construct(const FArguments& InArgs, UTypedElementSelectionSe
 		.ObjectContext(this, &SActorDetails::GetObjectContext)
 		.AllowEditing(this, &SActorDetails::GetAllowComponentTreeEditing)
 		.OnSelectionUpdated(this, &SActorDetails::OnSubobjectEditorTreeViewSelectionChanged)
-		.OnItemDoubleClicked(this, &SActorDetails::OnSubobjectEditorTreeViewItemDoubleClicked);
+		.OnItemDoubleClicked(this, &SActorDetails::OnSubobjectEditorTreeViewItemDoubleClicked)
+		.SCSEditorUICustomization(InArgs._SCSEditorUICustomization);
 
 	ComponentsBox->SetContent(SubobjectEditor.ToSharedRef());
 
diff --git a/Engine/Source/Editor/LevelEditor/Private/SActorDetails.h b/Engine/Source/Editor/LevelEditor/Private/SActorDetails.h
--- a/Engine/Source/Editor/LevelEditor/Private/SActorDetails.h
+++ b/Engine/Source/Editor/LevelEditor/Private/SActorDetails.h
@@ -11,6 +11,7 @@
 #include "EditorUndoClient.h"
 #include "Elements/Interfaces/TypedElementDetailsInterface.h"
 
+class FActorDetailsSCSEditorUICustomization;
 class AActor;
 class FSCSEditorTreeNode;
 class FTabManager;
@@ -39,6 +40,7 @@ class SActorDetails : public SCompoundWidget, public FEditorUndoClient, public F
 {
 public:
 	SLATE_BEGIN_ARGS(SActorDetails) {}
+		SLATE_ARGUMENT(TSharedPtr<ISCSEditorUICustomization>, SCSEditorUICustomization)
 	SLATE_END_ARGS()
 
 	void Construct(const FArguments& InArgs, UTypedElementSelectionSet* InSelectionSet, const FName TabIdentifier, TSharedPtr<FUICommandList> InCommandList, TSharedPtr<FTabManager> InTabManager);
diff --git a/Engine/Source/Editor/LevelEditor/Private/SLevelEditor.cpp b/Engine/Source/Editor/LevelEditor/Private/SLevelEditor.cpp
--- a/Engine/Source/Editor/LevelEditor/Private/SLevelEditor.cpp
+++ b/Engine/Source/Editor/LevelEditor/Private/SLevelEditor.cpp
@@ -2216,10 +2216,10 @@ void SLevelEditor::AddStandaloneLevelViewport( const TSharedRef<SLevelViewport>&
 
 TSharedRef<SWidget> SLevelEditor::CreateActorDetails( const FName TabIdentifier )
 {
-	TSharedRef<SActorDetails> ActorDetails = SNew( SActorDetails, GetMutableElementSelectionSet(), TabIdentifier, LevelEditorCommands, GetTabManager() );
+	TSharedRef<SActorDetails> ActorDetails = SNew( SActorDetails, GetMutableElementSelectionSet(), TabIdentifier, LevelEditorCommands, GetTabManager())
+		.SCSEditorUICustomization(ActorDetailsSCSEditorUICustomization);
 
 	ActorDetails->SetActorDetailsRootCustomization(ActorDetailsObjectFilter, ActorDetailsRootCustomization);
-	ActorDetails->SetSubobjectEditorUICustomization(ActorDetailsSCSEditorUICustomization);
 
 	AllActorDetailPanels.Add( ActorDetails );
 	return ActorDetails;
diff --git a/Engine/Source/Editor/SubobjectEditor/Private/SSubobjectBlueprintEditor.cpp b/Engine/Source/Editor/SubobjectEditor/Private/SSubobjectBlueprintEditor.cpp
--- a/Engine/Source/Editor/SubobjectEditor/Private/SSubobjectBlueprintEditor.cpp
+++ b/Engine/Source/Editor/SubobjectEditor/Private/SSubobjectBlueprintEditor.cpp
@@ -50,6 +50,7 @@ void SSubobjectBlueprintEditor::Construct(const FArguments& InArgs)
 	AllowEditing = InArgs._AllowEditing;
 	HideComponentClassCombo = InArgs._HideComponentClassCombo;
 	bAllowTreeUpdates = true;
+	UICustomization = InArgs._SCSEditorUICustomization;
 	
 	CreateCommandList();
 
diff --git a/Engine/Source/Editor/SubobjectEditor/Private/SSubobjectEditor.cpp b/Engine/Source/Editor/SubobjectEditor/Private/SSubobjectEditor.cpp
--- a/Engine/Source/Editor/SubobjectEditor/Private/SSubobjectEditor.cpp
+++ b/Engine/Source/Editor/SubobjectEditor/Private/SSubobjectEditor.cpp
@@ -2535,6 +2535,11 @@ void SSubobjectEditor::UpdateTree(bool bRegenerateTreeNodes /* = true */)
 
 			DataSubsystem->GatherSubobjectData(Context, SubobjectData);
 
+			if (UICustomization)
+			{
+				UICustomization->SortSubobjectData(SubobjectData);
+			}
+
 			FSubobjectEditorTreeNodePtrType SeperatorNode;
 			TMap<FSubobjectDataHandle, FSubobjectEditorTreeNodePtrType> AddedNodes;
 
diff --git a/Engine/Source/Editor/SubobjectEditor/Private/SSubobjectInstanceEditor.cpp b/Engine/Source/Editor/SubobjectEditor/Private/SSubobjectInstanceEditor.cpp
--- a/Engine/Source/Editor/SubobjectEditor/Private/SSubobjectInstanceEditor.cpp
+++ b/Engine/Source/Editor/SubobjectEditor/Private/SSubobjectInstanceEditor.cpp
@@ -43,6 +43,7 @@ void SSubobjectInstanceEditor::Construct(const FArguments& InArgs)
 	OnItemDoubleClicked = InArgs._OnItemDoubleClicked;
 	AllowEditing = InArgs._AllowEditing;
 	bAllowTreeUpdates = true;
+	UICustomization = InArgs._SCSEditorUICustomization;
 	
 	CreateCommandList();
 
diff --git a/Engine/Source/Editor/SubobjectEditor/Public/SSubobjectBlueprintEditor.h b/Engine/Source/Editor/SubobjectEditor/Public/SSubobjectBlueprintEditor.h
--- a/Engine/Source/Editor/SubobjectEditor/Public/SSubobjectBlueprintEditor.h
+++ b/Engine/Source/Editor/SubobjectEditor/Public/SSubobjectBlueprintEditor.h
@@ -34,6 +34,7 @@ private:
 	    SLATE_EVENT(FOnItemDoubleClicked, OnItemDoubleClicked)
 		SLATE_EVENT(FOnHighlightPropertyInDetailsView, OnHighlightPropertyInDetailsView)
 		SLATE_ARGUMENT(TArray<TSharedRef<IClassViewerFilter>>, SubobjectClassListFilters)
+		SLATE_ARGUMENT(TSharedPtr<ISCSEditorUICustomization>, SCSEditorUICustomization)
 	    
 	SLATE_END_ARGS()
 	
diff --git a/Engine/Source/Editor/SubobjectEditor/Public/SSubobjectInstanceEditor.h b/Engine/Source/Editor/SubobjectEditor/Public/SSubobjectInstanceEditor.h
--- a/Engine/Source/Editor/SubobjectEditor/Public/SSubobjectInstanceEditor.h
+++ b/Engine/Source/Editor/SubobjectEditor/Public/SSubobjectInstanceEditor.h
@@ -23,6 +23,7 @@ private:
 	    SLATE_ATTRIBUTE(bool, AllowEditing)
 	    SLATE_EVENT(FOnSelectionUpdated, OnSelectionUpdated)
 	    SLATE_EVENT(FOnItemDoubleClicked, OnItemDoubleClicked)
+		SLATE_ARGUMENT(TSharedPtr<ISCSEditorUICustomization>, SCSEditorUICustomization)
 	
 	SLATE_END_ARGS()
 	
