FindInBlueprintManager no longer crashes with access violation when engine exists while a search in blueprints is ongoing

diff --git a/Engine/Source/Editor/Kismet/Private/FindInBlueprintManager.cpp b/Engine/Source/Editor/Kismet/Private/FindInBlueprintManager.cpp
--- a/Engine/Source/Editor/Kismet/Private/FindInBlueprintManager.cpp
+++ b/Engine/Source/Editor/Kismet/Private/FindInBlueprintManager.cpp
@@ -1428,11 +1428,24 @@ public:
 		:Controller(InController)
 	{
 		Thread = TUniquePtr<FRunnableThread>(FRunnableThread::Create(this, TEXT("FAsyncSearchIndexTaskRunnable"), 0, TPri_BelowNormal));
+
+		// Make sure engine shut down is properly handled. Modules involved in the async search will be destroyed after engine pre exit.
+		FCoreDelegates::OnEnginePreExit.AddRaw(this, &FAsyncSearchIndexTaskRunnable::OnEnginePreExit);
+	}
+
+	virtual ~FAsyncSearchIndexTaskRunnable()
+	{
+		bDestructing = true;
+		EnsureCompletion();
+
+		FCoreDelegates::OnEnginePreExit.RemoveAll(this);
 	}
 
 	virtual uint32 Run() override
 	{
-		while (Controller->IsWorkPending())
+		while (
+			!bDestructing &&
+			Controller->IsWorkPending())
 		{
 			// Get the next batch of asset paths to be indexed.
 			TArray<FSoftObjectPath> AssetPathsToIndex;
@@ -1449,9 +1462,13 @@ public:
 				ParallelForFlags |= EParallelForFlags::ForceSingleThread;
 			}
 
-			ParallelFor(AssetPathsToIndex.Num(), [&AssetPathsToIndex, Controller = this->Controller, bEnableFullIndexingPass, ParallelForFlags](int32 ArrayIdx)
+			ParallelFor(AssetPathsToIndex.Num(), [&AssetPathsToIndex, bEnableFullIndexingPass, ParallelForFlags, this](int32 ArrayIdx)
 			{
 				CSV_CUSTOM_STAT(FindInBlueprint, IndexedAssetCountThisFrame, 1, ECsvCustomStatOp::Accumulate);
+				if (bDestructing)
+				{
+					return;
+				}
 
 				FSoftObjectPath AssetPath = AssetPathsToIndex[ArrayIdx];
 				FFindInBlueprintSearchManager& FindManager = FFindInBlueprintSearchManager::Get();
@@ -1475,6 +1492,11 @@ public:
 								TArray<FImaginaryFiBDataSharedPtr> IndexNodes = { SearchData.ImaginaryBlueprint };
 								while (IndexNodes.Num() > 0)
 								{
+									if (bDestructing)
+									{
+										break;
+									}
+
 									ParallelFor(IndexNodes.Num(), [&IndexNodes](int32 NodeIdx)
 									{
 										IndexNodes[NodeIdx]->ParseAllChildData();
@@ -1521,8 +1543,18 @@ public:
 	}
 
 private:
+	/** Called before engine exits. Useful to terminate async work before involved engine modules are being destroyed. */
+	void OnEnginePreExit()
+	{
+		bDestructing = true;
+		EnsureCompletion();
+	}
+
 	TUniquePtr<FRunnableThread> Thread;
 	IAsyncSearchIndexTaskController* Controller;
+
+	/** True when this object is being destructed, useful to terminate intermediate work instead of blocking until completion is ensured */
+	TAtomic<bool> bDestructing = false;
 };
 
 class FCacheAllBlueprintsTickableObject : public IAsyncSearchIndexTaskController
