Editor: Make it possible to filter objects in the sub object outliner and use this for entities.

diff --git a/Engine/Source/Editor/Kismet/Public/ISCSEditorUICustomization.h b/Engine/Source/Editor/Kismet/Public/ISCSEditorUICustomization.h
--- a/Engine/Source/Editor/Kismet/Public/ISCSEditorUICustomization.h
+++ b/Engine/Source/Editor/Kismet/Public/ISCSEditorUICustomization.h
@@ -70,6 +70,11 @@ public:
 	/** @return A component type that limits visible nodes when filtering the tree view */
 	virtual TSubclassOf<UActorComponent> GetComponentTypeFilter(TArrayView<UObject*> Context) const { return nullptr; }
 
+	/** Optionally filter the given SubobjectDataHandles for the tree view as appropriate for this editor.
+	 * @return Whether any edits were made */
+	UE_EXPERIMENTAL(5.7, "The interface for filtering the sub object editor may change.")
+	virtual bool FilterSubobjectData(TArray<FSubobjectDataHandle>& SubobjectData) const { return false; }
+
 	/** Optionally sorts the given SubobjectDataHandles for the tree view as appropriate for this editor.
 	 * @return Whether any sorting was attempted */
 	virtual bool SortSubobjectData(TArray<FSubobjectDataHandle>& SubobjectData) { return false; }
diff --git a/Engine/Source/Editor/LevelEditor/Private/ActorDetailsSCSEditorUICustomization.cpp b/Engine/Source/Editor/LevelEditor/Private/ActorDetailsSCSEditorUICustomization.cpp
--- a/Engine/Source/Editor/LevelEditor/Private/ActorDetailsSCSEditorUICustomization.cpp
+++ b/Engine/Source/Editor/LevelEditor/Private/ActorDetailsSCSEditorUICustomization.cpp
@@ -115,6 +115,16 @@ TSubclassOf<UActorComponent> FActorDetailsSCSEditorUICustomization::GetComponent
 	return nullptr;
 }
 
+bool FActorDetailsSCSEditorUICustomization::FilterSubobjectData(TArray<FSubobjectDataHandle>& SubobjectData) const
+{
+	bool bChanged = false;
+	for (const TSharedPtr<ISCSEditorUICustomization>& Customization : Customizations)
+	{
+		bChanged = bChanged || Customization->FilterSubobjectData(SubobjectData);
+	}
+	return bChanged;
+}
+
 void FActorDetailsSCSEditorUICustomization::AddCustomization(TSharedPtr<ISCSEditorUICustomization> Customization)
 {
 	Customizations.AddUnique(Customization);
diff --git a/Engine/Source/Editor/LevelEditor/Private/ActorDetailsSCSEditorUICustomization.h b/Engine/Source/Editor/LevelEditor/Private/ActorDetailsSCSEditorUICustomization.h
--- a/Engine/Source/Editor/LevelEditor/Private/ActorDetailsSCSEditorUICustomization.h
+++ b/Engine/Source/Editor/LevelEditor/Private/ActorDetailsSCSEditorUICustomization.h
@@ -27,6 +27,8 @@ public:
 	virtual EChildActorComponentTreeViewVisualizationMode GetChildActorVisualizationMode() const override;
 
 	virtual TSubclassOf<UActorComponent> GetComponentTypeFilter(TArrayView<UObject*> Context) const override;
+	
+	virtual bool FilterSubobjectData(TArray<FSubobjectDataHandle>& SubobjectData) const override;
 
 	void AddCustomization(TSharedPtr<ISCSEditorUICustomization> Customization);
 	void RemoveCustomization(TSharedPtr<ISCSEditorUICustomization> Customization);
diff --git a/Engine/Source/Editor/SubobjectEditor/Private/SSubobjectEditor.cpp b/Engine/Source/Editor/SubobjectEditor/Private/SSubobjectEditor.cpp
--- a/Engine/Source/Editor/SubobjectEditor/Private/SSubobjectEditor.cpp
+++ b/Engine/Source/Editor/SubobjectEditor/Private/SSubobjectEditor.cpp
@@ -2537,6 +2537,7 @@ void SSubobjectEditor::UpdateTree(bool bRegenerateTreeNodes /* = true */)
 
 			if (UICustomization)
 			{
+				UICustomization->FilterSubobjectData(SubobjectData);
 				UICustomization->SortSubobjectData(SubobjectData);
 			}
 
