Fixed issue where referencing "set w/ broadcast" nodes wouldn't be updated if the corresponding variable's notify list was updated

diff --git a/Engine/Source/Editor/Kismet/Private/SFieldNotificationCheckList.cpp b/Engine/Source/Editor/Kismet/Private/SFieldNotificationCheckList.cpp
--- a/Engine/Source/Editor/Kismet/Private/SFieldNotificationCheckList.cpp
+++ b/Engine/Source/Editor/Kismet/Private/SFieldNotificationCheckList.cpp
@@ -2,10 +2,12 @@
 
 #include "SFieldNotificationCheckList.h"
 
+#include "AssetRegistry/IAssetRegistry.h"
 #include "Engine/BlueprintGeneratedClass.h"
 #include "INotifyFieldValueChanged.h"
 #include "Layout/Children.h"
 #include "K2Node_FunctionEntry.h"
+#include "K2Node_VariableSet.h"
 #include "Kismet2/BlueprintEditorUtils.h"
 #include "UObject/Class.h"
 #include "UObject/NameTypes.h"
@@ -138,6 +140,43 @@ void SFieldNotificationCheckList::OnVariableCheckBoxChanged(ECheckBoxState InNew
 					BlueprintObj->NewVariables[VarIndex].SetMetaData(FBlueprintMetadata::MD_FieldNotify, TEXT(""));
 				}
 			}
+
+			// Brute force mark all dependent blueprints modified if they have setter nodes for this member
+			// Use asset registry to find all referencing blueprints
+			TArray<FAssetIdentifier> Referencers;
+			const FAssetIdentifier Source(BlueprintObj->GetPackage()->GetFName(), NAME_None);
+			IAssetRegistry::Get()->GetReferencers(Source, Referencers);
+
+			TArray<UBlueprint*> OtherBlueprints;
+			for (const FAssetIdentifier& AssetIdentifier : Referencers)
+			{
+				TArray<FAssetData> Assets;
+				IAssetRegistry::Get()->GetAssetsByPackageName(AssetIdentifier.PackageName, Assets);
+
+				for (const FAssetData& AssetData : Assets)
+				{
+					if (AssetData.GetClass()->IsChildOf(UBlueprint::StaticClass()))
+					{
+						OtherBlueprints.Add(Cast<UBlueprint>(AssetData.GetAsset()));
+					}
+				}
+			}
+
+			// Mark all the blueprints modified if they have broadcasting setter nodes that refer to a variable with this ones name. 
+			for (UBlueprint* OtherBlueprint : OtherBlueprints)
+			{
+				TArray<UK2Node_VariableSet*> SetNodes;
+				FBlueprintEditorUtils::GetAllNodesOfClass(OtherBlueprint, SetNodes);
+
+				for (const UK2Node_VariableSet* SetNode : SetNodes)
+				{
+					if (SetNode->HasFieldNotificationBroadcast() && SetNode->VariableReference.GetMemberName() == FieldName)
+					{
+						FBlueprintEditorUtils::MarkBlueprintAsModified(OtherBlueprint);
+						break;
+					}
+				}
+			}
 		}
 	}
 }
