Fix up and improve BP.EnableActionMenuDumpToFile, which now sorts menu entries and doesn't include an unstable action index in the action menu dump

#rnx

diff --git a/Engine/Source/Editor/Kismet/Private/BlueprintActionMenuBuilder.cpp b/Engine/Source/Editor/Kismet/Private/BlueprintActionMenuBuilder.cpp
--- a/Engine/Source/Editor/Kismet/Private/BlueprintActionMenuBuilder.cpp
+++ b/Engine/Source/Editor/Kismet/Private/BlueprintActionMenuBuilder.cpp
@@ -726,36 +726,38 @@ void FBlueprintActionMenuBuilder::BuildCompleted()
 
 	if (CVarBPEnableActionMenuDumpToFile.GetValueOnGameThread())
 	{
-		TStringBuilder<512> CurrentLine;
-		TStringBuilder<256> CategoryPrefix;
-
 		FString DumpFilePath = FPaths::ProfilingDir() / TEXT("ActionMenu") / FString::Printf(TEXT("MenuDump_%s.txt"), *FDateTime::Now().ToString());
 		FOutputDeviceFile DumpFile(*DumpFilePath, true);
 
 		DumpFile.SetSuppressEventTag(true);
 
 		UE_LOG(LogBlueprintActionMenuItemFactory, Log, TEXT("Dumping menu builder action list to \"%s\"..."), *DumpFilePath);
+		
+		const int32 NumActions = GetNumActions();
 
-		int32 TotalCount = 0;
-
-		for (int32 i = 0; i < GetNumActions(); ++i)
+		TArray<FString> Actions;
+		TStringBuilder<512> CurrentLine;
+		for (int32 I = 0; I < NumActions; ++I)
 		{
-			const ActionGroup& CurrentEntry = GetAction(i);
+			const TSharedPtr<FEdGraphSchemaAction>& Action = GetSchemaAction(I);
 
-			CategoryPrefix.Reset();
-			for (const FString& Category : CurrentEntry.GetCategoryChain())
+			CurrentLine = TEXT("");
+			for (const FString& Category : Action->GetCategoryChain())
 			{
-				CategoryPrefix.Append(FString::Printf(TEXT("%s|"), *Category));
+				CurrentLine.Append(FString::Printf(TEXT("%s|"), *Category));
 			}
-
-			for (const TSharedPtr<FEdGraphSchemaAction>& Action : CurrentEntry.Actions)
+			CurrentLine.Append(Action->GetMenuDescription().ToString());
+			Actions.Emplace(CurrentLine);
+		}
+		Actions.Sort(
+			[](const FString& A, const FString& B)
 			{
-				CurrentLine = FString::Printf(TEXT("%05d "), ++TotalCount);
-				CurrentLine.Append(CategoryPrefix.ToString());
-				CurrentLine.Append(Action->GetMenuDescription().ToString());
-
-				DumpFile.Log(CurrentLine.ToString());
+				return FTextComparison::CompareTo(A, B) < 0;
 			}
+		);
+		for( const FString& Line : Actions )
+		{
+			DumpFile.Log(Line);
 		}
 
 		CVarBPEnableActionMenuDumpToFile.AsVariable()->Set(false);
