Archetype reinstantiation fix, start by reinstantiating the archetype dependencies first by sorting them all together.

diff --git a/Engine/Source/Editor/Kismet/Private/BlueprintCompilationManager.cpp b/Engine/Source/Editor/Kismet/Private/BlueprintCompilationManager.cpp
--- a/Engine/Source/Editor/Kismet/Private/BlueprintCompilationManager.cpp
+++ b/Engine/Source/Editor/Kismet/Private/BlueprintCompilationManager.cpp
@@ -586,31 +586,6 @@ struct FReinstancingJob
 	// always set:
 	TPair<UClass*, UClass*> OldToNew;
 
-	struct FArchetypeInfo
-	{
-		FArchetypeInfo(UObject* Archetype)
-			: Archetype(Archetype)
-			, ArchetypeTemplate(Archetype ? Archetype->GetArchetype() : nullptr)
-		{}
-
-		bool operator==(const FArchetypeInfo& Other) const
-		{
-			return Archetype == Other.Archetype;
-		}
-
-		friend int32 GetTypeHash(const FArchetypeInfo& Info)
-		{
-			return GetTypeHash(Info.Archetype);
-		}
-
-		UObject* Archetype = nullptr;
-		UObject* ArchetypeTemplate;
-	};
-
-
-	// Old archetype to re-instantiate and its old associated template
-	TArray<FArchetypeInfo> OldArchetypeObjects;
-
 	// Keeps track of old archetype objects with cached archetypes which may be needed for reinstancing
 	// of non-archetypes to assist with archetype lookup when they have partially-reinstanced outer chains.
 	// We track these objects because they are annotated by the global archetype cache manager, and we'll
@@ -645,6 +620,30 @@ FReinstancingJob::FReinstancingJob(TPair<UClass*, UClass*> InOldToNew)
 {
 }
 
+struct FArchetypeInfo
+{
+	FArchetypeInfo(UObject* InArchetype, FReinstancingJob* InReinstancingJob = nullptr)
+		: Archetype(InArchetype)
+		, ArchetypeTemplate(InArchetype ? InArchetype->GetArchetype() : nullptr)
+		, ReinstancingJob(InReinstancingJob)
+
+	{}
+
+	bool operator==(const FArchetypeInfo& Other) const
+	{
+		return Archetype == Other.Archetype;
+	}
+
+	friend int32 GetTypeHash(const FArchetypeInfo& Info)
+	{
+		return GetTypeHash(Info.Archetype);
+	}
+
+	UObject* Archetype = nullptr;
+	UObject* ArchetypeTemplate = nullptr;
+	FReinstancingJob* ReinstancingJob = nullptr;
+};
+
 UE_TRACE_EVENT_BEGIN(Cpu, FlushCompilationQueue, NoSync)
 	UE_TRACE_EVENT_FIELD(UE::Trace::WideString, BlueprintPath)
 	UE_TRACE_EVENT_FIELD(int32, TotalBlueprints)
@@ -3014,6 +3013,10 @@ void FBlueprintCompilationManagerImpl::ReinstanceBatch(TArray<FReinstancingJob>&
 		ArchetypeReferencers.Add(GUnrealEd->Trans);
 	}
 
+	// Old archetype to re-instantiate and its old associated template
+	TArray<FArchetypeInfo> OldArchetypeObjects;
+	TSet<UObject*> OldArchetypeObjectsSet;
+
 	for (FReinstancingJob* ReinstancingJobPtr : ReinstancersPtr)
 	{
 		FReinstancingJob& ReinstancingJob = *ReinstancingJobPtr;
@@ -3063,7 +3066,8 @@ void FBlueprintCompilationManagerImpl::ReinstanceBatch(TArray<FReinstancingJob>&
 			// for each archetype:
 			for (UObject* Archetype : ArchetypeObjects)
 			{
-				ReinstancingJob.OldArchetypeObjects.Add(Archetype);
+				OldArchetypeObjects.Emplace(Archetype, ReinstancingJobPtr);
+				OldArchetypeObjectsSet.Add(Archetype);
 
 				// make sure we fix up references in the owner:
 				{
@@ -3116,44 +3120,57 @@ void FBlueprintCompilationManagerImpl::ReinstanceBatch(TArray<FReinstancingJob>&
 					}
 				}
 			}
+		}
+	}
 
-			// Sort all archetype between themselves, as one might depends on an other
-			// This happens when the class containing the archetype has derived classes.
-			Algo::TopologicalSort(ReinstancingJob.OldArchetypeObjects, [&ArchetypeObjects](const FReinstancingJob::FArchetypeInfo& OldArchetypeInfo)
+
+	// Sort all archetype between themselves, as one might depends on an other
+	// This happens when the class containing the archetype has derived classes or embed subobject of other classes.
+	Algo::TopologicalSort(OldArchetypeObjects, [&OldArchetypeObjectsSet](const FArchetypeInfo& OldArchetypeInfo)
+	{
+		TArray<UObject*> Dependencies;
+		if (OldArchetypeObjectsSet.Find(OldArchetypeInfo.ArchetypeTemplate))
+		{
+			Dependencies.Add(OldArchetypeInfo.ArchetypeTemplate);
+		}
+
+		auto GetArchetypeSubObjestDependencies = [&OldArchetypeObjectsSet, &Dependencies](UObject* OldArchetype, auto& Recurse) -> void
+		{
+			TArray<UObject*> OldArchetypeContainedObjects;
+			GetObjectsWithOuter(OldArchetype, OldArchetypeContainedObjects, /*bIncludeNestedObjects*/false);
+			for (UObject* OldArchetypeSubObject : OldArchetypeContainedObjects)
 			{
-				TArray<UObject*> Dependencies;
-				if (OldArchetypeInfo.ArchetypeTemplate && !OldArchetypeInfo.ArchetypeTemplate->HasAnyFlags(RF_ClassDefaultObject))
+				if (OldArchetypeObjectsSet.Find(OldArchetypeSubObject))
 				{
-					if (ArchetypeObjects.Contains(OldArchetypeInfo.ArchetypeTemplate))
-					{
-						Dependencies.Add(OldArchetypeInfo.ArchetypeTemplate);
-					}
+					Dependencies.Add(OldArchetypeSubObject);
 				}
+				Recurse(OldArchetypeSubObject, Recurse);
+			}
+		};
+		GetArchetypeSubObjestDependencies(OldArchetypeInfo.Archetype, GetArchetypeSubObjestDependencies);
 
-				return Dependencies;
-			});
-		}
-	}
+		return Dependencies;
+	});
 
 	FEditorCacheArchetypeManager& CacheManager = FEditorCacheArchetypeManager::Get();
 
 	// This loop finishes the reinstancing of archetypes after the entire Outer hierarchy has been updated with new instances:
-	for (FReinstancingJob* ReinstancingJobPtr : ReinstancersPtr)
 	{
-		FReinstancingJob& ReinstancingJob = *ReinstancingJobPtr;
-		UClass* OldClass = ReinstancingJob.OldToNew.Key;
-		if(OldClass)
-		{
-			SCOPED_LOADTIMER_ASSET_TEXT(*WriteToString<256>(TEXT("FinishReinstancing "), *GetPathNameSafe(ReinstancingJob.OldToNew.Value)));
-
-			UClass* NewClass = ReinstancingJob.OldToNew.Value;
-
-			TMap<UObject*, UObject*>* OldToNewTemplatesForClass = OldToNewTemplates ? &OldToNewTemplates->FindOrAdd(OldClass) : nullptr;
+		SCOPED_LOADTIMER_ASSET_TEXT(TEXT("FinishReinstancing Archetypes"));
 
-			for(const FReinstancingJob::FArchetypeInfo& OldArchetypeInfo : ReinstancingJob.OldArchetypeObjects)
+		for (const FArchetypeInfo& OldArchetypeInfo : OldArchetypeObjects)
+		{
+			checkf(OldArchetypeInfo.ReinstancingJob, TEXT("Expecting the job ptr to be set."));
+			FReinstancingJob& ReinstancingJob = *OldArchetypeInfo.ReinstancingJob;
+			UClass* OldClass = ReinstancingJob.OldToNew.Key;
+			if (OldClass)
 			{
+				UClass* NewClass = ReinstancingJob.OldToNew.Value;
 				UObject* OldInstance = OldArchetypeInfo.Archetype;
 
+				TMap<UObject*, UObject*> OldToNewInstanceMapTemp;
+				TMap<UObject*, UObject*>* OldToNewTemplatesForClass = OldToNewTemplates ? &OldToNewTemplates->FindOrAdd(OldClass) : &OldToNewInstanceMapTemp;
+
 				// move aside:
 				FName OriginalName = OldInstance->GetFName();
 				UObject* OriginalOuter = OldInstance->GetOuter();
@@ -3176,12 +3193,27 @@ void FBlueprintCompilationManagerImpl::ReinstanceBatch(TArray<FReinstancingJob>&
 					Template = *NewArchetypeTemplate;
 				}
 
+				TArray<UClass*> ClassesOfArchetypeOuterChain;
 				FObjectInstancingGraph PrePopulatedInstancingGraph;
 				const bool bUsePrePopulatedInstancingGraph = FBlueprintCompileReinstancer::PrePopulateInstancingGraphForArchetype(PrePopulatedInstancingGraph, OldInstance, Template, &OldArchetypeToNewArchetype);
 				if (bUsePrePopulatedInstancingGraph)
 				{
-					// Make sure the template object as all it references path up before using to instance another archetype
-					FArchiveReplaceObjectRef<UObject> ReplaceInCDOAr(Template, OldArchetypeToNewArchetype);
+					// Make sure the old archetype has all it references patched up before using copying
+					FArchiveReplaceObjectRef<UObject> ReplaceInOldInstanceAr(OldInstance, *OldToNewTemplatesForClass);
+
+					// Make sure the template object has all it references patched up before using it to this archetype
+					FArchiveReplaceObjectRef<UObject> ReplaceInCDOAr(Template, *OldToNewTemplatesForClass);
+
+					// This needs to be done before we move the archetype to transient package is done, otherwise the outerchain is lost
+					UObject* ArchetypeOuter = OldInstance->GetOuter();
+					while (ArchetypeOuter)
+					{
+						if (InOutOldToNewClassMap.Find(ArchetypeOuter->GetClass()))
+						{
+							ClassesOfArchetypeOuterChain.Add(ArchetypeOuter->GetClass());
+						}
+						ArchetypeOuter = ArchetypeOuter->GetOuter();
+					}
 				}
 
 				UObject* Destination = GetTransientOuterForRename(OldInstance->GetClass());
@@ -3200,10 +3232,7 @@ void FBlueprintCompilationManagerImpl::ReinstanceBatch(TArray<FReinstancingJob>&
 				UObject* NewArchetype = NewObject<UObject>(OriginalOuter, NewClass, OriginalName, OriginalFlags & FlagMask, Template, false, bUsePrePopulatedInstancingGraph ? &PrePopulatedInstancingGraph : nullptr);
 
 				OldArchetypeToNewArchetype.Add(OldInstance, NewArchetype);
-				if (OldToNewTemplatesForClass)
-				{
-					OldToNewTemplatesForClass->Add(OldInstance, NewArchetype);
-				}
+				OldToNewTemplatesForClass->Add(OldInstance, NewArchetype);
 
 				// also map old *default* subobjects to new default subobjects:
 				BuildDSOMap(OldInstance, NewArchetype, InOutOldToNewClassMap, OldArchetypeToNewArchetype);
@@ -3220,31 +3249,33 @@ void FBlueprintCompilationManagerImpl::ReinstanceBatch(TArray<FReinstancingJob>&
 					// up below:			
 					const bool bUseDeltaSerialization = ReinstancingJob.Reinstancer.IsValid() ? ReinstancingJob.Reinstancer->bUseDeltaSerializationToCopyProperties : false;
 
+					if (TMap<UObject*, UObject*>* OldToNewTemplateMapForSuperClass = OldToNewTemplates ? OldToNewTemplates->Find(OldClass->GetSuperClass()) : nullptr)
+					{
+						OldToNewTemplatesForClass->Append(*OldToNewTemplateMapForSuperClass);
+					}
+
 					TMap<UObject*, UObject*> CreatedInstanceMap;
 					TArray< TTuple<UObject*, UObject*>> OrderedListOfObjectToCopy;
-					FBlueprintCompileReinstancer::PreCreateSubObjectsForReinstantiation(InOutOldToNewClassMap, OldInstance, NewArchetype, CreatedInstanceMap, nullptr, &OrderedListOfObjectToCopy);
+					FBlueprintCompileReinstancer::PreCreateSubObjectsForReinstantiation(InOutOldToNewClassMap, OldInstance, NewArchetype, CreatedInstanceMap, OldToNewTemplatesForClass, &OrderedListOfObjectToCopy);
 
 					// We only need to copy properties of the pre-created instances, the rest of the default sub object is done inside the UEditorEngine::CopyPropertiesForUnrelatedObjects
-					TMap<UObject*, UObject*> OldToNewInstanceMap(CreatedInstanceMap);
-					if (TMap<UObject*, UObject*>* OldToNewTemplateMap = OldToNewTemplates ? OldToNewTemplates->Find(OldClass->GetSuperClass()) : nullptr)
-					{
-						OldToNewInstanceMap.Append(*OldToNewTemplateMap);
-					}
+					OldToNewTemplatesForClass->Append(CreatedInstanceMap);
 					for (const TTuple<UObject*, UObject*>& Pair : OrderedListOfObjectToCopy)
 					{
-						FBlueprintCompileReinstancer::CopyPropertiesForUnrelatedObjects(Pair.Key, Pair.Value, /*bClearExternalReferences*/false, bUseDeltaSerialization, /*bOnlyHandleDirectSubObjects*/true, &OldToNewInstanceMap, &InOutOldToNewClassMap);
+						FBlueprintCompileReinstancer::CopyPropertiesForUnrelatedObjects(Pair.Key, Pair.Value, /*bClearExternalReferences*/false, bUseDeltaSerialization, /*bOnlyHandleDirectSubObjects*/true, OldToNewTemplatesForClass, &InOutOldToNewClassMap);
 					}
 					FOverridableManager::Get().RestoreOverrideState(OldInstance, NewArchetype);
 
-					if (OldToNewTemplates && !OldToNewInstanceMap.IsEmpty())
+					// Append to any outer archetype classes, in case they are also re-instantiated during this pass.
+					for (UClass* OuterArchetypeClasses : ClassesOfArchetypeOuterChain)
 					{
-						OldToNewTemplates->FindOrAdd(OldClass).Append(OldToNewInstanceMap);
+						OldToNewTemplates->FindOrAdd(OuterArchetypeClasses).Append(*OldToNewTemplatesForClass);
 					}
 
 					// Ensure that our archetype instance mapping is up to date, as this will be used for reference replacement below. Note that this
 					// might be populated by BuildDSO() initially to contain instance mapping pairs for which the new instance won't be valid until we
 					// run through each of the PreCreateSubObjectsForReinstantiation() and subsequent CopyPropertiesForUnrelatedObjects() steps above.
-					for (const TPair<UObject*, UObject*>& Pair : OldToNewInstanceMap)
+					for (const TPair<UObject*, UObject*>& Pair : *OldToNewTemplatesForClass)
 					{
 						if (UObject** NewArchetypeMapping = OldArchetypeToNewArchetype.Find(Pair.Key))
 						{
diff --git a/Engine/Source/Editor/UnrealEd/Private/Kismet2/KismetReinstanceUtilities.cpp b/Engine/Source/Editor/UnrealEd/Private/Kismet2/KismetReinstanceUtilities.cpp
--- a/Engine/Source/Editor/UnrealEd/Private/Kismet2/KismetReinstanceUtilities.cpp
+++ b/Engine/Source/Editor/UnrealEd/Private/Kismet2/KismetReinstanceUtilities.cpp
@@ -2709,6 +2709,15 @@ void FBlueprintCompileReinstancer::ReplaceInstancesOfClass_Inner(const TMap<UCla
 		}
 		else
 		{
+			// Pre gather all the replacement map so that while doing reintanciation, we have all the information needed
+			for (TPair<UClass*, UClass*> OldToNewClass : InOldToNewClassMap)
+			{
+				if (TMap<UObject*, UObject*>* ReplaceTemplateMapping = Params.OldToNewTemplates ? Params.OldToNewTemplates->Find(OldToNewClass.Key) : nullptr)
+				{
+					OldToNewInstanceMap.Append(*ReplaceTemplateMapping);
+				}
+			}
+
 		// WARNING: for (TPair<UClass*, UClass*> OldToNewClass : InOldToNewClassMap) duplicated below 
 		// to handle reconstructing actors which need to be reinstanced after their owned components 
 		// have been updated:
@@ -2727,11 +2736,6 @@ void FBlueprintCompileReinstancer::ReplaceInstancesOfClass_Inner(const TMap<UCla
 					bFixupSCS = (NewClass->IsChildOf<USceneComponent>() != OldClass->IsChildOf<USceneComponent>());
 				}
 
-				if (TMap<UObject*, UObject*>* ReplaceTemplateMapping = Params.OldToNewTemplates ? Params.OldToNewTemplates->Find(OldClass) : nullptr)
-				{
-					OldToNewInstanceMap.Append(*ReplaceTemplateMapping);
-				}
-
 				const bool bIncludeDerivedClasses = false;
 				ObjectsToReplace.Reset();
 				GetObjectsOfClass(OldClass, ObjectsToReplace, bIncludeDerivedClasses);
@@ -3309,12 +3313,6 @@ bool FBlueprintCompileReinstancer::PrePopulateInstancingGraphForArchetype(FObjec
 {
 	checkf(!InstancingGraph.HasDestinationRoot(), TEXT("Can only pre-populate an empty InstancingGraph"));
 
-	// Check if we are re-instancing from an archetype, otherwise nothing to do
-	if (!InArchetype || InArchetype->HasAnyFlags(RF_ClassDefaultObject))
-	{
-		return false;
-	}
-
 	// Only needed for objects that are using OS for now
 	if (!FOverridableManager::Get().IsEnabled(OldInstance))
 	{
@@ -3426,8 +3424,15 @@ void FBlueprintCompileReinstancer::PreCreateSubObjectsForReinstantiation_Inner(
 			if (UObject* const* AlreadyCreatedSubObject = OldToNewInstanceMap ? OldToNewInstanceMap->Find(OldSubObject) : nullptr)
 			{
 				int32 AlreadyCreatedSubObjectIndex = ContainedOldSubObjects.Find(*AlreadyCreatedSubObject);
-				checkf(AlreadyCreatedSubObjectIndex > i, TEXT("Expecting the already created subobject to be in the old subobject list after this sub object"));
-				ContainedOldSubObjects.RemoveAt(AlreadyCreatedSubObjectIndex);
+				if (AlreadyCreatedSubObjectIndex > i)
+				{
+					ContainedOldSubObjects.RemoveAt(AlreadyCreatedSubObjectIndex);
+				}
+				else
+				{
+					// This happens on object that were reinstantiated and were not yet outered to the right object because it was not yet created
+					checkf(OldInstancedSubObjects && !OldInstancedSubObjects->Contains(*AlreadyCreatedSubObject), TEXT("The object must have already been created but should have been skipped by the fact that the old object was not referencing it"));
+				}
 				(*AlreadyCreatedSubObject)->Rename(nullptr, NewUObject, REN_DoNotDirty | REN_DontCreateRedirectors | REN_NonTransactional);
 			}
 			else
@@ -3472,7 +3477,14 @@ void FBlueprintCompileReinstancer::PreCreateSubObjectsForReinstantiation_Inner(
 					}
 					else
 					{
-						NewSubObject = NewObject<UObject>(NewUObject, SubObjectClass, SubObjectName, SubObjectFlags);
+						// Making sure the template we are about to use had all its references patched up
+						UObject* Template = SubObjectClass->GetDefaultObject();
+						if (OldToNewInstanceMap)
+						{
+							FArchiveReplaceObjectRef<UObject> ReplaceInCDOAr(Template, *OldToNewInstanceMap);
+						}
+
+						NewSubObject = NewObject<UObject>(NewUObject, SubObjectClass, SubObjectName, SubObjectFlags, Template);
 					}
 					CreatedInstanceMap.Add(OldSubObject, NewSubObject);
 					PreCreateSubObjectsForReinstantiation_Inner(OldInstancedSubObjects, OldToNewClassMap, OldSubObject, NewSubObject, CreatedInstanceMap, OldToNewInstanceMap, OrderedListOfObjectToCopy);
