BugHawk: fix possible null ptr usage.

diff --git a/Engine/Source/Editor/Kismet/Private/BlueprintCompilationManager.cpp b/Engine/Source/Editor/Kismet/Private/BlueprintCompilationManager.cpp
--- a/Engine/Source/Editor/Kismet/Private/BlueprintCompilationManager.cpp
+++ b/Engine/Source/Editor/Kismet/Private/BlueprintCompilationManager.cpp
@@ -3267,9 +3267,12 @@ void FBlueprintCompilationManagerImpl::ReinstanceBatch(TArray<FReinstancingJob>&
 					FOverridableManager::Get().RestoreOverrideState(OldInstance, NewArchetype);
 
 					// Append to any outer archetype classes, in case they are also re-instantiated during this pass.
-					for (UClass* OuterArchetypeClasses : ClassesOfArchetypeOuterChain)
+					if (OldToNewTemplates)
 					{
-						OldToNewTemplates->FindOrAdd(OuterArchetypeClasses).Append(*OldToNewTemplatesForClass);
+						for (UClass* OuterArchetypeClasses : ClassesOfArchetypeOuterChain)
+						{
+							OldToNewTemplates->FindOrAdd(OuterArchetypeClasses).Append(*OldToNewTemplatesForClass);
+						}
 					}
 
 					// Ensure that our archetype instance mapping is up to date, as this will be used for reference replacement below. Note that this
