UE-347821 - Crash at compile time when using a VariableSet node with a local variable created by pasting a private property

When pasting a class property marked as "private" as a local variable in a function, we need to strip the "private" metadata. Otherwise, this can accidentally trip any logic that expects a private property to have a UClass as an outer (eg: IsPropertyWritableInBlueprint).

diff --git a/Engine/Source/Editor/Kismet/Private/SMyBlueprint.cpp b/Engine/Source/Editor/Kismet/Private/SMyBlueprint.cpp
--- a/Engine/Source/Editor/Kismet/Private/SMyBlueprint.cpp
+++ b/Engine/Source/Editor/Kismet/Private/SMyBlueprint.cpp
@@ -3929,7 +3929,7 @@ void SMyBlueprint::OnPasteLocalVariable()
 					FScopedTransaction Transaction(FText::Format(LOCTEXT("PasteLocalVariable", "Paste Local Variable: {0}"), FText::FromName(NewVar.VarName)));
 
 					NewVar.Category = GetPasteCategory();
-
+					NewVar.RemoveMetaData(FBlueprintMetadata::MD_Private);
 
 					if (FunctionEntry.Num() == 1)
 					{
