GameplayTag dependencies for blueprints :

- GameplayTagManager : factorize the code creating dependencies for a gameplay tag so it can be used for blueprints.
- BlueprintDependencies : go through all the pins of type FGameplayType to create a dependency. Create a new type of blueprint dependency for gameplay tag.
- Kismet module : new dependency to module GameplayTags.

#rnx

diff --git a/Engine/Source/Editor/Kismet/Internal/Blueprints/BlueprintDependencies.h b/Engine/Source/Editor/Kismet/Internal/Blueprints/BlueprintDependencies.h
--- a/Engine/Source/Editor/Kismet/Internal/Blueprints/BlueprintDependencies.h
+++ b/Engine/Source/Editor/Kismet/Internal/Blueprints/BlueprintDependencies.h
@@ -13,7 +13,8 @@ enum class EBPDependencyType
 {
 	Asset,
 	Struct,
-	Class
+	Class,
+	GameplayTag
 };
 
 USTRUCT()
diff --git a/Engine/Source/Editor/Kismet/Kismet.Build.cs b/Engine/Source/Editor/Kismet/Kismet.Build.cs
--- a/Engine/Source/Editor/Kismet/Kismet.Build.cs
+++ b/Engine/Source/Editor/Kismet/Kismet.Build.cs
@@ -61,6 +61,7 @@ public class Kismet : ModuleRules
 				"SubobjectDataInterface",
 				"ToolWidgets",
 				"TraceLog",
+				"GameplayTags"
 			}
 			);
 
diff --git a/Engine/Source/Editor/Kismet/Private/Blueprints/BlueprintDependencies.cpp b/Engine/Source/Editor/Kismet/Private/Blueprints/BlueprintDependencies.cpp
--- a/Engine/Source/Editor/Kismet/Private/Blueprints/BlueprintDependencies.cpp
+++ b/Engine/Source/Editor/Kismet/Private/Blueprints/BlueprintDependencies.cpp
@@ -6,10 +6,13 @@
 #include "AssetRegistry/IAssetRegistry.h"
 #include "Cooker/CookDependencyContext.h"
 #include "Engine/Blueprint.h"
+#include "GameplayTagContainer.h"
+#include "GameplayTagsManager.h"
 #include "HAL/FileManager.h"
 #include "IO/IoHash.h"
 #include "JsonObjectConverter.h"
 #include "JsonObjectGraph/Stringify.h"
+#include "K2Node_CallFunction.h"
 #include "Kismet2/BlueprintEditorUtils.h"
 #include "Misc/CommandLine.h"
 #include "Misc/Optional.h"
@@ -51,20 +54,29 @@ DEFINE_LOG_CATEGORY(LogCookBlueprint);
 
 namespace UE::Private
 {
-constexpr int32 BlueprintCookDependenciesVersion = 3;
+constexpr int32 BlueprintCookDependenciesVersion = 4;
 
 // gathering helpers:
-static void GetAllImportedObjects(TConstArrayView<const UObject*> Roots, TArray<const UObject*>& OutImports, TSet<FName>& OutSoftImports);
-static void GetAllImportedObjects(const UBlueprint* ForBP, TArray<const UObject*>& OutImports, TSet<FName>& OutSoftImports);
+static void GetAllImportedObjects(
+	TConstArrayView<const UObject*> Roots, 
+	TArray<const UObject*>& OutImports, 
+	TSet<FName>& OutSoftImports,
+	TSet<FName>& OutTagNames);
+
+static void GetAllImportedObjects(const UBlueprint* ForBP, TArray<const UObject*>& OutImports, TSet<FName>& OutSoftImports, TSet<FName>& OutTagNames);
+
 // helpers for handling transitiveness of dependencies (e.g. if i use a struct, i also depend on its member structs):
 static void AddStructDependencyImpl(const UScriptStruct* Struct, TSet<const UScriptStruct*>& ReferencedStructs);
 static void AddClassDependencyImpl(const UClass* Class, TSet<const UClass*>& ReferencedClasses);
+
 // returns all the dependencies of a bp, for recording/generating a snapshot of dependencies:
 static void GetAllDependencies(
 	const UBlueprint* ForBP, 
 	TArray<FName>& AssetDependencies, 
 	TArray<const UClass*>& NativeClassDependencies,
-	TArray<const UScriptStruct*>& NativeStructDependencies);
+	TArray<const UScriptStruct*>& NativeStructDependencies,
+	TArray<FName>& GameplayTagsArray);
+
 // Native struct and class canonization:
 static FString HashNativeStruct(const UScriptStruct* Struct);
 static FString HashNativeClass( const UClass* NativeClass);
@@ -109,7 +121,11 @@ static bool LoadBPCookDependenciesImpl(FCbFieldViewIterator Args, UE::Cook::FCoo
 static void SaveBPCookDepenenciesImpl(FCbWriter& Writer, const FBlueprintDependencies& Dependencies);
 }
 
-static void UE::Private::GetAllImportedObjects(TConstArrayView<const UObject*> Roots, TArray<const UObject*>& OutImports, TSet<FName>& OutSoftImports)
+static void UE::Private::GetAllImportedObjects(
+	TConstArrayView<const UObject*> Roots, 
+	TArray<const UObject*>& OutImports, 
+	TSet<FName>& OutSoftImports,
+	TSet<FName>& OutTagNames)
 {
 	// This function just gathers all of the UObjects the Root objects (or their inners)
 	// depend upon. Inners are not currently tested for reachability. Tautological self reference
@@ -214,14 +230,42 @@ static void UE::Private::GetAllImportedObjects(TConstArrayView<const UObject*> R
 	for(const UObject* Obj: Objects)
 	{
 		FExternalReferenceFinder References(Obj, Result, RootSet);
+
+		if (const UK2Node_CallFunction* Node = Cast<UK2Node_CallFunction>(Obj))
+		{
+			int PinIndex = 0;
+			while (const UEdGraphPin* Pin = Node->GetPinAt(PinIndex))
+			{
+				++PinIndex;
+
+				const FEdGraphPinType& PinType = Pin->PinType;
+				const UScriptStruct* PinStruct = Cast<UScriptStruct>(PinType.PinSubCategoryObject);
+				
+				if(PinStruct != FGameplayTag::StaticStruct())
+				{
+					continue;
+				}
+
+				FString DefaultValue = Pin->GetDefaultAsString();
+
+				if (DefaultValue.IsEmpty())
+				{
+					continue;
+				}
+				
+				FGameplayTag Tag;
+				PinStruct->ImportText(*DefaultValue, &Tag, nullptr, 0, nullptr, PinStruct->GetName(), false);
+				OutTagNames.Add(Tag.GetTagName());
+			}
+		}
 	}
 	
 	OutImports = Result.Array();
 }
 
-static void UE::Private::GetAllImportedObjects( const UBlueprint* ForBP, TArray<const UObject*>& OutImports, TSet<FName>& OutSoftImports)
+static void UE::Private::GetAllImportedObjects( const UBlueprint* ForBP, TArray<const UObject*>& OutImports, TSet<FName>& OutSoftImports, TSet<FName>& OutTagNames)
 {
-	return GetAllImportedObjects({ForBP, ForBP->GeneratedClass, ForBP->GeneratedClass->GetDefaultObject(false) }, OutImports, OutSoftImports);
+	return GetAllImportedObjects({ForBP, ForBP->GeneratedClass, ForBP->GeneratedClass->GetDefaultObject(false) }, OutImports, OutSoftImports, OutTagNames);
 }
 
 static void UE::Private::AddStructDependencyImpl(const UScriptStruct* Struct, TSet<const UScriptStruct*>& ReferencedStructs)
@@ -266,14 +310,16 @@ static void UE::Private::GetAllDependencies(
 	const UBlueprint* ForBP,
 	TArray<FName>& AssetDependencies, 
 	TArray<const UClass*>& NativeClassDependencies,
-	TArray<const UScriptStruct*>& NativeStructDependencies)
+	TArray<const UScriptStruct*>& NativeStructDependencies,
+	TArray<FName>& GameplayTagsArray)
 {
 	TRACE_CPUPROFILER_EVENT_SCOPE(BlueprintDependencies::GetAllDependencies);
 
 	// gather all references in the BP, its class, and the CDO, along with any subobjects:
 	TArray<const UObject*> ImportedObjects;
 	TSet<FName> ReferencedAssets;
-	GetAllImportedObjects(ForBP, ImportedObjects, ReferencedAssets);
+	TSet<FName> TagNames;
+	GetAllImportedObjects(ForBP, ImportedObjects, ReferencedAssets, TagNames);
 
 	// honor transitive dependencies - e.g. inherited classes, inherited structs, aggregate structs
 	// for now all class and struct references are fully transitive:
@@ -351,6 +397,9 @@ static void UE::Private::GetAllDependencies(
 	{
 		return A.LexicalLess(B);
 	});
+
+	GameplayTagsArray = TagNames.Array();
+	GameplayTagsArray.Sort(FNameLexicalLess());
 }
 
 static FString UE::Private::HashNativeStruct(const UScriptStruct* Struct)
@@ -652,6 +701,20 @@ void UE::Private::ValidateBPCookDependenciesImpl(FCbFieldViewIterator Args, UE::
 				}
 				break;
 			}
+			case EBPDependencyType::GameplayTag:
+			{
+				FBlake3 Hasher;
+				UGameplayTagsManager::Get().AppendDependencyHash(BPDep.NativeObjectName, Hasher);
+				FBlake3Hash NewHash = Hasher.Finalize();
+
+				FBlake3Hash OldHash(BPDep.Hash);
+				if (NewHash != OldHash)
+				{
+					Context.LogInvalidated(FString::Printf(TEXT("Gameplay tag changed : %s"), *BPDep.NativeObjectName.ToString()));
+					return;
+				}
+				break;
+			}
 		}
 	}
 }
@@ -662,7 +725,8 @@ void UE::Private::GenerateBlueprintDependencies(const UBlueprint* ForBP, FBluepr
 	TArray<FName> AssetDependencies;
 	TArray<const UClass*> NativeClassDependencies;
 	TArray<const UScriptStruct*> NativeStructDependencies;
-	UE::Private::GetAllDependencies(ForBP, AssetDependencies, NativeClassDependencies, NativeStructDependencies);
+	TArray<FName> GameplayTagsArray;
+	UE::Private::GetAllDependencies(ForBP, AssetDependencies, NativeClassDependencies, NativeStructDependencies, GameplayTagsArray);
 
 	TRACE_CPUPROFILER_EVENT_SCOPE(BlueprintDependencies::GenerateBlueprintDependencies);
 	// MUST MATCH UE::Private::ValidateBPCookDependenciesImpl
@@ -708,6 +772,22 @@ void UE::Private::GenerateBlueprintDependencies(const UBlueprint* ForBP, FBluepr
 			MoveTemp(Dependency)
 		);
 	}
+
+	for (const FName& SingleTagName : GameplayTagsArray)
+	{
+		FBlake3 Hasher;
+		UGameplayTagsManager::Get().AppendDependencyHash(SingleTagName, Hasher);
+		FBlake3Hash Hash = Hasher.Finalize();
+
+		FBlueprintDependency Dependency = {
+			.DependencyType = EBPDependencyType::GameplayTag,
+			.PackageName = FName(),
+			.NativeObjectName = SingleTagName,
+			.Hash = LexToString(Hash)
+		};
+
+		OutDependencies.BlueprintDependencies.Emplace(MoveTemp(Dependency));
+	}
 }
 // END PAIR VALIDATE/GENERATE
 
diff --git a/Engine/Source/Runtime/GameplayTags/Classes/GameplayTagsManager.h b/Engine/Source/Runtime/GameplayTags/Classes/GameplayTagsManager.h
--- a/Engine/Source/Runtime/GameplayTags/Classes/GameplayTagsManager.h
+++ b/Engine/Source/Runtime/GameplayTags/Classes/GameplayTagsManager.h
@@ -26,6 +26,7 @@ class UGameplayTagsList;
 struct FStreamableHandle;
 
 #if WITH_EDITOR
+namespace GameplayTagUtil { struct FGameplayTagCookDependency; }
 namespace UE::Cook { class FCookDependency; }
 namespace UE::Cook { class ICookInfo; }
 namespace UE::Cook { struct FCookDependencyContext; }
@@ -785,6 +786,8 @@ public:
 	 */
 	GAMEPLAYTAGS_API UE::Cook::FCookDependency CreateCookDependency(const FGameplayTag& Tag) const;
 
+	GAMEPLAYTAGS_API void AppendDependencyHash(const FName& TagName, FBlake3& OutHasher) const;
+
 	/** Implementation of console command GameplayTags.DumpSources */
 	void DumpSources(FOutputDevice& Out) const;
 #endif //WITH_EDITOR
@@ -924,6 +927,9 @@ private:
 	GAMEPLAYTAGS_API void HandleGameplayTagTreeChanged(bool bRecreateTree);
 
 #if WITH_EDITOR
+
+	GameplayTagUtil::FGameplayTagCookDependency CreateGameplayTagCookDependency(const FGameplayTag& Tag) const;
+	
 	void ValidateGameplayTagCookDependencies(const FCbFieldViewIterator& Args, UE::Cook::FCookDependencyContext& Context) const;
 #endif //#if WITH_EDITOR
 
diff --git a/Engine/Source/Runtime/GameplayTags/Private/GameplayTagsManager.cpp b/Engine/Source/Runtime/GameplayTags/Private/GameplayTagsManager.cpp
--- a/Engine/Source/Runtime/GameplayTags/Private/GameplayTagsManager.cpp
+++ b/Engine/Source/Runtime/GameplayTags/Private/GameplayTagsManager.cpp
@@ -234,6 +234,21 @@ namespace GameplayTagUtil
 			return true;
 		}
 
+		void Hash(FBlake3& OutHasher)
+		{
+			FString TagNameString = TagName.ToString();
+			OutHasher.Update(TagNameString.GetCharArray().GetData(), TagNameString.NumBytesWithoutNull());
+
+			SourceTagNames.Sort(FNameLexicalLess());
+			for (const FName& SourceTagName : SourceTagNames)
+			{
+				FString SourceTagNameString = SourceTagName.ToString();
+				OutHasher.Update(SourceTagNameString.GetCharArray().GetData(), SourceTagNameString.NumBytesWithoutNull());
+			}
+
+			OutHasher.Update(&bTagExists, sizeof(bTagExists));
+		}
+
 		FName TagName;
 		TArray<FName> SourceTagNames; //if the tag was redirected, contains all the source tag names.
 		bool bTagExists;
@@ -1989,6 +2004,24 @@ public:
 UE_COOK_DEPENDENCY_FUNCTION(GameplayTagDependency, UGameplayTagsManagerIncrementalCookFunctions::Validate);
 
 UE::Cook::FCookDependency UGameplayTagsManager::CreateCookDependency(const FGameplayTag& Tag) const
+{
+	GameplayTagUtil::FGameplayTagCookDependency Dependencies = CreateGameplayTagCookDependency(Tag);
+
+	FCbWriter Writer;
+	Dependencies.Save(Writer);
+	FCbFieldIterator Fields = Writer.Save();
+
+	return UE::Cook::FCookDependency::Function(UE_COOK_DEPENDENCY_FUNCTION_CALL(GameplayTagDependency), MoveTemp(Fields));
+}
+
+void UGameplayTagsManager::AppendDependencyHash(const FName& TagName, FBlake3& OutHasher) const
+{
+	FGameplayTag Tag(TagName);
+	GameplayTagUtil::FGameplayTagCookDependency Dependencies = CreateGameplayTagCookDependency(Tag);
+	Dependencies.Hash(OutHasher);
+}
+
+GameplayTagUtil::FGameplayTagCookDependency UGameplayTagsManager::CreateGameplayTagCookDependency(const FGameplayTag& Tag) const
 {
 	GameplayTagUtil::FGameplayTagCookDependency Dependencies;
 	{
@@ -2004,11 +2037,7 @@ UE::Cook::FCookDependency UGameplayTagsManager::CreateCookDependency(const FGame
 		Dependencies.TagName = Tag.GetTagName();
 	}
 
-	FCbWriter Writer;
-	Dependencies.Save(Writer);
-	FCbFieldIterator Fields = Writer.Save();
-
-	return UE::Cook::FCookDependency::Function(UE_COOK_DEPENDENCY_FUNCTION_CALL(GameplayTagDependency), MoveTemp(Fields));
+	return Dependencies;
 }
 
 void UGameplayTagsManager::ValidateGameplayTagCookDependencies(const FCbFieldViewIterator& Args, UE::Cook::FCookDependencyContext& Context) const
