Find in Blueprint: Will now search keys (in addition to values) in the default key-value pairs in blueprint Map variables. Previously only values were searched.
Because the current metadata model does not explicitly save whether an FiB-indexed property represents a map, this has to be estimated. This isn't the cleanest, but beats the overhead of modifying the metadata model requiring all assets to be resaved.

diff --git a/Engine/Source/Editor/Kismet/Private/ImaginaryBlueprintData.cpp b/Engine/Source/Editor/Kismet/Private/ImaginaryBlueprintData.cpp
--- a/Engine/Source/Editor/Kismet/Private/ImaginaryBlueprintData.cpp
+++ b/Engine/Source/Editor/Kismet/Private/ImaginaryBlueprintData.cpp
@@ -10,6 +10,7 @@
 #include "Dom/JsonObject.h"
 #include "Dom/JsonValue.h"
 #include "Engine/Blueprint.h"
+#include "EdGraphSchema_K2.h"
 #include "HAL/PlatformCrt.h"
 #include "Internationalization/Internationalization.h"
 #include "Internationalization/StringTableCore.h"
@@ -228,8 +229,19 @@ void FImaginaryFiBData::ParseAllChildData_Internal(ESearchableValueStatus InSear
 			if (!TrySpecialHandleJsonValue(KeyText, JsonValue))
 			{
 				TArray<FSearchableValueInfo> ParsedValues;
+
+				// Parse the value making it searchable
 				ParseJsonValue(KeyText, KeyText, JsonValue, ParsedValues, false, InSearchabilityOverride);
 
+				// For some property types like TMaps we want KeyText to be a search candidate since it's entered by the user.
+				// In most other cases KeyText represents property or category names which we don't want to return as search results.
+				if (AreKeysSearchable())
+				{
+					FSearchableValueInfo& ParsedValue = ParsedValues.AddDefaulted_GetRef();
+					ParsedValue = FSearchableValueInfo(FText::FromString("Key"), KeyText, InSearchabilityOverride);
+					ParsedTagsAndValues.Add(FindInBlueprintsHelpers::FSimpleFTextKeyStorage(KeyText), ParsedValue);
+				}
+
 				if (FFindInBlueprintSearchManager::Get().ShouldEnableSearchResultTemplates())
 				{
 					for (const FSearchableValueInfo& ParsedValue : ParsedValues)
@@ -283,7 +295,7 @@ void FImaginaryFiBData::ParseJsonValue(FText InKey, FText InDisplayKey, TSharedP
 	ESearchableValueStatus SearchabilityStatus = (InSearchabilityOverride == ESearchableValueStatus::Searchable)? GetSearchabilityStatus(InKey.ToString()) : InSearchabilityOverride;
 	if( InJsonValue->Type == EJson::Array)
 	{
-		TSharedPtr< FCategorySectionHelper, ESPMode::ThreadSafe > ArrayCategory = MakeShareable(new FCategorySectionHelper(AsShared(), LookupTablePtr, InKey, true));
+		TSharedPtr< FCategorySectionHelper, ESPMode::ThreadSafe > ArrayCategory = MakeShareable(new FCategorySectionHelper(AsShared(), LookupTablePtr, InKey, /*bIsTagAndValueCategory=*/true));
 		ParsedChildData.Add(ArrayCategory);
 		TArray<TSharedPtr< FJsonValue > > ArrayList = InJsonValue->AsArray();
 		for( int32 ArrayIdx = 0; ArrayIdx < ArrayList.Num(); ++ArrayIdx)
@@ -298,10 +310,28 @@ void FImaginaryFiBData::ParseJsonValue(FText InKey, FText InDisplayKey, TSharedP
 		// bIsTagAndValueCategory represents whether a match with the category's contents should return its parent as a search result.
 		// Let the following category return this as search result when:
 		// - The following category represents an array element, return this (the array)
+		// - The following category represents a map entry, return this (the map)
 		// - The following category represents the default value for a struct or object, return this (the default value)
 		// - The following category represents a nested element for one of the above
-		const bool bIsTagAndValueCategory = bIsInArray || IsTagAndValueCategory() || InDisplayKey.EqualTo(FFindInBlueprintSearchTags::FiB_DefaultValue);
-		TSharedPtr< FCategorySectionHelper, ESPMode::ThreadSafe > SubObjectCategory = MakeShareable(new FCategorySectionHelper(AsShared(), InJsonValue->AsObject(), LookupTablePtr, InDisplayKey, bIsTagAndValueCategory));
+		const bool bIsDefaultValueKey = InDisplayKey.EqualTo(FFindInBlueprintSearchTags::FiB_DefaultValue);
+		const bool bIsTagAndValueCategory = bIsInArray || IsTagAndValueCategory() || bIsDefaultValueKey;
+
+		// Evaluate now whether this child's keys should be searchable. Currently this requires peeking at the value syntax 
+		// because we don't store metadata to detect that we're processing a TMap. It's detectable from value though.
+		TSharedPtr<FJsonObject> JsonValueObject = InJsonValue->AsObject();
+		const bool bAreChildKeysSearchable = bIsDefaultValueKey && ShouldSearchChildKeys(JsonValueObject);
+
+		ECategorySectionHelperFlags CategoryFlags = ECategorySectionHelperFlags::None;
+		if (bIsTagAndValueCategory)
+		{
+			CategoryFlags |= ECategorySectionHelperFlags::IsTagAndValue;
+		}
+		if (bAreChildKeysSearchable)
+		{
+			CategoryFlags |= ECategorySectionHelperFlags::AreKeysSearchable;
+		}
+
+		TSharedPtr< FCategorySectionHelper, ESPMode::ThreadSafe > SubObjectCategory = MakeShareable(new FCategorySectionHelper(AsShared(), JsonValueObject, LookupTablePtr, InDisplayKey, CategoryFlags));
 		SubObjectCategory->ParseAllChildData(SearchabilityStatus);
 		ParsedChildData.Add(SubObjectCategory);
 	}
@@ -472,6 +502,7 @@ FCategorySectionHelper::FCategorySectionHelper(FImaginaryFiBDataWeakPtr InOuter,
 	: FImaginaryFiBData(InOuter, nullptr, InLookupTablePtr)
 	, CategoryName(InCategoryName)
 	, bIsTagAndValue(bInTagAndValueCategory)
+	, bAreKeysSearchable(false)
 {
 }
 
@@ -479,6 +510,7 @@ FCategorySectionHelper::FCategorySectionHelper(FImaginaryFiBDataWeakPtr InOuter,
 	: FImaginaryFiBData(InOuter, InUnparsedJsonObject, InLookupTablePtr)
 	, CategoryName(InCategoryName)
 	, bIsTagAndValue(bInTagAndValueCategory)
+	, bAreKeysSearchable(false)
 {
 }
 
@@ -487,8 +519,16 @@ FCategorySectionHelper::FCategorySectionHelper(FImaginaryFiBDataWeakPtr InOuter,
 	, SpecialHandlingCallback(InSpecialHandlingCallback)
 	, CategoryName(InCategoryName)
 	, bIsTagAndValue(bInTagAndValueCategory)
+	, bAreKeysSearchable(false)
 {
+}
 
+FCategorySectionHelper::FCategorySectionHelper(FImaginaryFiBDataWeakPtr InOuter, TSharedPtr< FJsonObject > InUnparsedJsonObject, TMap<int32, FText>* InLookupTablePtr, FText InCategoryName, ECategorySectionHelperFlags Flags)
+	: FImaginaryFiBData(InOuter, InUnparsedJsonObject, InLookupTablePtr)
+	, CategoryName(InCategoryName)
+	, bIsTagAndValue((Flags & ECategorySectionHelperFlags::IsTagAndValue) != ECategorySectionHelperFlags::None)
+	, bAreKeysSearchable((Flags & ECategorySectionHelperFlags::AreKeysSearchable) != ECategorySectionHelperFlags::None)
+{
 }
 
 bool FCategorySectionHelper::CanCallFilter(ESearchQueryFilter InSearchQueryFilter) const
@@ -580,7 +620,7 @@ FImaginaryBlueprint::FImaginaryBlueprint(const FString& InBlueprintName, const F
 	ParsedTagsAndValues.Add(FindInBlueprintsHelpers::FSimpleFTextKeyStorage(FFindInBlueprintSearchTags::FiB_Path), FSearchableValueInfo(FFindInBlueprintSearchTags::FiB_Path, FText::FromString(InBlueprintPath), ESearchableValueStatus::ExplicitySearchable));
 	ParsedTagsAndValues.Add(FindInBlueprintsHelpers::FSimpleFTextKeyStorage(FFindInBlueprintSearchTags::FiB_ParentClass), FSearchableValueInfo(FFindInBlueprintSearchTags::FiB_ParentClass, FText::FromString(InBlueprintParentClass), ESearchableValueStatus::ExplicitySearchable));
 
-	TSharedPtr< FCategorySectionHelper, ESPMode::ThreadSafe > InterfaceCategory = MakeShareable(new FCategorySectionHelper(nullptr, &LookupTable, FFindInBlueprintSearchTags::FiB_Interfaces, true));
+	TSharedPtr< FCategorySectionHelper, ESPMode::ThreadSafe > InterfaceCategory = MakeShareable(new FCategorySectionHelper(nullptr, &LookupTable, FFindInBlueprintSearchTags::FiB_Interfaces, /*bIsTagAndValueCategory=*/true));
 	for( int32 InterfaceIdx = 0; InterfaceIdx < InInterfaces.Num(); ++InterfaceIdx)
 	{
 		const FString& Interface = InInterfaces[InterfaceIdx];
@@ -682,7 +722,7 @@ bool FImaginaryBlueprint::TrySpecialHandleJsonValue(FText InKey, TSharedPtr< FJs
 		TArray<TSharedPtr< FJsonValue > > ComponentList = InJsonValue->AsArray();
 		TSharedPtr< FJsonObject > ComponentsWrapperObject(new FJsonObject);
 		ComponentsWrapperObject->Values.Add(FFindInBlueprintSearchTags::FiB_Components.ToString(), InJsonValue);
-		ParsedChildData.Add(MakeShareable(new FCategorySectionHelper(AsShared(), ComponentsWrapperObject, &LookupTable, FFindInBlueprintSearchTags::FiB_Components, false, FCategorySectionHelper::FCategorySectionHelperCallback::CreateRaw(this, &FImaginaryBlueprint::ParseComponents))));
+		ParsedChildData.Add(MakeShareable(new FCategorySectionHelper(AsShared(), ComponentsWrapperObject, &LookupTable, FFindInBlueprintSearchTags::FiB_Components, /*bIsTagAndValueCategory=*/false, FCategorySectionHelper::FCategorySectionHelperCallback::CreateRaw(this, &FImaginaryBlueprint::ParseComponents))));
 		bResult = true;
 	}
 
@@ -904,6 +944,7 @@ void FImaginaryGraphNode::ParseAllChildData_Internal(ESearchableValueStatus InSe
 
 FImaginaryProperty::FImaginaryProperty(FImaginaryFiBDataWeakPtr InOuter, TSharedPtr< FJsonObject > InUnparsedJsonObject, TMap<int32, FText>* InLookupTablePtr)
 	: FImaginaryFiBData(InOuter, InUnparsedJsonObject, InLookupTablePtr)
+	, bIsStructProperty(false)
 {
 }
 
@@ -914,6 +955,48 @@ bool FImaginaryProperty::IsCompatibleWithFilter(ESearchQueryFilter InSearchQuery
 		InSearchQueryFilter == ESearchQueryFilter::VariablesFilter;
 }
 
+void FImaginaryProperty::ParseAllChildData_Internal(ESearchableValueStatus InSearchabilityOverride)
+{
+	// Detect and cache whether PinCategory == PC_Struct. If so, it can represent:
+	// a) a struct, or
+	// b) a TMap whose key type is a struct
+	// The purpose for caching this is to detect TMap with struct keys, which must be searchable.
+	// We don't explicitly save whether some search metadata represents a TMap default value.
+	for (auto MapValues : UnparsedJsonObject->Values)
+	{
+		FText KeyText = FindInBlueprintsHelpers::AsFText(FCString::Atoi(*MapValues.Key), *LookupTablePtr);
+		if (KeyText.EqualTo(FFindInBlueprintSearchTags::FiB_PinCategory))
+		{
+			const FText PinCategoryText = FindInBlueprintsHelpers::AsFText(FCString::Atoi(*MapValues.Value->AsString()), *LookupTablePtr);
+			static FText PinCategoryStructValue = FText::FromName(UEdGraphSchema_K2::PC_Struct);
+			bIsStructProperty = PinCategoryText.EqualTo(PinCategoryStructValue);
+		}
+	}
+
+	FImaginaryFiBData::ParseAllChildData_Internal(InSearchabilityOverride);
+}
+
+bool FImaginaryProperty::ShouldSearchChildKeys(TSharedPtr<FJsonObject> ChildNode) const
+{
+	// There are two scenarios where we want to search child keys:
+	// - This is a TMap with a struct key type
+	// - This is a TMap with a plain-old-data (POD) key type
+	// In both cases ChildNode represents a key-value pair where we want to search the keys.
+	if (bIsStructProperty)
+	{
+		// Entering here implies this FProperty is either a struct or a TMap with a struct key. Then:
+		// - For structs, child nodes represent member properties where the key is the member name like 'MyMemberProperty'.
+		// - For TMaps, child nodes represent key-value pairs where the keys store the struct's FStructProperty::ExportText output like '(MyMemberProperty=111)'.
+		// Detect the latter case by peaking at the first element's syntax. This is not great, but implementing
+		// TMap key searchability isn't worth changing the metadata format for.
+		return !ChildNode->Values.IsEmpty() && FindInBlueprintsHelpers::AsFText(FCString::Atoi(*ChildNode->Values.Array()[0].Key), *LookupTablePtr).ToString().StartsWith("(");
+	}
+
+	// Entering here implies this is a TMap with a POD key type, because actual 
+	// POD FProperties don't have FJsonObject child nodes.
+	return true;
+}
+
 FSearchResult FImaginaryProperty::CreateSearchResult_Internal(FSearchResult InTemplate) const
 {
 	if (InTemplate.IsValid())
diff --git a/Engine/Source/Editor/Kismet/Public/ImaginaryBlueprintData.h b/Engine/Source/Editor/Kismet/Public/ImaginaryBlueprintData.h
--- a/Engine/Source/Editor/Kismet/Public/ImaginaryBlueprintData.h
+++ b/Engine/Source/Editor/Kismet/Public/ImaginaryBlueprintData.h
@@ -140,6 +140,12 @@ public:
 	/** Returns TRUE if this item is considered a Tag and Value category, where it's contents should be considered no different than the parent owner */
 	virtual bool IsTagAndValueCategory() const { return false; }
 
+	/** Returns TRUE if the child item's keys should be searchable. */
+	virtual bool ShouldSearchChildKeys(TSharedPtr<FJsonObject> ChildNode) const { return false; }
+
+	/** Returns TRUE if this item's keys should be searchable. For example if they don't represent property names, but values. Introduced for TMap keys. */
+	virtual bool AreKeysSearchable() const { return false; }
+
 	/** Checks if the filter is compatible with the current object, returns TRUE by default */
 	virtual bool IsCompatibleWithFilter(ESearchQueryFilter InSearchQueryFilter) const;
 
@@ -295,6 +301,16 @@ private:
 	bool bIsExplicit;
 };
 
+/** Flags to control the UX while caching */
+enum class ECategorySectionHelperFlags : uint8
+{
+	None = 0,
+	IsTagAndValue		= 1 << 0,	// See bIsTagAndValue
+	AreKeysSearchable	= 1 << 1,	// See bAreKeysSearchable
+};
+
+ENUM_CLASS_FLAGS(ECategorySectionHelperFlags);
+
 class FCategorySectionHelper : public FImaginaryFiBData
 {
 public:
@@ -304,9 +320,11 @@ public:
 	FCategorySectionHelper(FImaginaryFiBDataWeakPtr InOuter, TMap<int32, FText>* InLookupTablePtr, FText InCategoryName, bool bInTagAndValueCategory);
 	FCategorySectionHelper(FImaginaryFiBDataWeakPtr InOuter, TSharedPtr< FJsonObject > InUnparsedJsonObject, TMap<int32, FText>* InLookupTablePtr, FText InCategoryName, bool bInTagAndValueCategory);
 	FCategorySectionHelper(FImaginaryFiBDataWeakPtr InOuter, TSharedPtr< FJsonObject > InUnparsedJsonObject, TMap<int32, FText>* InLookupTablePtr, FText InCategoryName, bool bInTagAndValueCategory, FCategorySectionHelperCallback InSpecialHandlingCallback);
+	FCategorySectionHelper(FImaginaryFiBDataWeakPtr InOuter, TSharedPtr< FJsonObject > InUnparsedJsonObject, TMap<int32, FText>* InLookupTablePtr, FText InCategoryName, ECategorySectionHelperFlags Flags);
 
 	/** FImaginaryFiBData Interface */
 	virtual bool IsCategory() const override { return true; }
+	virtual bool AreKeysSearchable() const { return bAreKeysSearchable; }
 	virtual bool IsTagAndValueCategory() const override { return bIsTagAndValue; }
 	virtual bool CanCallFilter(ESearchQueryFilter InSearchQueryFilter) const override;
 	/** End FImaginaryFiBData Interface */
@@ -336,7 +354,10 @@ protected:
 	FText CategoryName;
 
 	/** TRUE if this category should be considered no different than a normal Tag and Value in it's parent */
-	bool bIsTagAndValue;
+	uint8 bIsTagAndValue:1;
+
+	/** TRUE if this category's keys should be searchable */
+	uint8 bAreKeysSearchable:1;
 };
 
 /** An "imaginary" representation of a UBlueprint, featuring raw strings or other imaginary objects in the place of more structured substances */
@@ -432,6 +453,8 @@ public:
 
 	/** FImaginaryFiBData Interface */
 	virtual bool IsCompatibleWithFilter(ESearchQueryFilter InSearchQueryFilter) const override;
+	virtual void ParseAllChildData_Internal(ESearchableValueStatus InSearchabilityOverride) override;
+	virtual bool ShouldSearchChildKeys(TSharedPtr<FJsonObject> ChildNode) const override;
 	/** End FImaginaryFiBData Interface */
 	
 protected:
@@ -440,6 +463,10 @@ protected:
 	virtual FSearchResult CreateSearchResult_Internal(FSearchResult InTemplate) const override;
 	virtual void DumpParsedObject_Internal(FArchive& Ar) const override;
 	/** End FImaginaryFiBData Interface */
+
+private:
+	/** TRUE if this property's metadata's PinCategory was PC_Struct. If so, it can represent a struct, or a TMap whose key type is a struct. */
+	uint8 bIsStructProperty:1;
 };
 
 /** An "imaginary" representation of a FProperty of an instanced component, featuring raw strings or other imaginary objects in the place of more structured substances */
