Editor: Add hooks to customize default selection behavior in subobject editor.

diff --git a/Engine/Source/Editor/Kismet/Public/ISCSEditorUICustomization.h b/Engine/Source/Editor/Kismet/Public/ISCSEditorUICustomization.h
--- a/Engine/Source/Editor/Kismet/Public/ISCSEditorUICustomization.h
+++ b/Engine/Source/Editor/Kismet/Public/ISCSEditorUICustomization.h
@@ -12,6 +12,7 @@
 class UActorComponent;
 struct FSubobjectData;
 struct FSubobjectDataHandle;
+class SSubobjectEditor;
 namespace UE::Editor
 {
 	struct FSubobjectEditorContext;
@@ -94,6 +95,21 @@ public:
 	 * @return Whether any sorting was attempted */
 	virtual bool SortSubobjectData(TArray<FSubobjectDataHandle>& SubobjectData) { return false; }
 
+	/**
+	 * Called before events that might change the targets of the subobject editor, particularly
+	 *  useful for gathering data to use in a subsequent SelectDefaultSubobject call.
+	 */
+	UE_EXPERIMENTAL(5.8, "This function is experimental for now.")
+	virtual void PrePotentialTargetChange(SSubobjectEditor& SubobjectEditor) {};
+
+	/**
+	 * Gives a customization the opportunity to control what the "default" subobject selection is.
+	 * @return true if custom selection was performed, false if subobject editor should perform its
+	 *  own default selection (typically selecting root).
+	 */
+	UE_EXPERIMENTAL(5.8, "This function is experimental for now.")
+	virtual bool SelectDefaultSubobject(SSubobjectEditor& SubobjectEditor) { return false; }
+
 	/** Decides if this UI Customization provides an overlay for the tree view */
 	virtual bool SupportsTreeViewOverlay(TArrayView<UObject*> Context) const { return false; }
 
diff --git a/Engine/Source/Editor/LevelEditor/Private/ActorDetailsSCSEditorUICustomization.cpp b/Engine/Source/Editor/LevelEditor/Private/ActorDetailsSCSEditorUICustomization.cpp
--- a/Engine/Source/Editor/LevelEditor/Private/ActorDetailsSCSEditorUICustomization.cpp
+++ b/Engine/Source/Editor/LevelEditor/Private/ActorDetailsSCSEditorUICustomization.cpp
@@ -137,6 +137,26 @@ bool FActorDetailsSCSEditorUICustomization::FilterSubobjectData(TArray<FSubobjec
 	return bChanged;
 }
 
+void FActorDetailsSCSEditorUICustomization::PrePotentialTargetChange(SSubobjectEditor& SubobjectEditor)
+{
+	for (const TSharedPtr<ISCSEditorUICustomization>& Customization : Customizations)
+	{
+		Customization->PrePotentialTargetChange(SubobjectEditor);
+	}
+}
+
+bool FActorDetailsSCSEditorUICustomization::SelectDefaultSubobject(SSubobjectEditor& SubobjectEditor)
+{
+	for (const TSharedPtr<ISCSEditorUICustomization>& Customization : Customizations)
+	{
+		if (Customization->SelectDefaultSubobject(SubobjectEditor))
+		{
+			return true;
+		}
+	}
+	return false;
+}
+
 void FActorDetailsSCSEditorUICustomization::AddCustomization(TSharedPtr<ISCSEditorUICustomization> Customization)
 {
 	Customizations.AddUnique(Customization);
diff --git a/Engine/Source/Editor/LevelEditor/Private/ActorDetailsSCSEditorUICustomization.h b/Engine/Source/Editor/LevelEditor/Private/ActorDetailsSCSEditorUICustomization.h
--- a/Engine/Source/Editor/LevelEditor/Private/ActorDetailsSCSEditorUICustomization.h
+++ b/Engine/Source/Editor/LevelEditor/Private/ActorDetailsSCSEditorUICustomization.h
@@ -31,6 +31,9 @@ public:
 	
 	virtual bool FilterSubobjectData(TArray<FSubobjectDataHandle>& SubobjectData) const override;
 
+	virtual void PrePotentialTargetChange(SSubobjectEditor& SubobjectEditor) override;
+	virtual bool SelectDefaultSubobject(SSubobjectEditor& SubobjectEditor) override;
+
 	virtual bool SupportsTreeViewOverlay(TArrayView<UObject*> Context) const override;
 	virtual TSharedRef<SWidget> ConstructTreeViewOverlay(TWeakPtr<SSubobjectEditor> Editor) const override;
 
diff --git a/Engine/Source/Editor/LevelEditor/Private/SActorDetails.cpp b/Engine/Source/Editor/LevelEditor/Private/SActorDetails.cpp
--- a/Engine/Source/Editor/LevelEditor/Private/SActorDetails.cpp
+++ b/Engine/Source/Editor/LevelEditor/Private/SActorDetails.cpp
@@ -499,7 +499,10 @@ void SActorDetails::RefreshTopLevelElements(TArrayView<const TTypedElement<IType
 	{
 		// Enable the selection guard to prevent OnTreeSelectionChanged() from altering the editor's component selection
 		TGuardValue<bool> SelectionGuard(bSelectionGuard, true);
+		SubobjectEditor->PrePotentialTargetChange();
+		// Make our subobject tree correspond to our selected top-level elements (e.g. rebuild from current actor)
 		SubobjectEditor->UpdateTree();
+		// Adjust the selection of subobjects based on editor selection (e.g. based on actor component selection)
 		UpdateComponentTreeFromEditorSelection();
 	}
 
@@ -842,7 +845,7 @@ void SActorDetails::UpdateComponentTreeFromEditorSelection()
 	TArray<FSubobjectEditorTreeNodePtrType> SelectedNodes = SubobjectEditor->GetSelectedNodes();
 	if (SelectedNodes.Num() == 0)
 	{
-		SubobjectEditor->SelectRoot();
+		SubobjectEditor->SelectDefaultSubobject();
 		SelectedNodes = SubobjectEditor->GetSelectedNodes();
 	}
 
diff --git a/Engine/Source/Editor/SubobjectEditor/Private/SSubobjectEditor.cpp b/Engine/Source/Editor/SubobjectEditor/Private/SSubobjectEditor.cpp
--- a/Engine/Source/Editor/SubobjectEditor/Private/SSubobjectEditor.cpp
+++ b/Engine/Source/Editor/SubobjectEditor/Private/SSubobjectEditor.cpp
@@ -2654,6 +2654,14 @@ void SSubobjectEditor::RefreshSelectionDetails()
 	UpdateSelectionFromNodes(TreeWidget->GetSelectedItems());
 }
 
+void SSubobjectEditor::PrePotentialTargetChange()
+{
+	if (UICustomization)
+	{
+		UICustomization->PrePotentialTargetChange(*this);
+	}
+}
+
 void SSubobjectEditor::ClearSelection()
 {
 	if (bUpdatingSelection == false)
@@ -2663,6 +2671,18 @@ void SSubobjectEditor::ClearSelection()
 	}
 }
 
+void SSubobjectEditor::SelectDefaultSubobject()
+{
+	if (UICustomization)
+	{
+		if (UICustomization->SelectDefaultSubobject(*this))
+		{
+			return;
+		}
+	}
+	SelectRoot();
+}
+
 void SSubobjectEditor::RestoreSelectionState(TArray<FSubobjectEditorTreeNodePtrType>& SelectedTreeNodes, bool bFallBackToVariableName)
 {
 	if (SelectedTreeNodes.Num() > 0)
diff --git a/Engine/Source/Editor/SubobjectEditor/Public/SSubobjectEditor.h b/Engine/Source/Editor/SubobjectEditor/Public/SSubobjectEditor.h
--- a/Engine/Source/Editor/SubobjectEditor/Public/SSubobjectEditor.h
+++ b/Engine/Source/Editor/SubobjectEditor/Public/SSubobjectEditor.h
@@ -492,6 +492,17 @@ public:
 	/** Clears the current selection */
 	UE_API void ClearSelection();
 
+	/** 
+	 * Called before a potential change of subobject editor target, so that customizations can be
+	 *  given a chance to gather context they might need for a subsequent SelectDefaultSubobject call.
+	 */
+	UE_EXPERIMENTAL(5.8, "This function is experimental for now.")
+	UE_API void PrePotentialTargetChange();
+
+	/** Select the default subobject, which is typically the root, but may be adjusted by a customization*/
+	UE_EXPERIMENTAL(5.8, "This function is experimental for now.")
+	UE_API void SelectDefaultSubobject();
+
 	/** Select the root of the tree */
 	UE_API void SelectRoot();
 
