EditorUsability : UnrealEd
- Improved menu filtering with additional context options for show flags & view modes & camera  settings
- Fixed cannot set camera clipping settings in the BP Editor's viewport
- Fixed cannot set other show flags in the BP Editor's viewport

BP Editor's viewport now contains :
- Show flags : AntiAliasing, Camera, Collisions, Decals, Grid, Hair, Particle Sprites, Skeletal Meshes, Static Meshes, Translucency, Widget Components
- View mode  : Lit, Unlit, Wireframe, Lit Wireframe, Reflections, Player Collision, Visibility Collision, Path Tracing, Geometry Inspection, Visualizers
- Camera options : Field of view, Near view plane, Far view plane

diff --git a/Engine/Source/Editor/Kismet/Private/SSCSEditorViewport.cpp b/Engine/Source/Editor/Kismet/Private/SSCSEditorViewport.cpp
--- a/Engine/Source/Editor/Kismet/Private/SSCSEditorViewport.cpp
+++ b/Engine/Source/Editor/Kismet/Private/SSCSEditorViewport.cpp
@@ -66,7 +66,11 @@ bool IsViewModeSupported(EViewModeIndex InViewModeIndex)
 	case VMI_Unlit:
 	case VMI_Lit:
 	case VMI_BrushWireframe:
+	case VMI_Lit_Wireframe:
+	case VMI_ReflectionOverride:
+	case VMI_CollisionPawn:
 	case VMI_CollisionVisibility:
+	case VMI_PathTracing:
 	case VMI_Clay:
 	case VMI_Zebra:
 	case VMI_FrontBackFace:
@@ -186,7 +190,10 @@ TSharedPtr<SWidget> SSCSEditorViewport::BuildViewportToolbar()
 
 			// Add the "Camera" submenu.
 			{
-				RightSection.AddEntry(UE::UnrealEd::CreateCameraSubmenu(UE::UnrealEd::FViewportCameraMenuOptions().ShowCameraMovement()));
+				// Add camera options for viewport and show all options
+				UE::UnrealEd::FViewportCameraMenuOptions CameraMenuOptions;
+				CameraMenuOptions.ShowAll();
+				RightSection.AddEntry(UE::UnrealEd::CreateCameraSubmenu(CameraMenuOptions));
 				
 				const FName SubmenuName = UToolMenus::JoinMenuPaths(ViewportToolbarName, "Camera");
                 UToolMenu* const Submenu = UToolMenus::Get()->ExtendMenu(SubmenuName);
@@ -194,9 +201,6 @@ TSharedPtr<SWidget> SSCSEditorViewport::BuildViewportToolbar()
                 MovementSection.AddMenuEntry(FBlueprintEditorCommands::Get().ResetCamera);
 			}
 
-			// TODO: Filter this menu with IsViewModeSupportedDelegate (see further down in this file) and remove the
-			// "Exposure" section.
-
 			// Add the "View Modes" sub menu.
 			{
 				// Stay backward-compatible with the old viewport toolbar.
@@ -215,16 +219,9 @@ TSharedPtr<SWidget> SSCSEditorViewport::BuildViewportToolbar()
 				RightSection.AddEntry(UE::UnrealEd::CreateViewModesSubmenu());
 			}
 
-			// Add the "Show" submenu.
+			// Add the "Show" submenu with all flags, exclude some level viewport specific flags below
 			{
-				RightSection.AddEntry(UE::UnrealEd::CreateShowSubmenu(FNewToolMenuDelegate::CreateLambda(
-					[](UToolMenu* Submenu) -> void
-					{
-						FToolMenuSection& UnnamedSection = Submenu->FindOrAddSection(NAME_None);
-						UnnamedSection.AddMenuEntry(FBlueprintEditorCommands::Get().ShowFloor);
-						UnnamedSection.AddMenuEntry(FBlueprintEditorCommands::Get().ShowGrid);
-					})
-				));
+				RightSection.AddEntry(UE::UnrealEd::CreateDefaultShowSubmenu());
 			}
 
 			// Add the "Performance & Scalability" submenu.
@@ -264,9 +261,38 @@ TSharedPtr<SWidget> SSCSEditorViewport::BuildViewportToolbar()
 			UUnrealEdViewportToolbarContext* const ContextObject = NewObject<UUnrealEdViewportToolbarContext>();
 			ContextObject->Viewport = SharedThis(this);
 
+			// Filter unrelated/irrelevant flags for BP editor
+			const TArray<FEngineShowFlags::EShowFlag> ExcludedShowFlags
+			{
+				FEngineShowFlags::EShowFlag::SF_Atmosphere,
+				FEngineShowFlags::EShowFlag::SF_BSP,
+				FEngineShowFlags::EShowFlag::SF_Cloud,
+				FEngineShowFlags::EShowFlag::SF_Fog,
+				FEngineShowFlags::EShowFlag::SF_Landscape,
+				FEngineShowFlags::EShowFlag::SF_MediaPlanes,
+				FEngineShowFlags::EShowFlag::SF_Navigation,
+			};
+			ContextObject->ExcludedShowMenuFlags.Append(ExcludedShowFlags);
+
+			const TArray<EShowFlagGroup> ExcludedShowGroupFlags
+			{
+				EShowFlagGroup::SFG_PostProcess,
+				EShowFlagGroup::SFG_LightTypes,
+				EShowFlagGroup::SFG_LightingComponents,
+				EShowFlagGroup::SFG_LightingFeatures,
+				EShowFlagGroup::SFG_Lumen,
+				EShowFlagGroup::SFG_Nanite,
+				EShowFlagGroup::SFG_Developer,
+				EShowFlagGroup::SFG_Visualize,
+				EShowFlagGroup::SFG_Advanced,
+				EShowFlagGroup::SFG_Custom,
+			};
+			ContextObject->ExcludedShowMenuGroupFlags.Append(ExcludedShowGroupFlags);
+
 			// Setup the callback to filter available view modes
-			ContextObject->IsViewModeSupported =
-				UE::UnrealEd::IsViewModeSupportedDelegate::CreateStatic(&UE::SCSEditor::Private::IsViewModeSupported);
+			ContextObject->IsViewModeSupported = UE::UnrealEd::IsViewModeSupportedDelegate::CreateStatic(
+			 	&UE::SCSEditor::Private::IsViewModeSupported
+			);
 
 			// Setup the callback to hide/show specific sections
 			ContextObject->DoesViewModeMenuShowSection = UE::UnrealEd::DoesViewModeMenuShowSectionDelegate::CreateStatic(
diff --git a/Engine/Source/Editor/UnrealEd/Private/ViewportToolbar/UnrealEdViewportToolbar.cpp b/Engine/Source/Editor/UnrealEd/Private/ViewportToolbar/UnrealEdViewportToolbar.cpp
--- a/Engine/Source/Editor/UnrealEd/Private/ViewportToolbar/UnrealEdViewportToolbar.cpp
+++ b/Engine/Source/Editor/UnrealEd/Private/ViewportToolbar/UnrealEdViewportToolbar.cpp
@@ -4537,18 +4537,23 @@ FToolMenuEntry CreateDefaultShowSubmenu()
 
 void AddDefaultShowFlags(UToolMenu* InMenu)
 {
+	FShowFlagFilter ShowFlagFilter = FShowFlagFilter(FShowFlagFilter::IncludeAllFlagsByDefault);
+	if (UUnrealEdViewportToolbarContext* EditorViewportContext = InMenu->FindContext<UUnrealEdViewportToolbarContext>())
 	{
-		FToolMenuSection& CommonShowFlagsSection =
-			InMenu->FindOrAddSection("CommonShowFlags", LOCTEXT("CommonShowFlagsLabel", "Common Show Flags"));
+		for (const FEngineShowFlags::EShowFlag& Flag : EditorViewportContext->ExcludedShowMenuFlags)
+		{
+			ShowFlagFilter.ExcludeFlag(Flag);
+		}
 
-		FShowFlagFilter ShowFlagFilter = FShowFlagFilter(FShowFlagFilter::IncludeAllFlagsByDefault);
-		if (UUnrealEdViewportToolbarContext* EditorViewportContext = InMenu->FindContext<UUnrealEdViewportToolbarContext>())
+		for (const EShowFlagGroup& Flag : EditorViewportContext->ExcludedShowMenuGroupFlags)
 		{
-			for (const FEngineShowFlags::EShowFlag& Flag : EditorViewportContext->ExcludedShowMenuFlags)
-			{
-				ShowFlagFilter.ExcludeFlag(Flag);
-			}
+			ShowFlagFilter.ExcludeGroup(Flag);
 		}
+	}
+
+	{
+		FToolMenuSection& CommonShowFlagsSection =
+			InMenu->FindOrAddSection("CommonShowFlags", LOCTEXT("CommonShowFlagsLabel", "Common Show Flags"));
 
 		FShowFlagMenuCommands::Get().PopulateCommonShowFlagsSection(CommonShowFlagsSection, ShowFlagFilter);
 	}
@@ -4557,7 +4562,7 @@ void AddDefaultShowFlags(UToolMenu* InMenu)
 		FToolMenuSection& AllShowFlagsSection =
 			InMenu->FindOrAddSection("AllShowFlags", LOCTEXT("AllShowFlagsLabel", "All Show Flags"));
 
-		FShowFlagMenuCommands::Get().PopulateAllShowFlagsSection(AllShowFlagsSection);
+		FShowFlagMenuCommands::Get().PopulateAllShowFlagsSection(AllShowFlagsSection, ShowFlagFilter);
 	}
 }
 
diff --git a/Engine/Source/Editor/UnrealEd/Public/ViewportToolbar/UnrealEdViewportToolbar.h b/Engine/Source/Editor/UnrealEd/Public/ViewportToolbar/UnrealEdViewportToolbar.h
--- a/Engine/Source/Editor/UnrealEd/Public/ViewportToolbar/UnrealEdViewportToolbar.h
+++ b/Engine/Source/Editor/UnrealEd/Public/ViewportToolbar/UnrealEdViewportToolbar.h
@@ -16,7 +16,6 @@ class FEditorViewportClient;
 class FTabManager;
 class FText;
 class IPreviewProfileController;
-class IPreviewProfileController;
 class IPreviewLODController;
 class SEditorViewport;
 class UToolMenu;
diff --git a/Engine/Source/Editor/UnrealEd/Public/ViewportToolbar/UnrealEdViewportToolbarContext.h b/Engine/Source/Editor/UnrealEd/Public/ViewportToolbar/UnrealEdViewportToolbarContext.h
--- a/Engine/Source/Editor/UnrealEd/Public/ViewportToolbar/UnrealEdViewportToolbarContext.h
+++ b/Engine/Source/Editor/UnrealEd/Public/ViewportToolbar/UnrealEdViewportToolbarContext.h
@@ -18,12 +18,10 @@ enum ERotationGridMode : int;
 
 namespace UE::UnrealEd
 {
+	DECLARE_DELEGATE_RetVal_OneParam(bool, IsViewModeSupportedDelegate, EViewModeIndex);
 
-DECLARE_DELEGATE_RetVal_OneParam(bool, IsViewModeSupportedDelegate, EViewModeIndex);
-
-enum EHidableViewModeMenuSections : uint8;
-DECLARE_DELEGATE_RetVal_OneParam(bool, DoesViewModeMenuShowSectionDelegate, EHidableViewModeMenuSections);
-
+	enum EHidableViewModeMenuSections : uint8;
+	DECLARE_DELEGATE_RetVal_OneParam(bool, DoesViewModeMenuShowSectionDelegate, EHidableViewModeMenuSections);
 }
 
 class SEditorViewport;
@@ -36,7 +34,10 @@ class UUnrealEdViewportToolbarContext : public UObject
 	GENERATED_BODY()
 
 public:
+	/** Actual viewport for context */
 	TWeakPtr<SEditorViewport> Viewport;
+
+	/** View mode filters for context menu */
 	UE::UnrealEd::IsViewModeSupportedDelegate IsViewModeSupported;
 	UE::UnrealEd::DoesViewModeMenuShowSectionDelegate DoesViewModeMenuShowSection;
 
@@ -50,6 +51,11 @@ public:
 	 */
 	TArray<FEngineShowFlags::EShowFlag> ExcludedShowMenuFlags;
 
+	/**
+	 * Add flags to this array to remove them from the Show Menu (Group)
+	 */
+	TArray<EShowFlagGroup> ExcludedShowMenuGroupFlags;
+
 	/**
 	 * Can be used to retrieve data e.g. TabManager
 	 */
