Find-in-Blueprints: Map key-value pairs are no longer culled from search results just because the value is uninteresting (i.e. 0 or empty string). Because key-value pairs can be interesting due to the key.

diff --git a/Engine/Source/Editor/Kismet/Private/FindInBlueprintManager.cpp b/Engine/Source/Editor/Kismet/Private/FindInBlueprintManager.cpp
--- a/Engine/Source/Editor/Kismet/Private/FindInBlueprintManager.cpp
+++ b/Engine/Source/Editor/Kismet/Private/FindInBlueprintManager.cpp
@@ -666,7 +666,7 @@ namespace BlueprintSearchMetaDataHelpers
 	 * @param InJsonValue		The Json value object to examine for searchability
 	 * @return					TRUE if the value should be searchable
 	 */
-	bool CheckIfJsonValueIsSearchable( TSharedPtr< FJsonValue > InJsonValue )
+	bool CheckIfJsonValueIsSearchable( TSharedPtr< FJsonValue > InJsonValue, const FBPVariableDescription& InVariableDescription)
 	{
 		/** Check for interesting values
 		 *  booleans are not interesting, there are a lot of them
@@ -717,11 +717,12 @@ namespace BlueprintSearchMetaDataHelpers
 			const TSharedPtr<FJsonObject>& JsonObject = InJsonValue->AsObject();
 			for(TMap<FString, TSharedPtr<FJsonValue>>::TIterator Iter = JsonObject->Values.CreateIterator(); Iter; ++Iter)
 			{
-				// Empty keys don't convert to JSON, so we also remove the entry in that case. Note: This means the entry is not going to be searchable.
-				// @todo - Potentially use a placeholder string that uniquely identifies this as an empty key?
-				const bool bHasEmptyKey = Iter->Key.IsEmpty();
+				// Map variables should always be searchable. (Iter->Key, Iter->Value) represents a user entered key-value pair.
+				// Even if the value is uninteresting (empty string, 0, etc) then the key could be interesting.
+				const bool bIsMapVariable = InVariableDescription.VarType.ContainerType == EPinContainerType::Map;
 
-				if(!CheckIfJsonValueIsSearchable(Iter->Value) || bHasEmptyKey)
+				// If not a map variable, check per entry if the value is also 'interesting' (see comment at the top of this function).
+				if (!bIsMapVariable && !CheckIfJsonValueIsSearchable(Iter->Value, InVariableDescription))
 				{
 					// Note: It's safe to keep incrementing after this; the underlying logic maps to TSparseArray/TConstSetBitIterator::RemoveCurrent().
 					Iter.RemoveCurrent();
@@ -797,13 +798,13 @@ namespace BlueprintSearchMetaDataHelpers
 
 		// Find the FProperty and convert it into a Json value.
 		FProperty* VariableProperty = FindFProperty<FProperty>(InBlueprint->GeneratedClass, InVariableDescription.VarName);
-		if(VariableProperty)
+		if (VariableProperty)
 		{
 			const uint8* PropData = VariableProperty->ContainerPtrToValuePtr<uint8>(InBlueprint->GeneratedClass->GetDefaultObject());
 			auto JsonValue = FJsonObjectConverter::UPropertyToJsonValue(VariableProperty, PropData, 0, 0);
 
 			// Only use the value if it is searchable
-			if(BlueprintSearchMetaDataHelpers::CheckIfJsonValueIsSearchable(JsonValue))
+			if (BlueprintSearchMetaDataHelpers::CheckIfJsonValueIsSearchable(JsonValue, InVariableDescription))
 			{
 				TSharedRef< FJsonValue > JsonValueAsSharedRef = JsonValue.ToSharedRef();
 				FJsonSerializer::Serialize<TCHAR, PrintPolicy>(JsonValue, FFindInBlueprintSearchTags::FiB_DefaultValue.ToString(), InWriter, false );
