Optimize save of Find in Blueprints data by avoiding temporary string allocations used to compute FText identity

#rnx
#Synced-CL 48496709

diff --git a/Engine/Source/Editor/Kismet/Private/FindInBlueprintManager.cpp b/Engine/Source/Editor/Kismet/Private/FindInBlueprintManager.cpp
--- a/Engine/Source/Editor/Kismet/Private/FindInBlueprintManager.cpp
+++ b/Engine/Source/Editor/Kismet/Private/FindInBlueprintManager.cpp
@@ -569,7 +569,7 @@ namespace BlueprintSearchMetaDataHelpers
 		}
 
 		virtual bool Close() override
-					{
+		{
 			// This will copy the JSON output to the string given as input (must do this first)
 			bool bResult = TFindInBlueprintJsonStringWriter<TCondensedJsonPrintPolicy<TCHAR>>::Close();
 
@@ -580,7 +580,7 @@ namespace BlueprintSearchMetaDataHelpers
 				+ MoveTemp(*JsonOutput);
 
 			return bResult;
-			}
+		}
 
 	protected:
 		virtual void WriteStringValue(const FString& String) override
@@ -598,7 +598,7 @@ namespace BlueprintSearchMetaDataHelpers
 		virtual void WriteTextValue(const FText& Text) override
 		{
 			// Check to see if the value has already been added.
-			if (int32* TableLookupValuePtr = ReverseLookupTable.Find(FLookupTableItem(Text)))
+			if (int32* TableLookupValuePtr = ReverseLookupTable.Find(Text))
 			{
 				TFindInBlueprintJsonStringWriter<TCondensedJsonPrintPolicy<TCHAR>>::WriteStringValue(FString::FromInt(*TableLookupValuePtr));
 			}
@@ -608,48 +608,13 @@ namespace BlueprintSearchMetaDataHelpers
 				int32 TableLookupValue = LookupTable.Num();
 				{
 					LookupTable.Add(TableLookupValue, Text);
-					ReverseLookupTable.Add(FLookupTableItem(Text), TableLookupValue);
+					ReverseLookupTable.Add(Text, TableLookupValue);
 				}
 				TFindInBlueprintJsonStringWriter<TCondensedJsonPrintPolicy<TCHAR>>::WriteStringValue(FString::FromInt(TableLookupValue));
 			}
 		}
 
 	private:
-		struct FLookupTableItem
-		{
-			FText Text;
-
-			FLookupTableItem(FText InText)
-				: Text(InText)
-			{
-
-		}
-
-			bool operator==(const FLookupTableItem& InObject) const
-			{
-				if (!Text.CompareTo(InObject.Text))
-				{
-					if (FTextInspector::GetNamespace(Text).Get(TEXT("DefaultNamespace")) == FTextInspector::GetNamespace(InObject.Text).Get(TEXT("DefaultNamespace")))
-	{
-						if (FTextInspector::GetKey(Text).Get(TEXT("DefaultKey")) == FTextInspector::GetKey(InObject.Text).Get(TEXT("DefaultKey")))
-		{
-							return true;
-						}
-					}
-		}
-
-				return false;
-		}
-
-			friend uint32 GetTypeHash(const FLookupTableItem& InObject)
-		{
-				FString Namespace = FTextInspector::GetNamespace(InObject.Text).Get(TEXT("DefaultNamespace"));
-				FString Key = FTextInspector::GetKey(InObject.Text).Get(TEXT("DefaultKey"));
-				uint32 Hash = HashCombine(GetTypeHash(InObject.Text.ToString()), HashCombine(GetTypeHash(Namespace), GetTypeHash(Key)));
-				return Hash;
-		}
-		};
-
 		// Output stream
 		FString* JsonOutput;
 
@@ -657,7 +622,29 @@ namespace BlueprintSearchMetaDataHelpers
 		TMap< int32, FText > LookupTable;
 
 		// This is just locally needed for the write, to lookup the integer value by using the string of the FText
-		TMap< FLookupTableItem, int32 > ReverseLookupTable;
+		/** Case sensitive hashing function for TMap */
+		struct FTextKeyMapFuncs : BaseKeyFuncs<int32, FText, /*bInAllowDuplicateKeys*/false>
+		{
+			static const FString& GetSourceString(const FText& InText)
+			{
+				const FString* SourceString = FTextInspector::GetSourceString(InText);
+				check(SourceString);
+				return *SourceString;
+			}
+			static const FText& GetSetKey(const TPair<FText, int32>& Element)
+			{
+				return Element.Key;
+			}
+			static bool Matches(const FText& A, const FText& B)
+			{
+				return GetSourceString(A).Equals(GetSourceString(B), ESearchCase::CaseSensitive);
+			}
+			static uint32 GetKeyHash(const FText& Key)
+			{
+				return FLocKey::ProduceHash(GetSourceString(Key));
+			}
+		};
+		TMap< FText, int32, FDefaultSetAllocator, FTextKeyMapFuncs > ReverseLookupTable;
 	};
 
 	/**
