[Blueprint/ALT]
- Fix many race conditions by only serializing objects owned by the loader during reinstanciation on ALT

#Synced-CL Unknown

diff --git a/Engine/Source/Editor/Kismet/Private/BlueprintCompilationManager.cpp b/Engine/Source/Editor/Kismet/Private/BlueprintCompilationManager.cpp
--- a/Engine/Source/Editor/Kismet/Private/BlueprintCompilationManager.cpp
+++ b/Engine/Source/Editor/Kismet/Private/BlueprintCompilationManager.cpp
@@ -3322,11 +3322,14 @@ void FBlueprintCompilationManagerImpl::ReinstanceBatch(TArray<FReinstancingJob>&
 			{
 				ArchetypeReferencers.Add(BP->SkeletonGeneratedClass);
 			}
-			for(const TWeakObjectPtr<UBlueprint>& Dependency : BP->CachedDependencies)
+
+			// Previously, this was CachedDependencies, but it doesn't make much sense to search for ourselves in the 
+			// classes that we depend on. This has been modified to now search into the ones that depends on the BP instead.
+			for(const TWeakObjectPtr<UBlueprint>& Dependent : BP->CachedDependents)
 			{
-				if (UBlueprint* DependencyBP = Dependency.Get())
+				if (UBlueprint* DependentBP = Dependent.Get())
 				{
-					ArchetypeReferencers.Add(DependencyBP);
+					ArchetypeReferencers.Add(DependentBP);
 				}
 			}
 		}
@@ -3334,11 +3337,18 @@ void FBlueprintCompilationManagerImpl::ReinstanceBatch(TArray<FReinstancingJob>&
 
 	// Data loading may be in flight, lets immediately patch existing redirectors - 
 	// we may want to search the entire graph some day, but that would be expensive:
+	const bool bIsInActualLoadingThread = IsInActualLoadingThread();
 	for (TObjectIterator<UObjectRedirector> Itr; Itr; ++Itr)
 	{
 		if (Itr->DestinationObject && 
 			Itr->DestinationObject->HasAnyFlags(RF_ClassDefaultObject|RF_ArchetypeObject))
 		{
+			// When we're in the loading thread, we limit the replacements to UObjectRedirector that the loader thread owns since
+			// it would not be possible for an already loaded redirector to point to something still in the loader.
+			if (bIsInActualLoadingThread && !Itr->HasAnyInternalFlags(EInternalObjectFlags::AsyncLoadingPhase1))
+			{
+				continue;
+			}
 			ArchetypeReferencers.Add(*Itr);
 		}
 	}
