Avoid gathering dependents and dependencies on load, this is the responsibility of the linker and no further gathering is required at load time

#rnx
#Synced-CL 49181000

diff --git a/Engine/Source/Editor/Kismet/Private/BlueprintCompilationManager.cpp b/Engine/Source/Editor/Kismet/Private/BlueprintCompilationManager.cpp
--- a/Engine/Source/Editor/Kismet/Private/BlueprintCompilationManager.cpp
+++ b/Engine/Source/Editor/Kismet/Private/BlueprintCompilationManager.cpp
@@ -3989,6 +3989,11 @@ void FBlueprintCompilationManagerImpl::RelinkSkeleton(UClass* SkeletonToRelink)
 void FBlueprintCompilationManagerImpl::GatherOutOfDateDependenciesRecursive(TObjectPtr<UBlueprint> Gather, TSet<TObjectPtr<UBlueprint>>& OutOfDateDeps)
 {	
 	FBlueprintEditorUtils::EnsureCachedDependenciesUpToDate(Gather);
+	
+	if(!Gather->bHasBeenRegenerated && Gather->GetLinker())
+	{
+		return;
+	}
 		
 	// make sure any dirty dependencies are also compiled:
 	for (TWeakObjectPtr<UBlueprint> BPWeak : Gather->CachedDependencies)
