EditorUsability : Tooltips
- Editor tooltips now follow unreal theme instead of being white (Curve, Kismet, Material, Nodes, Niagara, Menus, ...)
- PIE and Packaged project still use the legacy white background tooltips

#Synced-CL 49388435

diff --git a/Engine/Plugins/Animation/DeformerGraph/Source/OptimusEditor/Private/OptimusEditorStyle.cpp b/Engine/Plugins/Animation/DeformerGraph/Source/OptimusEditor/Private/OptimusEditorStyle.cpp
--- a/Engine/Plugins/Animation/DeformerGraph/Source/OptimusEditor/Private/OptimusEditorStyle.cpp
+++ b/Engine/Plugins/Animation/DeformerGraph/Source/OptimusEditor/Private/OptimusEditorStyle.cpp
@@ -116,11 +116,9 @@ FOptimusEditorStyle::FOptimusEditorStyle() :
 
 		Set( "Node.ToolTipLabel", FTextBlockStyle(NormalText)
 			.SetFont( DEFAULT_FONT( "Bold", 9 ))
-			.SetColorAndOpacity(FLinearColor::Black)
 			);
 		Set( "Node.ToolTipContent", FTextBlockStyle(NormalText)
 			.SetFont( DEFAULT_FONT( "Regular", 9))
-			.SetColorAndOpacity(FLinearColor::Black)
 			);
 		
 		
diff --git a/Engine/Plugins/Animation/DeformerGraph/Source/OptimusEditor/Private/Widgets/SOptimusEditorGraphNode.cpp b/Engine/Plugins/Animation/DeformerGraph/Source/OptimusEditor/Private/Widgets/SOptimusEditorGraphNode.cpp
--- a/Engine/Plugins/Animation/DeformerGraph/Source/OptimusEditor/Private/Widgets/SOptimusEditorGraphNode.cpp
+++ b/Engine/Plugins/Animation/DeformerGraph/Source/OptimusEditor/Private/Widgets/SOptimusEditorGraphNode.cpp
@@ -1016,7 +1016,6 @@ TSharedPtr<IToolTip> SOptimusEditorGraphNode::MakePinToolTip(
 	) const
 {
 	return SNew(SToolTip)
-	.BorderImage(FCoreStyle::Get().GetBrush("ToolTip.BrightBackground"))
 	.TextMargin(11.0f)
 	[
 		SNew(SOptimusEditorGraphPinToolTipWidget)
diff --git a/Engine/Plugins/Animation/GameplayInsights/Source/GameplayInsights/Private/SCurveTimelineView.cpp b/Engine/Plugins/Animation/GameplayInsights/Source/GameplayInsights/Private/SCurveTimelineView.cpp
--- a/Engine/Plugins/Animation/GameplayInsights/Source/GameplayInsights/Private/SCurveTimelineView.cpp
+++ b/Engine/Plugins/Animation/GameplayInsights/Source/GameplayInsights/Private/SCurveTimelineView.cpp
@@ -87,7 +87,6 @@ void SCurveTimelineView::UpdateCurveToolTip(const FGeometry& InMyGeometry, const
 						{
 							SetToolTip(
 								SAssignNew(CurveToolTip, SToolTip)
-								.BorderImage( FCoreStyle::Get().GetBrush( "ToolTip.BrightBackground" ) )
 								[
 									SNew(SVerticalBox)
 									+ SVerticalBox::Slot()
@@ -95,21 +94,18 @@ void SCurveTimelineView::UpdateCurveToolTip(const FGeometry& InMyGeometry, const
 										SNew(STextBlock)
 										.Text(this, &SCurveTimelineView::GetTrackName)
 										.Font(FCoreStyle::Get().GetFontStyle("ToolTip.LargerFont"))
-										.ColorAndOpacity(FLinearColor::Black)
 									]
 									+ SVerticalBox::Slot()
 									[
 										SNew(STextBlock)
 										.Text(this, &SCurveTimelineView::GetCurveToolTipInputText)
 										.Font(FCoreStyle::Get().GetFontStyle("ToolTip.LargerFont"))
-										.ColorAndOpacity(FLinearColor::Black)
 									]
 									+ SVerticalBox::Slot()
 									[
 										SNew(STextBlock)
 										.Text(this, &SCurveTimelineView::GetCurveToolTipOutputText)
 										.Font(FCoreStyle::Get().GetFontStyle("ToolTip.LargerFont"))
-										.ColorAndOpacity(FLinearColor::Black)
 									]
 								]);
 						}
diff --git a/Engine/Plugins/AudioInsights/Source/AudioInsights/Private/Views/SAudioCurveView.cpp b/Engine/Plugins/AudioInsights/Source/AudioInsights/Private/Views/SAudioCurveView.cpp
--- a/Engine/Plugins/AudioInsights/Source/AudioInsights/Private/Views/SAudioCurveView.cpp
+++ b/Engine/Plugins/AudioInsights/Source/AudioInsights/Private/Views/SAudioCurveView.cpp
@@ -102,7 +102,6 @@ TSharedRef<SToolTip> SAudioCurveView::CreateCurveTooltip()
 		{ 
 			return MetadataPerCurve == nullptr ? EVisibility::Collapsed : EVisibility::Visible; 
 		})
-		.BorderImage(FCoreStyle::Get().GetBrush("ToolTip.BrightBackground"))
 		[
 			SNew(SVerticalBox)
 			+ SVerticalBox::Slot()
@@ -110,28 +109,24 @@ TSharedRef<SToolTip> SAudioCurveView::CreateCurveTooltip()
 				SNew(STextBlock)
 				.Text(this, &SAudioCurveView::GetCurveToolTipDisplayNameText)
 				.Font(FCoreStyle::Get().GetFontStyle("ToolTip.LargerFont"))
-				.ColorAndOpacity(FLinearColor::Black)
 			]
 			+ SVerticalBox::Slot()
 			[
 				SNew(STextBlock)
 				.Text(this, &SAudioCurveView::GetCurveToolTipXValueText)
 				.Font(FCoreStyle::Get().GetFontStyle("ToolTip.LargerFont"))
-				.ColorAndOpacity(FLinearColor::Black)
 			]
 			+ SVerticalBox::Slot()
 			[
 				SNew(STextBlock)
 				.Text(this, &SAudioCurveView::GetCurveToolTipYValueText)
 				.Font(FCoreStyle::Get().GetFontStyle("ToolTip.LargerFont"))
-				.ColorAndOpacity(FLinearColor::Black)
 			]
 			+ SVerticalBox::Slot()
 			[
 				SNew(STextBlock)
 				.Text(AdditionalToolTipText)
 				.Font(FCoreStyle::Get().GetFontStyle("ToolTip.LargerFont"))
-				.ColorAndOpacity(FLinearColor::Black)
 				.Visibility_Lambda([this]() { return AdditionalToolTipText.IsSet() && !AdditionalToolTipText.Get().IsEmpty() ? EVisibility::Visible : EVisibility::Collapsed; })
 			]
 		];
diff --git a/Engine/Plugins/FX/Niagara/Source/NiagaraEditor/Private/NiagaraEditorUtilities.cpp b/Engine/Plugins/FX/Niagara/Source/NiagaraEditor/Private/NiagaraEditorUtilities.cpp
--- a/Engine/Plugins/FX/Niagara/Source/NiagaraEditor/Private/NiagaraEditorUtilities.cpp
+++ b/Engine/Plugins/FX/Niagara/Source/NiagaraEditor/Private/NiagaraEditorUtilities.cpp
@@ -5030,13 +5030,11 @@ TSharedRef<SToolTip> FNiagaraEditorUtilities::Tooltips::CreateTooltip(FText Desc
 	TSharedRef<STextBlock> DescriptionText = SNew(STextBlock)
 		.Text(Description)
 		.WrapTextAt(SToolTip::GetToolTipWrapWidth() - 500.f)
-		.Font(FCoreStyle::Get().GetFontStyle("ToolTip.Font"))
-		.ColorAndOpacity(FLinearColor::Black);
+		.Font(FCoreStyle::Get().GetFontStyle("ToolTip.Font"));
 	
 	DescriptionSlot->AttachWidget(DescriptionText);
 	
 	return SNew(SToolTip)
-		.BorderImage(FCoreStyle::Get().GetBrush("ToolTip.BrightBackground"))
 		[
 			TooltipContent
 		];
diff --git a/Engine/Plugins/Runtime/RigVM/Source/RigVMEditor/Private/Widgets/SRigVMVariantWidget.cpp b/Engine/Plugins/Runtime/RigVM/Source/RigVMEditor/Private/Widgets/SRigVMVariantWidget.cpp
--- a/Engine/Plugins/Runtime/RigVM/Source/RigVMEditor/Private/Widgets/SRigVMVariantWidget.cpp
+++ b/Engine/Plugins/Runtime/RigVM/Source/RigVMEditor/Private/Widgets/SRigVMVariantWidget.cpp
@@ -24,7 +24,6 @@ void SRigVMVariantToolTipWithTags::Construct(const FArguments& InArgs)
 	SToolTip::Construct(
 	SuperClassArgs
 		.TextMargin(11.0f)
-		.BorderImage( FCoreStyle::Get().GetBrush("ToolTip.BrightBackground") )
 	);
 }
 
@@ -53,7 +52,6 @@ void SRigVMVariantToolTipWithTags::OnOpening()
 		SNew( STextBlock )
 		.Text( SuperClassArgs._Text )
 		.Font( SuperClassArgs._Font )
-		.ColorAndOpacity( FLinearColor::Black )
 		.WrapTextAt_Static( &SToolTip::GetToolTipWrapWidth )
 	];
 
diff --git a/Engine/Source/Editor/ContentBrowser/Private/Menus/AssetViewMenus.cpp b/Engine/Source/Editor/ContentBrowser/Private/Menus/AssetViewMenus.cpp
--- a/Engine/Source/Editor/ContentBrowser/Private/Menus/AssetViewMenus.cpp
+++ b/Engine/Source/Editor/ContentBrowser/Private/Menus/AssetViewMenus.cpp
@@ -28,17 +28,15 @@ namespace UE::Editor::ContentBrowser
 					[
 						SNew(SRichTextBlock)
 						.DecoratorStyleSet(&FAppStyle::Get())
-						.TextStyle(&FAppStyle::GetWidgetStyle<FTextBlockStyle>("RichTextBlock.DarkText"))
 						.Text(LOCTEXT("CustomSizeToolTip",
-							"Hold <RichTextBlock.BoldDarkText>Ctrl</> and use the <RichTextBlock.BoldDarkText>Scroll Wheel</> to scale the thumbnails."
-							"\nHold <RichTextBlock.BoldDarkText>Ctrl+Shift</> and use the <RichTextBlock.BoldDarkText>Scroll Wheel</> to cycle through predefined thumbnail sizes."))
+							"Hold <RichTextBlock.Bold>Ctrl</> and use the <RichTextBlock.Bold>Scroll Wheel</> to scale the thumbnails."
+							"\nHold <RichTextBlock.Bold>Ctrl+Shift</> and use the <RichTextBlock.Bold>Scroll Wheel</> to cycle through predefined thumbnail sizes."))
 					];
 		}
 
 		TSharedRef<IToolTip> GetThumbnailSizeToolTip()
 		{
 			return SNew(SToolTip)
-					.BorderImage(FAppStyle::GetBrush("ToolTip.BrightBackground"))
 					.TextMargin(FMargin(1.f, 0.f))
 					.Content()
 					[
@@ -46,7 +44,6 @@ namespace UE::Editor::ContentBrowser
 						.OverrideExtendedToolTipContent(GetExtendedToolTipMouseWheel())
 						[
 							SNew(STextBlock)
-							.ColorAndOpacity(FLinearColor::Black)
 							.Text(LOCTEXT("ThumbnailSizeToolTip", "Adjust the size of thumbnails."))
 						]
 					];
diff --git a/Engine/Source/Editor/CurveEditor/Private/Views/SInteractiveCurveEditorView.cpp b/Engine/Source/Editor/CurveEditor/Private/Views/SInteractiveCurveEditorView.cpp
--- a/Engine/Source/Editor/CurveEditor/Private/Views/SInteractiveCurveEditorView.cpp
+++ b/Engine/Source/Editor/CurveEditor/Private/Views/SInteractiveCurveEditorView.cpp
@@ -127,7 +127,6 @@ void SInteractiveCurveEditorView::Construct(const FArguments& InArgs, TWeakPtr<F
 
 	TSharedRef<SDynamicToolTip> ToolTipWidget =
 		SNew(SDynamicToolTip)
-		.BorderImage(FCoreStyle::Get().GetBrush("ToolTip.BrightBackground"))
 		[
 			SNew(SVerticalBox)
 			+ SVerticalBox::Slot()
@@ -135,21 +134,18 @@ void SInteractiveCurveEditorView::Construct(const FArguments& InArgs, TWeakPtr<F
 				SNew(STextBlock)
 				.Text(this, &SInteractiveCurveEditorView::GetToolTipCurveName)
 				.Font(FCoreStyle::Get().GetFontStyle("ToolTip.LargerFont"))
-				.ColorAndOpacity(FLinearColor::Black)
 			]
 		+ SVerticalBox::Slot()
 		[
 			SNew(STextBlock)
 			.Text(this, &SInteractiveCurveEditorView::GetToolTipTimeText)
 			.Font(FCoreStyle::Get().GetFontStyle("ToolTip.LargerFont"))
-			.ColorAndOpacity(FLinearColor::Black)
 		]
 		+ SVerticalBox::Slot()
 		[
 			SNew(STextBlock)
 			.Text(this, &SInteractiveCurveEditorView::GetToolTipValueText)
 			.Font(FCoreStyle::Get().GetFontStyle("ToolTip.LargerFont"))
-			.ColorAndOpacity(FLinearColor::Black)
 		]
 	];
 
diff --git a/Engine/Source/Editor/Documentation/Private/Documentation.cpp b/Engine/Source/Editor/Documentation/Private/Documentation.cpp
--- a/Engine/Source/Editor/Documentation/Private/Documentation.cpp
+++ b/Engine/Source/Editor/Documentation/Private/Documentation.cpp
@@ -222,10 +222,6 @@ TSharedRef< class SToolTip > FDocumentation::CreateToolTip(const TAttribute<FTex
 	
 	return SNew( SToolTip )
 		.IsInteractive( DocToolTip.ToSharedRef(), &SDocumentationToolTip::IsInteractive )
-
-		// Emulate text-only tool-tip styling that SToolTip uses when no custom content is supplied.  We want documentation tool-tips to 
-		// be styled just like text-only tool-tips
-		.BorderImage( FCoreStyle::Get().GetBrush("ToolTip.BrightBackground") )
 		// The documentation tooltip has padding built in, so that it can style its own background within the tooltip borders.
 		.TextMargin(FMargin(0.0f))
 		[
@@ -253,10 +249,6 @@ TSharedRef< class SToolTip > FDocumentation::CreateToolTip(const TAttribute<FTex
 
 	return SNew(SToolTip)
 		.IsInteractive(DocToolTip, &SDocumentationToolTip::IsInteractive)
-
-		// Emulate text-only tool-tip styling that SToolTip uses when no custom content is supplied.  We want documentation tool-tips to 
-		// be styled just like text-only tool-tips
-		.BorderImage( FCoreStyle::Get().GetBrush("ToolTip.BrightBackground") )
 		// The documentation tooltip has padding built in, so that it can style its own background within the tooltip borders.
 		.TextMargin(FMargin(0.0f))
 		[
diff --git a/Engine/Source/Editor/Documentation/Public/SDocumentationToolTip.h b/Engine/Source/Editor/Documentation/Public/SDocumentationToolTip.h
--- a/Engine/Source/Editor/Documentation/Public/SDocumentationToolTip.h
+++ b/Engine/Source/Editor/Documentation/Public/SDocumentationToolTip.h
@@ -9,6 +9,7 @@
 #include "Layout/Margin.h"
 #include "Math/Color.h"
 #include "Misc/Attribute.h"
+#include "Styling/AppStyle.h"
 #include "Styling/SlateColor.h"
 #include "Styling/SlateTypes.h"
 #include "Templates/SharedPointer.h"
@@ -37,7 +38,7 @@ public:
 		, _SubduedStyle( TEXT("Documentation.SDocumentationTooltipSubdued") )
 		, _HyperlinkTextStyle( TEXT("Documentation.SDocumentationTooltipHyperlinkText") )
 		, _HyperlinkButtonStyle( TEXT("Documentation.SDocumentationTooltipHyperlinkButton") )
-		, _ColorAndOpacity( FLinearColor::Black )
+		, _ColorAndOpacity(FAppStyle::Get().GetSlateColor("Documentation.SDocumentationTooltipTextColor"))
 		, _AddDocumentation( true )
 		, _DocumentationMargin(0)
 		, _OverrideExtendedToolTipContent(nullptr)
diff --git a/Engine/Source/Editor/EditorStyle/Private/StarshipStyle.cpp b/Engine/Source/Editor/EditorStyle/Private/StarshipStyle.cpp
--- a/Engine/Source/Editor/EditorStyle/Private/StarshipStyle.cpp
+++ b/Engine/Source/Editor/EditorStyle/Private/StarshipStyle.cpp
@@ -1323,6 +1323,12 @@ void FStarshipEditorStyle::FStyle::SetupGeneralStyles()
 			.SetFont(DEFAULT_FONT("Regular", 8))
 			.SetColorAndOpacity(FStyleColors::Foreground);
 		Set("AssetThumbnail.Tooltip.MoreInfoText", TooltipTextSubdued);
+		
+		// SToolTip defaults...
+		{
+			Set("ToolTip.TextColor", FSlateColor(FStyleColors::Foreground));
+			Set("ToolTip.TextBackground", new FSlateRoundedBoxBrush(FStyleColors::Dropdown, 0.0f, FStyleColors::DropdownOutline, 1.0f));
+		}
 
 		const FButtonStyle EditModePrimitives = FButtonStyle()
 			.SetNormal(FSlateRoundedBoxBrush(FStyleColors::Secondary, 4.f, COLOR("#121212FF"), 1.f))
@@ -3377,7 +3383,7 @@ void FStarshipEditorStyle::FStyle::SetupTutorialStyles()
 
 		const FTextBlockStyle DocumentationTooltipTextSubdued = FTextBlockStyle( NormalText )
 			.SetFont( DEFAULT_FONT( "Regular", 8 ) )
-			.SetColorAndOpacity( FLinearColor( 0.1f, 0.1f, 0.1f ) );
+			.SetColorAndOpacity( FLinearColor::Gray );
 		Set("Documentation.SDocumentationTooltipSubdued", FTextBlockStyle(DocumentationTooltipTextSubdued));
 
 		const FTextBlockStyle DocumentationTooltipHyperlinkText = FTextBlockStyle( NormalText )
@@ -3390,6 +3396,8 @@ void FStarshipEditorStyle::FStyle::SetupTutorialStyles()
 				.SetPressed(FSlateNoResource())
 				.SetHovered(BORDER_BRUSH( "Old/HyperlinkUnderline", FMargin(0,0,0,3/16.0f), HyperlinkColor ) );
 		Set("Documentation.SDocumentationTooltipHyperlinkButton", FButtonStyle(DocumentationTooltipHyperlinkButton));
+
+		Set("Documentation.SDocumentationTooltipTextColor", FSlateColor::UseForeground());
 	}
 
 
@@ -4753,44 +4761,28 @@ void FStarshipEditorStyle::FStyle::SetupGraphEditorStyles()
 		Set("SCSEditor.ComponentTooltip.Title",
 			FTextBlockStyle(NormalText)
 			.SetFont(DEFAULT_FONT("Regular", 12))
-			.SetColorAndOpacity(FLinearColor::Black)
 			);
 
 		Set("SCSEditor.ComponentTooltip.Label",
 			FTextBlockStyle(NormalText)
-			.SetColorAndOpacity(FLinearColor(0.075f, 0.075f, 0.075f))
-			.SetShadowOffset(FVector2f::UnitVector)
-			.SetShadowColorAndOpacity(FLinearColor(0.8f, 0.8f, 0.8f))
 		);
 		Set("SCSEditor.ComponentTooltip.ImportantLabel",
 			FTextBlockStyle(NormalText)
-			.SetColorAndOpacity(FLinearColor(0.05f, 0.05f, 0.05f))
-			.SetShadowOffset(FVector2f::UnitVector)
-			.SetShadowColorAndOpacity(FLinearColor(0.8f, 0.8f, 0.8f))
 		);
 
 
 		Set("SCSEditor.ComponentTooltip.Value", 
 			FTextBlockStyle(NormalText)
 			.SetFont(DEFAULT_FONT("Bold", 10))
-			.SetColorAndOpacity(FLinearColor(0.0f, 0.0f, 0.0f))
-			.SetShadowOffset(FVector2f::UnitVector)
-			.SetShadowColorAndOpacity(FLinearColor(0.8f, 0.8f, 0.8f))
 		);
 		Set("SCSEditor.ComponentTooltip.ImportantValue",
 			FTextBlockStyle(NormalText)
 			.SetFont(DEFAULT_FONT("Bold", 10))
-			.SetColorAndOpacity(FLinearColor(0.3f, 0.0f, 0.0f))
-			.SetShadowOffset(FVector2f::UnitVector)
-			.SetShadowColorAndOpacity(FLinearColor(0.8f, 0.8f, 0.8f))
 		);
 
 		Set("SCSEditor.ComponentTooltip.ClassDescription",
 			FTextBlockStyle(NormalText)
 			.SetFont(DEFAULT_FONT("Italic", 10))
-			.SetColorAndOpacity(FLinearColor(0.1f, 0.1f, 0.1f))
-			.SetShadowOffset(FVector2f::UnitVector)
-			.SetShadowColorAndOpacity(FLinearColor(0.8f, 0.8f, 0.8f))
 		);
 	}
 
diff --git a/Engine/Source/Editor/GraphEditor/Private/KismetNodes/SGraphNodeK2Base.cpp b/Engine/Source/Editor/GraphEditor/Private/KismetNodes/SGraphNodeK2Base.cpp
--- a/Engine/Source/Editor/GraphEditor/Private/KismetNodes/SGraphNodeK2Base.cpp
+++ b/Engine/Source/Editor/GraphEditor/Private/KismetNodes/SGraphNodeK2Base.cpp
@@ -339,10 +339,6 @@ TSharedPtr<SToolTip> SGraphNodeK2Base::GetComplexTooltip()
 	SAssignNew(NodeToolTip, SToolTip)
 		.Visibility_Static(&LocalUtils::IsToolTipVisible, ThisWeak)
 		.IsInteractive_Static(&LocalUtils::IsInteractive)
-
-		// Emulate text-only tool-tip styling that SToolTip uses when no custom content is supplied.  We want node tool-tips to 
-		// be styled just like text-only tool-tips
-		.BorderImage( FCoreStyle::Get().GetBrush("ToolTip.BrightBackground") )
 		.TextMargin(FMargin(11.0f))
 	[
 		SAssignNew(VerticalBoxWidget, SVerticalBox)
@@ -350,17 +346,17 @@ TSharedPtr<SToolTip> SGraphNodeK2Base::GetComplexTooltip()
 		+SVerticalBox::Slot()
 		[
 			SNew(SVerticalBox)
-				.Visibility_Static(&LocalUtils::IsToolTipHeadingVisible, ThisWeak)
+			.Visibility_Static(&LocalUtils::IsToolTipHeadingVisible, ThisWeak)
 			+SVerticalBox::Slot()
-				.AutoHeight()
+			.AutoHeight()
 			[
 				SNew(STextBlock)
-					.TextStyle( GetStyleSet(), "Documentation.SDocumentationTooltipSubdued")
-					.Text(this, &SGraphNodeK2Base::GetToolTipHeading)
+				.TextStyle( GetStyleSet(), "Documentation.SDocumentationTooltipSubdued")
+				.Text(this, &SGraphNodeK2Base::GetToolTipHeading)
 			]
 			+SVerticalBox::Slot()
-				.AutoHeight()
-				.Padding(0.f, 2.f, 0.f, 5.f)
+			.AutoHeight()
+			.Padding(0.f, 2.f, 0.f, 5.f)
 			[
 				SNew(SBorder)
 				// use the border's padding to actually create the horizontal line
@@ -392,9 +388,8 @@ TSharedPtr<SToolTip> SGraphNodeK2Base::GetComplexTooltip()
 			.HAlign( HAlign_Right )
 			[
 				SNew( STextBlock )
-				.ColorAndOpacity( FSlateColor::UseSubduedForeground() )
 				.Text( LOCTEXT( "NativeNodeName", "hold (Alt) for native node name" ) )
-				.TextStyle( &GetStyleSet().GetWidgetStyle<FTextBlockStyle>(TEXT("Documentation.SDocumentationTooltip")) )
+				.TextStyle( &GetStyleSet().GetWidgetStyle<FTextBlockStyle>(TEXT("Documentation.SDocumentationTooltipSubdued")) )
 				.Visibility_Static(&Local::GetNativeNodeNameVisibility)
 			];
 	}
diff --git a/Engine/Source/Editor/Kismet/Private/SBlueprintPalette.cpp b/Engine/Source/Editor/Kismet/Private/SBlueprintPalette.cpp
--- a/Engine/Source/Editor/Kismet/Private/SBlueprintPalette.cpp
+++ b/Engine/Source/Editor/Kismet/Private/SBlueprintPalette.cpp
@@ -518,32 +518,23 @@ static TSharedRef<IToolTip> ConstructToolTipWithActionPath(TSharedPtr<FEdGraphSc
 	FFavoritedBlueprintPaletteItem ActionItem(*ActionIn);
 	if (ActionItem.IsValid())
 	{
-		static FTextBlockStyle PathStyle = FTextBlockStyle()
-			.SetFont(FCoreStyle::GetDefaultFontStyle("Regular", 8))
-			.SetColorAndOpacity(FLinearColor(0.4f, 0.4f, 0.4f));
-
-		NewToolTip = SNew(SToolTip)
-
-		// Emulate text-only tool-tip styling that SToolTip uses when no custom content is supplied.  We want node tool-tips to 
-		// be styled just like text-only tool-tips
-		.BorderImage( FCoreStyle::Get().GetBrush("ToolTip.BrightBackground") )
-		.TextMargin(FMargin(11.0f))
-		[
+		NewToolTip->SetContentWidget(
 			SNew(SVerticalBox)
 			+ SVerticalBox::Slot()
+			.FillHeight(1.f)
 			[
 				OldToolTip->GetContentWidget()
 			]
-
 			+ SVerticalBox::Slot()
-			.HAlign(EHorizontalAlignment::HAlign_Right)
+			.AutoHeight()
 			[
 				SNew(STextBlock)
-				.TextStyle( FAppStyle::Get(), "Documentation.SDocumentationTooltip")
+				.Margin(FMargin(8))
+				.WrapTextAt_Static(&SToolTip::GetToolTipWrapWidth)
+				.Font(FCoreStyle::Get().GetFontStyle("ToolTip.Font"))
 				.Text(FText::FromString(ActionItem.ToString()))
-				//.TextStyle(&PathStyle)
 			]
-		];
+		);
 	}
 
 	return NewToolTip;
@@ -2162,17 +2153,11 @@ TSharedPtr<SToolTip> SBlueprintPaletteItem::ConstructToolTipWidget() const
 	if ( (NodeTemplate != nullptr) && (FInternationalization::Get().GetCurrentCulture()->GetTwoLetterISOLanguageName() != TEXT("en")) )
 	{
 		FText NativeNodeName = FText::FromString(NodeTemplate->GetNodeTitle(ENodeTitleType::ListView).BuildSourceString());
-		const FTextBlockStyle& SubduedTextStyle = FAppStyle::GetWidgetStyle<FTextBlockStyle>("Documentation.SDocumentationTooltipSubdued");
 
 		TSharedPtr<SToolTip> InternationalTooltip;
 		TSharedPtr<SVerticalBox> TooltipBody;
 
 		SAssignNew(InternationalTooltip, SToolTip)
-			// Emulate text-only tool-tip styling that SToolTip uses 
-			// when no custom content is supplied.  We want node tool-
-			// tips to be styled just like text-only tool-tips
-			.BorderImage( FCoreStyle::Get().GetBrush("ToolTip.BrightBackground") )
-			.TextMargin(FMargin(11.0f))
 		[
 			SAssignNew(TooltipBody, SVerticalBox)
 		];
@@ -2188,34 +2173,35 @@ TSharedPtr<SToolTip> SBlueprintPaletteItem::ConstructToolTipWidget() const
 			TooltipBody->AddSlot()
 			[
 				SNew(STextBlock)
-					.TextStyle(FAppStyle::Get(), "Documentation.SDocumentationTooltip")
-					.Text(NativeNodeName)
-					.Visibility_Lambda([GetNativeNamePromptVisibility]()->EVisibility
-					{
-						return (GetNativeNamePromptVisibility() == EVisibility::Visible) ? EVisibility::Collapsed : EVisibility::Visible;
-					})
+				.Margin(FMargin(8))
+				.Text(NativeNodeName)
+				.Visibility_Lambda([GetNativeNamePromptVisibility]()->EVisibility
+				{
+					return (GetNativeNamePromptVisibility() == EVisibility::Visible) ? EVisibility::Collapsed : EVisibility::Visible;
+				})
 			];
 
 			TooltipBody->AddSlot()
 			[
 				SNew(SHorizontalBox)
-					.Visibility_Lambda(GetNativeNamePromptVisibility)
+				.Visibility_Lambda(GetNativeNamePromptVisibility)
 				+SHorizontalBox::Slot()
 				[
 					TooltipWidget->GetContentWidget()
 				]
 			];
 
+			const FTextBlockStyle& SubduedTextStyle = FAppStyle::GetWidgetStyle<FTextBlockStyle>("Documentation.SDocumentationTooltipSubdued");
+
 			TooltipBody->AddSlot()
-				.AutoHeight()
-				.HAlign(HAlign_Center)
-				.Padding(0.f, 8.f, 0.f, 0.f)
+			.AutoHeight()
+			.HAlign(HAlign_Center)
 			[
-
 				SNew(STextBlock)
-					.Text( LOCTEXT("NativeNodeName", "hold (Alt) for native node name") )
-					.TextStyle(&SubduedTextStyle)
-					.Visibility_Lambda(GetNativeNamePromptVisibility)
+				.Margin(FMargin(8))
+				.TextStyle(&SubduedTextStyle)
+				.Text(LOCTEXT("NativeNodeName", "hold (Alt) for native node name"))
+				.Visibility_Lambda(GetNativeNamePromptVisibility)
 			];
 		}
 		else
@@ -2228,26 +2214,25 @@ TSharedPtr<SToolTip> SBlueprintPaletteItem::ConstructToolTipWidget() const
 
 			// give the "advanced" tooltip a header
 			TooltipBody->AddSlot()
-				.AutoHeight()
-				.HAlign(HAlign_Right)
-				.Padding(0.f, 0.f, 0.f, 8.f)
+			.AutoHeight()
+			.HAlign(HAlign_Right)
 			[
 				SNew(SHorizontalBox)
 				+SHorizontalBox::Slot()
-					.AutoWidth()
+				.AutoWidth()
 				[
 					SNew(STextBlock)
-					.TextStyle(&SubduedTextStyle)
-						.Text(LOCTEXT("NativeNodeNameLabel", "Native Node Name: "))
-						.Visibility_Lambda(GetNativeNodeNameVisibility)
+					.Margin(FMargin(8))
+					.Text(LOCTEXT("NativeNodeNameLabel", "Native Node Name: "))
+					.Visibility_Lambda(GetNativeNodeNameVisibility)
 				]
 				+SHorizontalBox::Slot()
-					.AutoWidth()
+				.AutoWidth()
 				[
 					SNew(STextBlock)
-						.TextStyle(&SubduedTextStyle)
-						.Text(NativeNodeName)
-						.Visibility_Lambda(GetNativeNodeNameVisibility)
+					.Margin(FMargin(8))
+					.Text(NativeNodeName)
+					.Visibility_Lambda(GetNativeNodeNameVisibility)
 				]
 			];
 
diff --git a/Engine/Source/Editor/MaterialEditor/Private/SMaterialEditorCustomPrimitiveDataWidget.cpp b/Engine/Source/Editor/MaterialEditor/Private/SMaterialEditorCustomPrimitiveDataWidget.cpp
--- a/Engine/Source/Editor/MaterialEditor/Private/SMaterialEditorCustomPrimitiveDataWidget.cpp
+++ b/Engine/Source/Editor/MaterialEditor/Private/SMaterialEditorCustomPrimitiveDataWidget.cpp
@@ -34,7 +34,10 @@ public:
 
 		if (RowData->bIsDuplicate)
 		{
-			TextBlock->SetToolTip(SNew(SToolTip).Text(FText::FromString("This slot is potentially incorrectly overlapping")).BorderImage(FCoreStyle::Get().GetBrush("ToolTip.BrightBackground")));
+			TextBlock->SetToolTip(
+				SNew(SToolTip)
+				.Text(FText::FromString("This slot is potentially incorrectly overlapping"))
+			);
 			TextBlock->SetColorAndOpacity(FSlateColor(FLinearColor::Red));
 		}
 
diff --git a/Engine/Source/Editor/MaterialEditor/Private/SMaterialExpressionsOverviewWidget.cpp b/Engine/Source/Editor/MaterialEditor/Private/SMaterialExpressionsOverviewWidget.cpp
--- a/Engine/Source/Editor/MaterialEditor/Private/SMaterialExpressionsOverviewWidget.cpp
+++ b/Engine/Source/Editor/MaterialEditor/Private/SMaterialExpressionsOverviewWidget.cpp
@@ -30,7 +30,10 @@ public:
 
 		if (RowData->bIsDuplicate)
 		{
-			TextBlock->SetToolTip(SNew(SToolTip).Text(FText::FromString("This slot is potentially incorrectly overlapping")).BorderImage(FCoreStyle::Get().GetBrush("ToolTip.BrightBackground")));
+			TextBlock->SetToolTip(
+				SNew(SToolTip)
+				.Text(FText::FromString("This slot is potentially incorrectly overlapping"))
+			);
 			TextBlock->SetColorAndOpacity(FSlateColor(FLinearColor::Red));
 		}
 
diff --git a/Engine/Source/Editor/UnrealEd/Private/SCurveEditor.cpp b/Engine/Source/Editor/UnrealEd/Private/SCurveEditor.cpp
--- a/Engine/Source/Editor/UnrealEd/Private/SCurveEditor.cpp
+++ b/Engine/Source/Editor/UnrealEd/Private/SCurveEditor.cpp
@@ -1568,7 +1568,6 @@ void SCurveEditor::UpdateCurveToolTip(const FGeometry& InMyGeometry, const FPoin
 			{
 				SetToolTip(
 					SAssignNew(CurveToolTip, SToolTip)
-					.BorderImage( FCoreStyle::Get().GetBrush( "ToolTip.BrightBackground" ) )
 					[
 						SNew(SVerticalBox)
 						+ SVerticalBox::Slot()
@@ -1576,21 +1575,18 @@ void SCurveEditor::UpdateCurveToolTip(const FGeometry& InMyGeometry, const FPoin
 							SNew(STextBlock)
 							.Text(this, &SCurveEditor::GetCurveToolTipNameText)
 							.Font(FCoreStyle::Get().GetFontStyle("ToolTip.LargerFont"))
-							.ColorAndOpacity( FLinearColor::Black)
 						]
 						+ SVerticalBox::Slot()
 						[
 							SNew(STextBlock)
 							.Text(this, &SCurveEditor::GetCurveToolTipInputText)
 							.Font(FCoreStyle::Get().GetFontStyle("ToolTip.LargerFont"))
-							.ColorAndOpacity(FLinearColor::Black)
 						]
 						+ SVerticalBox::Slot()
 						[
 							SNew(STextBlock)
 							.Text(this, &SCurveEditor::GetCurveToolTipOutputText)
 							.Font(FCoreStyle::Get().GetFontStyle("ToolTip.LargerFont"))
-							.ColorAndOpacity(FLinearColor::Black)
 						]
 					]);
 			}
diff --git a/Engine/Source/Runtime/Slate/Private/Widgets/SToolTip.cpp b/Engine/Source/Runtime/Slate/Private/Widgets/SToolTip.cpp
--- a/Engine/Source/Runtime/Slate/Private/Widgets/SToolTip.cpp
+++ b/Engine/Source/Runtime/Slate/Private/Widgets/SToolTip.cpp
@@ -1,36 +1,35 @@
 // Copyright Epic Games, Inc. All Rights Reserved.
 
 #include "Widgets/SToolTip.h"
+
 #include "HAL/IConsoleManager.h"
 #include "Widgets/Layout/SBorder.h"
 #include "Widgets/Text/STextBlock.h"
 
-
 static TAutoConsoleVariable<float> StaticToolTipWrapWidth(
 	TEXT( "Slate.ToolTipWrapWidth" ),
 	1000.0f,
 	TEXT( "Width of Slate tool-tips before we wrap the tool-tip text" ) );
 
-
 float SToolTip::GetToolTipWrapWidth()
 {
 	return StaticToolTipWrapWidth.GetValueOnAnyThread();
 }
 
-
 void SToolTip::Construct( const FArguments& InArgs )
 {
 	TextContent = InArgs._Text;
 	bIsInteractive = InArgs._IsInteractive;
 	Font = InArgs._Font;
 	TextMargin = InArgs._TextMargin;
+	TextColor = InArgs._TextColor;
+	TextBorderImage = InArgs._TextBorderImage;
 	BorderImage = InArgs._BorderImage;
 	OnSetInteractiveWindowLocationDelegate = InArgs._OnSetInteractiveWindowLocation;
 	
 	SetContentWidget(InArgs._Content.Widget);
 }
 
-
 void SToolTip::SetContentWidget(const TSharedRef<SWidget>& InContentWidget)
 {
 	if (InContentWidget != SNullWidget::NullWidget)
@@ -39,14 +38,10 @@ void SToolTip::SetContentWidget(const TSharedRef<SWidget>& InContentWidget)
 		WidgetContent = InContentWidget;
 	}
 
-	TSharedPtr< SWidget > PinnedWidgetContent = WidgetContent.Pin();
-	if( PinnedWidgetContent.IsValid() )
+	if (const TSharedPtr<SWidget> PinnedWidgetContent = WidgetContent.Pin())
 	{
 		ToolTipContent = PinnedWidgetContent;
-
-		// Tool-tip with entirely custom content.  We'll create a border with some padding (as customized by the user), then
-		// embed their custom widget right inside the border.  This tool-tip currently has a different styling than tool-tips
-		// that only contain text.
+		
 		ChildSlot
 		[
 			SNew(SBorder)
@@ -56,24 +51,20 @@ void SToolTip::SetContentWidget(const TSharedRef<SWidget>& InContentWidget)
 				ToolTipContent.ToSharedRef()
 			]
 		];
-
 	}
 	else
 	{
-		ToolTipContent =
-			SNew( STextBlock )
-			.Text( TextContent )
-			.Font( Font )
-			.ColorAndOpacity( FLinearColor::Black )
-			.WrapTextAt_Static( &SToolTip::GetToolTipWrapWidth );
+		ToolTipContent = SNew(STextBlock)
+			.Text(TextContent)
+			.Font(Font)
+			.ColorAndOpacity(TextColor)
+			.WrapTextAt_Static(&SToolTip::GetToolTipWrapWidth);
 
-		// Text-only tool-tip.  This tool-tip currently has a different styling than tool-tips with custom content.  We always want basic
-		// text tool-tips to look consistent.
 		ChildSlot
 		[
 			SNew(SBorder)
-			.BorderImage( FCoreStyle::Get().GetBrush("ToolTip.BrightBackground") )
-			.Padding(FMargin(11.0f))
+			.BorderImage(TextBorderImage)
+			.Padding(TextMargin)
 			[
 				ToolTipContent.ToSharedRef()
 			]
@@ -92,7 +83,6 @@ bool SToolTip::IsEmpty() const
 	return !WidgetContent.IsValid() && TextContent.Get().IsEmpty();
 }
 
-
 bool SToolTip::IsInteractive() const
 {
 	return bIsInteractive.Get();
diff --git a/Engine/Source/Runtime/Slate/Public/Widgets/SToolTip.h b/Engine/Source/Runtime/Slate/Public/Widgets/SToolTip.h
--- a/Engine/Source/Runtime/Slate/Public/Widgets/SToolTip.h
+++ b/Engine/Source/Runtime/Slate/Public/Widgets/SToolTip.h
@@ -2,15 +2,14 @@
 
 #pragma once
 
-#include "CoreMinimal.h"
 #include "Misc/Attribute.h"
 #include "Layout/Margin.h"
 #include "Widgets/DeclarativeSyntaxSupport.h"
 #include "Styling/SlateColor.h"
 #include "Widgets/SWidget.h"
 #include "Fonts/SlateFontInfo.h"
+#include "Styling/AppStyle.h"
 #include "Widgets/SCompoundWidget.h"
-#include "Styling/CoreStyle.h"
 #include "Widgets/IToolTip.h"
 
 /**
@@ -26,9 +25,11 @@ public:
 	SLATE_BEGIN_ARGS( SToolTip )
 		: _Text()
 		, _Content()
-		, _Font(FCoreStyle::Get().GetFontStyle("ToolTip.Font"))
+		, _Font(FAppStyle::Get().GetFontStyle("ToolTip.Font"))
 		, _TextMargin(FMargin(8.0f))
-		, _BorderImage(FCoreStyle::Get().GetBrush("ToolTip.Background"))
+		, _TextColor(FAppStyle::Get().GetSlateColor("ToolTip.TextColor"))
+		, _TextBorderImage(FAppStyle::Get().GetBrush("ToolTip.TextBackground"))
+		, _BorderImage(FAppStyle::Get().GetBrush("ToolTip.Background"))
 		, _IsInteractive(false)
 		, _OnSetInteractiveWindowLocation()
 	{ }
@@ -45,6 +46,12 @@ public:
 		/** Margin between the tool tip border and the text content */
 		SLATE_ATTRIBUTE(FMargin, TextMargin)
 
+		/** The color of the text displayed in this tool tip */
+		SLATE_ATTRIBUTE(FSlateColor, TextColor)
+
+		/** The background/border image to display for text tooltips */
+		SLATE_ATTRIBUTE(const FSlateBrush*, TextBorderImage)
+
 		/** The background/border image to display */
 		SLATE_ATTRIBUTE(const FSlateBrush*, BorderImage)
 
@@ -63,8 +70,6 @@ public:
 	 */
 	SLATE_API void Construct( const FArguments& InArgs );
 
-public:
-
 	// IToolTip interface
 
 	virtual TSharedRef<class SWidget> AsWidget( ) override
@@ -91,12 +96,9 @@ public:
 		return TextContent.Get();
 	}
 
-public:
-
 	static SLATE_API float GetToolTipWrapWidth();
 
 private:
-
 	// Text block widget.
 	TAttribute<FText> TextContent;
 
@@ -111,8 +113,14 @@ private:
 
 	// Margin between the tool tip border and the text content
 	TAttribute<FMargin> TextMargin;
+	
+	/** Color used for the text displayed */
+	TAttribute<FSlateColor> TextColor;
+	
+	/** The background/border image to display for text tooltips */
+	TAttribute<const FSlateBrush*> TextBorderImage;
 
-	// The background/border image to display
+	/** The background/border image to display for custom content */
 	TAttribute<const FSlateBrush*> BorderImage;
 	
 	// Whether the tooltip should be considered interactive.
diff --git a/Engine/Source/Runtime/SlateCore/Private/Styling/CoreStyle.cpp b/Engine/Source/Runtime/SlateCore/Private/Styling/CoreStyle.cpp
--- a/Engine/Source/Runtime/SlateCore/Private/Styling/CoreStyle.cpp
+++ b/Engine/Source/Runtime/SlateCore/Private/Styling/CoreStyle.cpp
@@ -642,6 +642,9 @@ TSharedRef<ISlateStyle> FCoreStyle::Create( const FName& InStyleSetName )
 
 		Style->Set("ToolTip.LargerFont", DEFAULT_FONT("Regular", 9));
 		Style->Set("ToolTip.BrightBackground", new BOX_BRUSH("Old/ToolTip_BrightBackground", FMargin(8.0f / 64.0f)));
+
+		Style->Set("ToolTip.TextColor", FSlateColor(FLinearColor::Black));
+		Style->Set("Documentation.SDocumentationTooltipTextColor", FSlateColor(FLinearColor::Black));
 	}
 
 	// SBorder defaults...
diff --git a/Engine/Source/Runtime/SlateCore/Private/Styling/StarshipCoreStyle.cpp b/Engine/Source/Runtime/SlateCore/Private/Styling/StarshipCoreStyle.cpp
--- a/Engine/Source/Runtime/SlateCore/Private/Styling/StarshipCoreStyle.cpp
+++ b/Engine/Source/Runtime/SlateCore/Private/Styling/StarshipCoreStyle.cpp
@@ -507,11 +507,13 @@ TSharedRef<ISlateStyle> FStarshipCoreStyle::Create()
 
 	// SToolTip defaults...
 	{
-		Style->Set("ToolTip.Font", DEFAULT_FONT("Regular", 9));
-		Style->Set("ToolTip.Background", new FSlateRoundedBoxBrush(FStyleColors::Dropdown, 0.0f, FStyleColors::DropdownOutline, 1.0f));
-
 		Style->Set("ToolTip.LargerFont", StyleFonts.Normal);
 		Style->Set("ToolTip.BrightBackground", new BOX_BRUSH("Old/ToolTip_BrightBackground", FMargin(8.0f / 64.0f)));
+		Style->Set("ToolTip.Font", DEFAULT_FONT("Regular", 9));
+		Style->Set("ToolTip.Background", new FSlateRoundedBoxBrush(FStyleColors::Dropdown, 0.0f, FStyleColors::DropdownOutline, 1.0f));
+		Style->Set("ToolTip.TextColor", FSlateColor(FLinearColor::Black));
+		Style->Set("Documentation.SDocumentationTooltipTextColor", FSlateColor(FLinearColor::Black));
+		Style->Set("ToolTip.TextBackground", new BOX_BRUSH("Old/ToolTip_BrightBackground", FMargin(8.0f / 64.0f)));
 	}
 
 	// SBorder defaults...
diff --git a/Engine/Source/Runtime/UMG/Private/Components/Widget.cpp b/Engine/Source/Runtime/UMG/Private/Components/Widget.cpp
--- a/Engine/Source/Runtime/UMG/Private/Components/Widget.cpp
+++ b/Engine/Source/Runtime/UMG/Private/Components/Widget.cpp
@@ -43,6 +43,8 @@
 #include "Extensions/UIComponent.h"
 #include "Extensions/UIComponentUserWidgetExtension.h"
 #include "Types/InvisibleToWidgetReflectorMetaData.h"
+#include "Widgets/Text/STextBlock.h"
+#include "Styling/AppStyle.h"
 
 #include UE_INLINE_GENERATED_CPP_BY_NAME(Widget)
 
@@ -531,7 +533,25 @@ void UWidget::SetToolTipText(const FText& InToolTipText)
 	TSharedPtr<SWidget> SafeWidget = GetCachedWidget();
 	if (SafeWidget.IsValid())
 	{
-		SafeWidget->SetToolTipText(InToolTipText);
+		if (!ToolTipText.IsEmpty())
+		{
+			// Match legacy tooltip style with white background
+			const TSharedRef<SToolTip> ToolTip = SNew(SToolTip)
+				.BorderImage(FAppStyle::Get().GetBrush("ToolTip.BrightBackground"))
+				[
+					SNew(STextBlock)
+					.Text(ToolTipText)
+					.Font(FAppStyle::Get().GetFontStyle("ToolTip.Font"))
+					.ColorAndOpacity(FLinearColor::Black)
+					.WrapTextAt_Static(&SToolTip::GetToolTipWrapWidth)
+				];
+
+			SafeWidget->SetToolTip(ToolTip);
+		}
+		else
+		{
+			SafeWidget->SetToolTip(TSharedPtr<IToolTip>());
+		}
 	}
 }
 
@@ -549,9 +569,9 @@ void UWidget::SetToolTip(UWidget* InToolTipWidget)
 	{
 		if ( ToolTipWidget )
 		{
+			// Match legacy tooltip style with white background
 			TSharedRef<SToolTip> ToolTip = SNew(SToolTip)
-				.TextMargin(FMargin(0))
-				.BorderImage(nullptr)
+				.BorderImage(FAppStyle::Get().GetBrush("ToolTip.BrightBackground"))
 				[
 					ToolTipWidget->TakeWidget()
 				];
@@ -1526,8 +1546,8 @@ PRAGMA_DISABLE_DEPRECATION_WARNINGS
 	else if ( ToolTipWidget != nullptr )
 	{
 		TSharedRef<SToolTip> ToolTip = SNew(SToolTip)
-			.TextMargin(FMargin(0))
-			.BorderImage(nullptr)
+			// Match legacy tooltip style with white background
+			.BorderImage(FAppStyle::Get().GetBrush("ToolTip.BrightBackground"))
 			[
 				ToolTipWidget->TakeWidget()
 			];
@@ -1536,7 +1556,18 @@ PRAGMA_DISABLE_DEPRECATION_WARNINGS
 	}
 	else if ( !ToolTipText.IsEmpty() || ToolTipTextDelegate.IsBound() )
 	{
-		SafeWidget->SetToolTipText(PROPERTY_BINDING(FText, ToolTipText));
+		// Match legacy tooltip style with white background
+		const TSharedRef<SToolTip> ToolTip = SNew(SToolTip)
+			.BorderImage(FAppStyle::Get().GetBrush("ToolTip.BrightBackground"))
+			[
+				SNew(STextBlock)
+				.Text(PROPERTY_BINDING(FText, ToolTipText))
+				.Font(FAppStyle::Get().GetFontStyle("ToolTip.Font"))
+				.ColorAndOpacity(FLinearColor::Black)
+				.WrapTextAt_Static(&SToolTip::GetToolTipWrapWidth)
+			];
+
+		SafeWidget->SetToolTip(ToolTip);
 	}
 
 #if WITH_ACCESSIBILITY
