Fix comment to reflect actual default value

#rnx

#Synced-CL 49832464

diff --git a/Engine/Source/Editor/Kismet/Private/BlueprintCompilationManager.cpp b/Engine/Source/Editor/Kismet/Private/BlueprintCompilationManager.cpp
--- a/Engine/Source/Editor/Kismet/Private/BlueprintCompilationManager.cpp
+++ b/Engine/Source/Editor/Kismet/Private/BlueprintCompilationManager.cpp
@@ -130,7 +130,7 @@ namespace UE::Kismet::BlueprintCompilationManager::Private
 		static bool bAllowFastReinstancing = false;
 		static FAutoConsoleVariableRef CVarAllowFastReinstancing(
 			TEXT("BP.bAllowFastReinstancing"), bAllowFastReinstancing,
-			TEXT("True by default, this flag enables a fast path in the logic that updates class \
+			TEXT("False by default, this flag enables a fast path in the logic that updates class \
 				objects at runtime (Reinstancing). When this flag is set we check class layout and \
 				construction data for a match, if the layout and construction data match then we \
 				don't replace instances of the class."),
