FindInBlueprintManager: Tweak ensure in FFindInBlueprintSearchManager::OnAssetLoaded to not trigger for existing SearchData that has already been marked as IsMarkedForDeletion().
Make the ensure logic more robust for external calls to UBlueprint::ClearEditorReferences for blueprints that has not already been marked as garbage.

#rnx

#Synced-CL 49944441

diff --git a/Engine/Source/Editor/Kismet/Private/FindInBlueprintManager.cpp b/Engine/Source/Editor/Kismet/Private/FindInBlueprintManager.cpp
--- a/Engine/Source/Editor/Kismet/Private/FindInBlueprintManager.cpp
+++ b/Engine/Source/Editor/Kismet/Private/FindInBlueprintManager.cpp
@@ -2557,9 +2557,11 @@ void FFindInBlueprintSearchManager::OnAssetLoaded(UObject* InAsset)
 		// The asset registry might not have informed us of this asset yet.
 		if (SearchData.IsValid())
 		{
-			// That index should never have a Blueprint already, but if it does, it should be the same Blueprint!
-			ensureMsgf(!SearchData.Blueprint.IsValid() || SearchData.Blueprint == BlueprintObject, TEXT("Blueprint in database has path %s and is being stomped by %s"), *(SearchData.AssetPath.ToString()), *AssetPath.ToString());
-			ensureMsgf(!SearchData.Blueprint.IsValid() || SearchData.AssetPath == AssetPath, TEXT("Blueprint in database has path %s and is being stomped by %s"), *(SearchData.AssetPath.ToString()), *AssetPath.ToString());
+			// That index should never have a valid Blueprint already, but if it does, it should be the same Blueprint!
+			ensureMsgf(SearchData.IsMarkedForDeletion() || !SearchData.Blueprint.IsValid() || SearchData.Blueprint == BlueprintObject,
+				TEXT("Blueprint in database has path %s and is being stomped by %s"), *(SearchData.AssetPath.ToString()), *AssetPath.ToString());
+			ensureMsgf(SearchData.IsMarkedForDeletion() || !SearchData.Blueprint.IsValid() || SearchData.AssetPath == AssetPath,
+				TEXT("Blueprint in database has path %s and is being stomped by %s"), *(SearchData.AssetPath.ToString()), *AssetPath.ToString());
 			SearchData.Blueprint = BlueprintObject;
 
 			// Apply the updated entry to the database.
