Skip reinstancing on load for prefabs

We already skip reinstancing for data-only blueprints with with a native parent class; Verse parent classes like `entity` can use the same optimization.

In addition, this avoids the need to reinstance VCells that live in these prefabs when cooking for VerseVM.

#rnx
#Synced-CL 50428254

diff --git a/Engine/Source/Editor/Kismet/Private/BlueprintCompilationManager.cpp b/Engine/Source/Editor/Kismet/Private/BlueprintCompilationManager.cpp
--- a/Engine/Source/Editor/Kismet/Private/BlueprintCompilationManager.cpp
+++ b/Engine/Source/Editor/Kismet/Private/BlueprintCompilationManager.cpp
@@ -47,6 +47,7 @@
 #include "UObject/ObjectInstancingGraph.h"
 #include "UObject/ReferenceChainSearch.h"
 #include "UObject/UObjectHash.h"
+#include "VerseVM/VVMVerseClass.h"
 #include "Kismet2/KismetDebugUtilities.h"
 #include "BlueprintEditorModule.h"
 #include "Algo/TopologicalSort.h"
@@ -548,7 +549,7 @@ struct FCompilerData
 	bool ShouldRegenerateSkeleton() const { return JobType != ECompilationManagerJobType::RelinkOnly; }
 	bool ShouldMarkUpToDateAfterSkeletonStage() const { return IsSkeletonOnly(); }
 	bool ShouldReconstructNodes() const { return JobType == ECompilationManagerJobType::Normal || (!IsSkeletonOnly() && BP->bIsRegeneratingOnLoad); }
-	bool ShouldSkipReinstancerCreation() const { return (IsSkeletonOnly() && (!BP->ParentClass || BP->ParentClass->IsNative())); }
+	bool ShouldSkipReinstancerCreation() const { return (IsSkeletonOnly() && (!BP->ParentClass || BP->ParentClass->IsNative() || BP->ParentClass->GetClass()->IsChildOf<UVerseClass>())); }
 	bool ShouldInitiateReinstancing() const { return JobType == ECompilationManagerJobType::Normal || BP->bIsRegeneratingOnLoad; }
 	bool ShouldCompileClassLayout() const { return JobType == ECompilationManagerJobType::Normal; }
 	bool ShouldCompileClassFunctions() const { return JobType == ECompilationManagerJobType::Normal; }
