[Editor viewport] Fix incorrect transform handling of static meshes in the viewport

When dragging static mesh instances in the blueprint editor viewport, the translation was applied incorrectly if the parent component had a non-identity scale.
We now apply the correct world -> component transform on the delta translation vector.

#Synced-CL 50789310

diff --git a/Engine/Source/Editor/Kismet/Private/InstancedStaticMeshSCSEditorCustomization.cpp b/Engine/Source/Editor/Kismet/Private/InstancedStaticMeshSCSEditorCustomization.cpp
--- a/Engine/Source/Editor/Kismet/Private/InstancedStaticMeshSCSEditorCustomization.cpp
+++ b/Engine/Source/Editor/Kismet/Private/InstancedStaticMeshSCSEditorCustomization.cpp
@@ -81,7 +81,8 @@ bool FInstancedStaticMeshSCSEditorCustomization::HandleViewportDrag(const UScene
 	UInstancedStaticMeshComponent* InstancedStaticMeshComponentTemplate = const_cast<UInstancedStaticMeshComponent*>(CastChecked<const UInstancedStaticMeshComponent>(InComponentTemplate));
 
 	// transform pivot into component's space
-	const FVector LocalPivot = InstancedStaticMeshComponentScene->GetComponentToWorld().InverseTransformPosition(InPivot);
+	const FTransform ComponentToWorldTransform = InstancedStaticMeshComponentScene->GetComponentToWorld();
+	const FVector LocalPivot = ComponentToWorldTransform.InverseTransformPosition(InPivot);
 
 	// Ensure that selected instances are up-to-date
 	ValidateSelectedInstances(InstancedStaticMeshComponentScene);
@@ -111,7 +112,8 @@ bool FInstancedStaticMeshSCSEditorCustomization::HandleViewportDrag(const UScene
 				NewTranslation += LocalPivot;
 			}
 
-			NewTranslation += InDeltaTranslation;
+			//transform translation from world space into to component space
+			NewTranslation += ComponentToWorldTransform.InverseTransformVector(InDeltaTranslation);
 
 			if( !InDeltaScale.IsNearlyZero() )
 			{
