Add snapshotclasses command for dumping all classes to text

#rnx

#Synced-CL 50958978

#Synced-CL 50958978

diff --git a/Engine/Source/Editor/Kismet/Private/Blueprints/JsonObjectGraphFunctionLibrary.cpp b/Engine/Source/Editor/Kismet/Private/Blueprints/JsonObjectGraphFunctionLibrary.cpp
--- a/Engine/Source/Editor/Kismet/Private/Blueprints/JsonObjectGraphFunctionLibrary.cpp
+++ b/Engine/Source/Editor/Kismet/Private/Blueprints/JsonObjectGraphFunctionLibrary.cpp
@@ -12,6 +12,7 @@
 
 namespace UE::Private 
 {
+FString GetIntermediateAssetName(const UObject* RootObject, const TCHAR* Prefix);
 
 static const TCHAR* SnapshotBlueprintsHelp = TEXT("Usage: snapshotblueprints label - label is a name for the folder where snapshots are saved");
 
@@ -72,12 +73,55 @@ FAutoConsoleCommand SnapshotBlueprintClassesCommand(
 	ECVF_Default
 );
 
-FString GetIntermediateAssetName(const UPackage* Package, const TCHAR* Prefix)
+static const TCHAR* SnapshotClassesHelp = TEXT("Usage: snapshotclasses <label> - where <label> will be used as the root for snapshot location within the Saved/Temp directory");
+
+void SnapshotClasses(const TArray<FString>& Args)
+{
+	if (Args.Num() == 0)
+	{
+		UE_LOG(LogEngine, Display, TEXT("%s"), SnapshotClassesHelp);
+		return;
+	}
+
+	const FString& Label = Args[0];
+	for (TObjectIterator<UClass> ClassIt; ClassIt; ++ClassIt)
+	{
+		const UClass* Class = *ClassIt;
+
+		if (!IsValid(Class) ||
+			!Class->GetDefaultObject(false) || 
+			Class->GetAuthoritativeClass() != Class ||
+			Class->HasAnyClassFlags(CLASS_Transient) // transient classes often have huge amounts of data, or are even incapable of being serialized
+			)
+		{
+			continue;
+		}
+
+		FJsonStringifyOptions Options;
+		Options.Flags |= EJsonStringifyFlags::FilterEditorOnlyData;
+		FUtf8String Result = UE::JsonObjectGraph::Stringify({ Class, Class->GetDefaultObject(false) }, Options);
+		FString Filename = UE::Private::GetIntermediateAssetName(Class, *Label);
+		TUniquePtr<FArchive> FileArchive(IFileManager::Get().CreateDebugFileWriter(*Filename));
+		if(FileArchive) // CreateDebugFileWriter fails if the filename is too long, we may have to truncate
+		{
+			FileArchive->Serialize((void*)Result.GetCharArray().GetData(), Result.GetCharArray().Num() - 1);
+		}
+	}
+}	
+
+FAutoConsoleCommand SnapshotClassesCommand(
+	TEXT("snapshotclasses"),
+	*FString::Format(TEXT("Write out a snapshot to the Saved directory of currently loaded classes.\n{0}"), { SnapshotClassesHelp }),
+	FConsoleCommandWithArgsDelegate::CreateStatic(&SnapshotClasses),
+	ECVF_Default
+);
+
+FString GetIntermediateAssetName(const UObject* RootObject, const TCHAR* Prefix)
 {
 	return
 		FPaths::ProjectSavedDir() +
 		FString(TEXT("Temp/")) + Prefix +
-		Package->GetPathName() +
+		RootObject->GetPathName() +
 		FString(TEXT("_snap.json"));
 }
 }
