Add a fast reinstancing path to verse compilation, disabled by default (Verse.FastReinstanceUnchangedClasses). This fast reinstancing path works by just immediately setting the ClassPrivate reference on existing instances, which is unsafe for some types which maintain backpointers to either their CDO (which has been regenerated) or sparse class data (which also may have been regenerated). Types that expect construction to occurr after every verse compilation job will also need to adapt to having construction skipped.
#rnx
#Synced-CL 51225671

#Synced-CL 51225671

diff --git a/Engine/Source/Editor/Kismet/Private/BlueprintCompilationManager.cpp b/Engine/Source/Editor/Kismet/Private/BlueprintCompilationManager.cpp
--- a/Engine/Source/Editor/Kismet/Private/BlueprintCompilationManager.cpp
+++ b/Engine/Source/Editor/Kismet/Private/BlueprintCompilationManager.cpp
@@ -3,6 +3,7 @@
 #include "BlueprintCompilationManager.h"
 
 #include "Async/Async.h"
+#include "Blueprints/BulkNewClassReinstancer.h"
 #include "Kismet2/BlueprintEditorUtils.h"
 #include "BlueprintClassFlagUtils.h"
 #include "BlueprintCompilerExtension.h"
@@ -253,12 +254,21 @@ void ReparentHierarchiesWrapper(const TMap<UClass*, UClass*>& OldToNewMap)
 	FBlueprintCompilationManagerImpl::ReparentHierarchies(OldToNewMap, EReparentClassOptions::ReplaceReferencesToOldClasses);
 }
 
+namespace UE::Private
+{
+	TUniquePtr<IClassReinstancer> CreateBulkNewClassReinstancer(TArray<UClass*>& InOutClasses)
+	{
+		return MakeUnique<FBulkNewClassReinstancer>(InOutClasses);
+	}
+}
+
 FBlueprintCompilationManagerImpl::FBlueprintCompilationManagerImpl()
 {
 	FBlueprintSupport::SetFlushReinstancingQueueFPtr(&FlushReinstancingQueueImplWrapper);
 	FBlueprintSupport::SetFlushCompilationQueueFPtr(&FlushCompilationQueueImplWrapper);
 	FBlueprintSupport::SetClassReparentingFPtr(&ReparentHierarchiesWrapper);
 	bGeneratedClassLayoutReady = true;
+	UE::Private::GCreateClassReinstancerFPtr = &UE::Private::CreateBulkNewClassReinstancer;
 }
 
 FBlueprintCompilationManagerImpl::~FBlueprintCompilationManagerImpl() 
@@ -266,6 +276,8 @@ FBlueprintCompilationManagerImpl::~FBlueprintCompilationManagerImpl()
 	FBlueprintSupport::SetFlushReinstancingQueueFPtr(nullptr);
 	FBlueprintSupport::SetFlushCompilationQueueFPtr(nullptr);
 	FBlueprintSupport::SetClassReparentingFPtr(nullptr);
+	
+	UE::Private::GCreateClassReinstancerFPtr = nullptr;
 }
 
 void FBlueprintCompilationManagerImpl::AddReferencedObjects(FReferenceCollector& Collector)
diff --git a/Engine/Source/Editor/Kismet/Private/Blueprints/BulkNewClassReinstancer.cpp b/Engine/Source/Editor/Kismet/Private/Blueprints/BulkNewClassReinstancer.cpp
new file mode 100644
--- /dev/null
+++ b/Engine/Source/Editor/Kismet/Private/Blueprints/BulkNewClassReinstancer.cpp
@@ -0,0 +1,134 @@
+// Copyright Epic Games, Inc. All Rights Reserved.
+
+#include "Blueprints/BulkNewClassReinstancer.h"
+#include "BlueprintCompilationManager.h"
+#include "HAL/IConsoleManager.h"
+#include "Kismet2/KismetReinstanceUtilities.h"
+#include "Serialization/ArchiveReplaceObjectRef.h"
+#include "UObject/Package.h"
+#include "JsonObjectGraph/Stringify.h"
+#include "Logging/StructuredLog.h"
+
+namespace UE::Private
+{
+
+// Use this to debug unexpectedly expensive ReparentHierarchies calls due to high numbers of class changes:
+bool bPrintFastReinstancingMismatches = false;
+static FAutoConsoleVariableRef CVarPrintFastReinstancingMismatches(
+	TEXT("BP.bPrintFastReinstancingMismatches"), bPrintFastReinstancingMismatches,
+	TEXT("If true we log out the before and after stable representations of UClasses in a reinstancer operation"),
+	ECVF_Default);
+
+FBulkNewClassReinstancer::FBulkNewClassReinstancer(TArray<UClass*>& InOutClasses)
+{
+	StoreIdentities(InOutClasses);
+}
+
+FBulkNewClassReinstancer::~FBulkNewClassReinstancer()
+{
+}
+
+void FBulkNewClassReinstancer::DoReinstance(const TMap<UClass*, UClass*>& OldClassToNewClass)
+{
+	if(OldClassToNewClass.Num() == 0)
+	{
+		return;
+	}
+	
+	TMap<UClass*, UClass*> OldToNewFastReinstancing;
+	MapIdenticalClasses(OldClassToNewClass, OldToNewFastReinstancing);
+	ReparentHierarchiesFast(OldToNewFastReinstancing);
+	FBlueprintSupport::ReparentHierarchies(OldClassToNewClass);
+}
+
+void FBulkNewClassReinstancer::ReparentHierarchiesFast(const TMap<UClass*, UClass*>& OldClassToNewClass)
+{
+	const bool bIncludeDerivedClasses = false;
+
+	for(TPair<UClass*, UClass*> OldToNew : OldClassToNewClass)
+	{
+		TArray< UObject* > ObjectsToReplace;
+		GetObjectsOfClass(OldToNew.Key, ObjectsToReplace, bIncludeDerivedClasses, RF_ClassDefaultObject | RF_MirroredGarbage);
+		for (UObject* Object : ObjectsToReplace)
+		{
+			// class layouts match, so we can just set the object back to the new class - the draw back here
+			// is that instances may have cached information about the UClass or related objects and will be
+			// distressed when/if it changes..
+			Object->SetClass(OldToNew.Value); 
+		}
+	}
+}
+
+void FBulkNewClassReinstancer::MapIdenticalClasses(const TMap<UClass*, UClass*>& OldClassToNewClass, TMap<UClass*, UClass*>& OutOldToNewFastReinstancing)
+{
+	TRACE_CPUPROFILER_EVENT_SCOPE(MapIdenticalClasses);
+
+	// performance note: both large sets of class (inputs) and modest performance of CalculateClassIdentity are problems
+	for(const TPair<UClass*, UClass*>& Pair : OldClassToNewClass)
+	{
+		const UClass* NewClass = Pair.Value;
+		if(!NewClass)
+		{
+			continue;
+		}
+		
+		FUtf8String OriginalId;
+		if(UE::Private::bPrintFastReinstancingMismatches)
+		{
+			// Assume NewClass is now using the original path, use it to look up the original id:
+			if(FUtf8String* OriginalEntry = DebugClassIdentity.Find(NewClass->GetPathName()))
+			{
+				OriginalId = *OriginalEntry;
+			}
+			else
+			{
+				OriginalId = "Class has been unexpected mapped to a new class with a different name";
+			}
+		}
+
+		FBlake3Hash NewId = FBulkNewClassReinstancer::CalculateClassIdentity(NewClass);
+		const FBlake3Hash* OldHash = OriginalClassIdentities.Find(Pair.Key); 
+		if(OldHash)
+		{
+			if( *OldHash == NewId)
+			{
+				OutOldToNewFastReinstancing.Add(Pair.Key, Pair.Value);
+			}
+			else if(UE::Private::bPrintFastReinstancingMismatches)
+			{
+				UE_LOGFMT(LogTemp, Warning, "original: {OriginalId} new: {NewId}", OriginalId, DebugClassIdentity.FindChecked(NewClass->GetPathName()));
+			}
+		}
+	}
+}
+
+void FBulkNewClassReinstancer::StoreIdentities(TArrayView<UClass*> InClasses)
+{
+	TRACE_CPUPROFILER_EVENT_SCOPE(StoreIdentities);
+
+	// performance note: class identity calculation could be optimized, but reducing the input set will be step function improvement
+	for(UClass* Class : InClasses)
+	{
+		OriginalClassIdentities.Add(Class, FBulkNewClassReinstancer::CalculateClassIdentity(Class));
+	}
+}
+
+FBlake3Hash FBulkNewClassReinstancer::CalculateClassIdentity(const UClass* ForClass)
+{
+	// Stringify is a bit slow, but it is stable and quite complete..
+	FUtf8String ClassIdString = UE::JsonObjectGraph::Stringify(
+		{ForClass, ForClass->GetDefaultObject(false)}, EJsonStringifyFlags::DisableDeltaEncoding);
+
+	FBlake3 Hash;
+	Hash.Update(ClassIdString.GetCharArray().GetData(), ClassIdString.GetCharArray().Num());
+	FBlake3Hash Result = Hash.Finalize();
+	
+	if(UE::Private::bPrintFastReinstancingMismatches)
+	{
+		DebugClassIdentity.Add(ForClass->GetPathName(), ClassIdString);
+	}
+
+	return Result;
+}
+
+}
diff --git a/Engine/Source/Editor/Kismet/Private/Blueprints/BulkNewClassReinstancer.h b/Engine/Source/Editor/Kismet/Private/Blueprints/BulkNewClassReinstancer.h
new file mode 100644
--- /dev/null
+++ b/Engine/Source/Editor/Kismet/Private/Blueprints/BulkNewClassReinstancer.h
@@ -0,0 +1,34 @@
+// Copyright Epic Games, Inc. All Rights Reserved.
+
+#pragma once
+
+#include "UObject/IClassReinstancer.h"
+#include "Containers/ArrayView.h"
+#include "Containers/Utf8String.h"
+#include "Hash/Blake3.h"
+#include "Templates/SharedPointer.h"
+
+class FBlueprintCompileReinstancer;
+class FKismetCompilerContext;
+
+namespace UE::Private
+{
+
+struct FBulkNewClassReinstancer : public IClassReinstancer
+{
+	FBulkNewClassReinstancer(TArray<UClass*>& InOutClasses);
+	virtual ~FBulkNewClassReinstancer();
+	virtual void DoReinstance(const TMap<UClass*, UClass*>& OldClassToNewClass) override;
+private:
+	static void ReparentHierarchiesFast(const TMap<UClass*, UClass*>& OldClassToNewClass);
+	void MapIdenticalClasses(const TMap<UClass*, UClass*>& OldClassToNewClass, TMap<UClass*, UClass*>& OutOldToNewFastReinstancing);
+	void StoreIdentities(TArrayView<UClass*> InClasses);
+	FBlake3Hash CalculateClassIdentity(const UClass* ForClass);
+
+	TMap<UClass*, FBlake3Hash> OriginalClassIdentities;
+	// This is typically not populated, and is only used when BP.bPrintFastReinstancingMismatches
+	// is set:
+	TMap<FString, FUtf8String> DebugClassIdentity;
+};
+
+}
\ No newline at end of file
diff --git a/Engine/Source/Runtime/CoreUObject/Internal/UObject/IClassReinstancer.h b/Engine/Source/Runtime/CoreUObject/Internal/UObject/IClassReinstancer.h
new file mode 100644
--- /dev/null
+++ b/Engine/Source/Runtime/CoreUObject/Internal/UObject/IClassReinstancer.h
@@ -0,0 +1,18 @@
+// Copyright Epic Games, Inc. All Rights Reserved.
+
+#include "Templates/UniquePtr.h"
+#include "Containers/Map.h"
+
+namespace UE::Private
+{
+	struct IClassReinstancer
+	{
+		virtual ~IClassReinstancer() = default;
+		virtual void DoReinstance(const TMap<UClass*, UClass*>& OldClassToNewClass) = 0;
+	};
+
+	COREUOBJECT_API TUniquePtr<IClassReinstancer> CreateClassReinstancer(TArray<UClass*>& InOutClasses);
+	
+	typedef TUniquePtr<IClassReinstancer> (*FReinstanceClassesPrepassFPtr)(TArray<UClass*>&);
+	extern COREUOBJECT_API FReinstanceClassesPrepassFPtr GCreateClassReinstancerFPtr;
+}
\ No newline at end of file
diff --git a/Engine/Source/Runtime/CoreUObject/Private/UObject/IClassReinstancer.cpp b/Engine/Source/Runtime/CoreUObject/Private/UObject/IClassReinstancer.cpp
new file mode 100644
--- /dev/null
+++ b/Engine/Source/Runtime/CoreUObject/Private/UObject/IClassReinstancer.cpp
@@ -0,0 +1,17 @@
+// Copyright Epic Games, Inc. All Rights Reserved.
+
+#include "UObject/IClassReinstancer.h"
+
+namespace UE::Private
+{
+	COREUOBJECT_API FReinstanceClassesPrepassFPtr GCreateClassReinstancerFPtr = nullptr;
+}
+
+TUniquePtr<UE::Private::IClassReinstancer> UE::Private::CreateClassReinstancer(TArray<UClass*>& InOutClasses)
+{
+	if (GCreateClassReinstancerFPtr)
+	{
+		return (*GCreateClassReinstancerFPtr)(InOutClasses);
+	}
+	return {};
+}
\ No newline at end of file
diff --git a/Engine/Source/Runtime/CoreUObject/Public/UObject/UObjectBase.h b/Engine/Source/Runtime/CoreUObject/Public/UObject/UObjectBase.h
--- a/Engine/Source/Runtime/CoreUObject/Public/UObject/UObjectBase.h
+++ b/Engine/Source/Runtime/CoreUObject/Public/UObject/UObjectBase.h
@@ -58,6 +58,8 @@ namespace UE::CoreUObject
 
 DECLARE_DWORD_COUNTER_STAT_EXTERN(TEXT("STAT_UObjectsStatGroupTester"), STAT_UObjectsStatGroupTester, STATGROUP_UObjects, COREUOBJECT_API);
 
+namespace UE::Private { struct FBulkNewClassReinstancer; }
+
 /** 
  * Low level implementation of UObject, should not be used directly in game code 
  */
@@ -464,6 +466,7 @@ private:
 	ObjectPtr_Private::TNonAccessTrackedObjectPtr<UObject>						OuterPrivate;
 	
 	friend class FBlueprintCompileReinstancer;
+	friend struct UE::Private::FBulkNewClassReinstancer;
 	friend class FContextObjectManager;
 	friend class FDeferredRegistry;
 	friend class FCompiledInObjectRegistry;
