Propagate ModeTools correctly into Blueprint/Sub-object editor viewport path...

#Synced-CL 51530715

diff --git a/Engine/Source/Editor/Kismet/Private/BlueprintEditor.cpp b/Engine/Source/Editor/Kismet/Private/BlueprintEditor.cpp
--- a/Engine/Source/Editor/Kismet/Private/BlueprintEditor.cpp
+++ b/Engine/Source/Editor/Kismet/Private/BlueprintEditor.cpp
@@ -3107,7 +3107,8 @@ void FBlueprintEditor::CreateSubobjectEditors()
 	LLM_SCOPE_BYNAME(TEXT("BPCreateSubobjectEditorViewportClient"));
 
 	SubobjectViewport = SAssignNew(SubobjectViewport, SSCSEditorViewport)
-		.BlueprintEditor(SharedThis(this));
+		.BlueprintEditor(SharedThis(this))
+		.ModeTools(EditorModeManager);
 
 	if (USubobjectDataSubsystem* SubobjectDataSubsystem = USubobjectDataSubsystem::Get())
 	{
diff --git a/Engine/Source/Editor/Kismet/Private/SCSEditorViewportClient.cpp b/Engine/Source/Editor/Kismet/Private/SCSEditorViewportClient.cpp
--- a/Engine/Source/Editor/Kismet/Private/SCSEditorViewportClient.cpp
+++ b/Engine/Source/Editor/Kismet/Private/SCSEditorViewportClient.cpp
@@ -163,8 +163,8 @@ namespace
 /////////////////////////////////////////////////////////////////////////
 // FSCSEditorViewportClient
 
-FSCSEditorViewportClient::FSCSEditorViewportClient(TWeakPtr<FBlueprintEditor>& InBlueprintEditorPtr, FPreviewScene* InPreviewScene, const TSharedRef<SSCSEditorViewport>& InSCSEditorViewport)
-	: FEditorViewportClient(nullptr, InPreviewScene, StaticCastSharedRef<SEditorViewport>(InSCSEditorViewport))
+FSCSEditorViewportClient::FSCSEditorViewportClient(TWeakPtr<FBlueprintEditor>& InBlueprintEditorPtr, FPreviewScene* InPreviewScene, const TSharedRef<SSCSEditorViewport>& InSCSEditorViewport, FEditorModeTools* InModeTools)
+	: FEditorViewportClient(InModeTools, InPreviewScene, StaticCastSharedRef<SEditorViewport>(InSCSEditorViewport))
 	, BlueprintEditorPtr(InBlueprintEditorPtr)
 	, PreviewActorBounds(ForceInitToZero)
 	, bIsManipulating(false)
diff --git a/Engine/Source/Editor/Kismet/Private/SCSEditorViewportClient.h b/Engine/Source/Editor/Kismet/Private/SCSEditorViewportClient.h
--- a/Engine/Source/Editor/Kismet/Private/SCSEditorViewportClient.h
+++ b/Engine/Source/Editor/Kismet/Private/SCSEditorViewportClient.h
@@ -39,8 +39,9 @@ public:
 	 *
 	 * @param InBlueprintEditorPtr A weak reference to the Blueprint Editor context.
 	 * @param InPreviewScene The preview scene to use.
+	 * @param InModeTools Optional pre-configured mode tools (if nullptr, a new one will be created).
 	 */
-	FSCSEditorViewportClient(TWeakPtr<class FBlueprintEditor>& InBlueprintEditorPtr, FPreviewScene* InPreviewScene, const TSharedRef<SSCSEditorViewport>& InSCSEditorViewport);
+	FSCSEditorViewportClient(TWeakPtr<class FBlueprintEditor>& InBlueprintEditorPtr, FPreviewScene* InPreviewScene, const TSharedRef<SSCSEditorViewport>& InSCSEditorViewport, FEditorModeTools* InModeTools = nullptr);
 
 	/**
 	 * Destructor.
diff --git a/Engine/Source/Editor/Kismet/Private/SSCSEditorViewport.cpp b/Engine/Source/Editor/Kismet/Private/SSCSEditorViewport.cpp
--- a/Engine/Source/Editor/Kismet/Private/SSCSEditorViewport.cpp
+++ b/Engine/Source/Editor/Kismet/Private/SSCSEditorViewport.cpp
@@ -99,6 +99,9 @@ void SSCSEditorViewport::Construct(const FArguments& InArgs)
 	// Save off the Blueprint editor reference, we'll need this later
 	BlueprintEditorPtr = InArgs._BlueprintEditor;
 
+	// Store the pre-configured mode tools (if provided)
+	ModeTools = InArgs._ModeTools;
+
 	SEditorViewport::Construct( SEditorViewport::FArguments() );
 
 	// Restore last used feature level
@@ -155,7 +158,7 @@ TSharedRef<FEditorViewportClient> SSCSEditorViewport::MakeEditorViewportClient()
 	FPreviewScene* PreviewScene = BlueprintEditorPtr.Pin()->GetPreviewScene();
 
 	// Construct a new viewport client instance.
-	ViewportClient = MakeShareable(new FSCSEditorViewportClient(BlueprintEditorPtr, PreviewScene, SharedThis(this)));
+	ViewportClient = MakeShareable(new FSCSEditorViewportClient(BlueprintEditorPtr, PreviewScene, SharedThis(this), ModeTools.Get()));
 	ViewportClient->SetRealtime(true);
 	ViewportClient->bSetListenerPosition = false;
 	ViewportClient->VisibilityDelegate.BindSP(this, &SSCSEditorViewport::IsVisible);
diff --git a/Engine/Source/Editor/Kismet/Private/SSCSEditorViewport.h b/Engine/Source/Editor/Kismet/Private/SSCSEditorViewport.h
--- a/Engine/Source/Editor/Kismet/Private/SSCSEditorViewport.h
+++ b/Engine/Source/Editor/Kismet/Private/SSCSEditorViewport.h
@@ -21,6 +21,7 @@ class SSCSEditorViewport : public SEditorViewport
 public:
 	SLATE_BEGIN_ARGS(SSCSEditorViewport){}
 		SLATE_ARGUMENT(TWeakPtr<class FBlueprintEditor>, BlueprintEditor)
+		SLATE_ARGUMENT(TSharedPtr<class FEditorModeTools>, ModeTools)
 	SLATE_END_ARGS()
 
 	/**
@@ -119,4 +120,7 @@ private:
 
 	/** Handle to the registered OnPreviewFeatureLevelChanged delegate. */
 	FDelegateHandle PreviewFeatureLevelChangedHandle;
+
+	/** Pre-configured mode tools for the viewport (optional) */
+	TSharedPtr<class FEditorModeTools> ModeTools;
 };
