Convert the Blueprint Editor to use ITF Gizmos & Interactions

#Synced-CL 51530715

diff --git a/Engine/Source/Editor/Kismet/Kismet.Build.cs b/Engine/Source/Editor/Kismet/Kismet.Build.cs
--- a/Engine/Source/Editor/Kismet/Kismet.Build.cs
+++ b/Engine/Source/Editor/Kismet/Kismet.Build.cs
@@ -62,7 +62,9 @@ public class Kismet : ModuleRules
 				"SubobjectDataInterface",
 				"ToolWidgets",
 				"TraceLog",
-				"GameplayTags"
+				"GameplayTags",
+				"InteractiveToolsFramework",
+				"TypedElementRuntime"
 			}
 			);
 
diff --git a/Engine/Source/Editor/Kismet/Private/BlueprintEditor.cpp b/Engine/Source/Editor/Kismet/Private/BlueprintEditor.cpp
--- a/Engine/Source/Editor/Kismet/Private/BlueprintEditor.cpp
+++ b/Engine/Source/Editor/Kismet/Private/BlueprintEditor.cpp
@@ -183,6 +183,7 @@
 #include "BlueprintEditorSCSEditorUICustomization.h"
 #include "EditorModeManager.h"
 #include "Algo/MinElement.h"
+#include "AssetEditorModeManager.h"
 #include "Editor/EditorEngine.h"
 #include "EditorViewportSelectabilityBridge.h"
 #include "IBlueprintSelectionNotifier.h"
@@ -1324,7 +1325,14 @@ void FBlueprintEditor::OnSelectionUpdated(const TArray<TSharedPtr<FSubobjectEdit
 	TArray<IBlueprintSelectionNotifier*> Notifiers = IModularFeatures::Get().GetModularFeatureImplementations<IBlueprintSelectionNotifier>(
 		IBlueprintSelectionNotifier::GetModularFeatureName());
 
-	// Notify each registered implementation
+	UTypedElementSelectionSet* SelectionSet = GetEditorModeManager().GetEditorSelectionSet();
+	const bool bHasSelectionSet = SelectionSet != nullptr;
+	if (bHasSelectionSet)
+	{
+		SelectionSet->ClearSelection(FTypedElementSelectionOptions());
+	}
+
+	// Notify each registered implementation (including ModeTools SelectionSet)
 	TArray<FTypedElementHandle> SelectedElements;
 	for (const FSubobjectEditorTreeNodePtrType& SelectedNode : SelectedNodes)
 	{
@@ -1337,6 +1345,11 @@ void FBlueprintEditor::OnSelectionUpdated(const TArray<TSharedPtr<FSubobjectEdit
 				if (ElementHandle.IsSet())
 				{
 					SelectedElements.Add(ElementHandle);
+
+					if (bHasSelectionSet)
+					{
+						SelectionSet->SelectElement(ElementHandle, FTypedElementSelectionOptions());
+					}
 				}
 			}
 		}
@@ -2543,8 +2556,8 @@ void FBlueprintEditor::InitBlueprintEditor(
 		{
 			if (FEditorModeTools* const ModeTools = ViewportClient->GetModeTools())
 			{
-			// Opt out of new ITF gizmos and viewport inputs for now
-				ModeTools->SetSupportsViewportITF(false);
+				// Opt in to new ITF gizmos and viewport inputs
+				ModeTools->SetSupportsViewportITF(true);
 			}
 		}
 	}
