Fix reinstancing to always fix references in created instances, some CDO subobjects were getting skipped which left behind pointers to REINST classes
#Synced-CL 51587375

diff --git a/Engine/Source/Editor/Kismet/Private/BlueprintCompilationManager.cpp b/Engine/Source/Editor/Kismet/Private/BlueprintCompilationManager.cpp
--- a/Engine/Source/Editor/Kismet/Private/BlueprintCompilationManager.cpp
+++ b/Engine/Source/Editor/Kismet/Private/BlueprintCompilationManager.cpp
@@ -3350,6 +3350,15 @@ void FBlueprintCompilationManagerImpl::ReinstanceBatch(TArray<FReinstancingJob>&
 		}
 	}
 
+	// Make sure all new archetypes get their cross references fixed up, most of them are already in the set so avoid modifying if possible
+	for (const TPair<UObject*, UObject*>& Pair : OldArchetypeToNewArchetype)
+	{
+		if (!ArchetypeReferencers.Contains(Pair.Value))
+		{
+			ArchetypeReferencers.Add(Pair.Value);
+		}
+	}
+
 	// Data loading may be in flight, lets immediately patch existing redirectors - 
 	// we may want to search the entire graph some day, but that would be expensive:
 	const bool bIsInActualLoadingThread = IsInActualLoadingThread();
