BP compilation hardening

[RN] Fixed crashes in FBlueprintActionFilter::IsFilteredByThis
#Synced-CL 51516599

diff --git a/Engine/Plugins/Runtime/GameplayAbilities/Source/GameplayAbilitiesEditor/Private/K2Node_GameplayCueEvent.cpp b/Engine/Plugins/Runtime/GameplayAbilities/Source/GameplayAbilitiesEditor/Private/K2Node_GameplayCueEvent.cpp
--- a/Engine/Plugins/Runtime/GameplayAbilities/Source/GameplayAbilitiesEditor/Private/K2Node_GameplayCueEvent.cpp
+++ b/Engine/Plugins/Runtime/GameplayAbilities/Source/GameplayAbilitiesEditor/Private/K2Node_GameplayCueEvent.cpp
@@ -52,8 +52,7 @@ bool UK2Node_GameplayCueEvent::IsCompatibleWithGraph(UEdGraph const* TargetGraph
 	bool bIsCompatible = false;
 	if (UBlueprint* Blueprint = FBlueprintEditorUtils::FindBlueprintForGraph(TargetGraph))
 	{
-		check(Blueprint->GeneratedClass != nullptr);
-		bIsCompatible = Blueprint->GeneratedClass->ImplementsInterface(UGameplayCueInterface::StaticClass());
+		bIsCompatible = Blueprint->GeneratedClass && Blueprint->GeneratedClass->ImplementsInterface(UGameplayCueInterface::StaticClass());
 	}	
 	return bIsCompatible && Super::IsCompatibleWithGraph(TargetGraph);
 }
diff --git a/Engine/Source/Editor/BlueprintGraph/Private/K2Node_SpawnActorFromClass.cpp b/Engine/Source/Editor/BlueprintGraph/Private/K2Node_SpawnActorFromClass.cpp
--- a/Engine/Source/Editor/BlueprintGraph/Private/K2Node_SpawnActorFromClass.cpp
+++ b/Engine/Source/Editor/BlueprintGraph/Private/K2Node_SpawnActorFromClass.cpp
@@ -349,7 +349,10 @@ FText UK2Node_SpawnActorFromClass::GetNodeTitle(ENodeTitleType::Type TitleType)
 bool UK2Node_SpawnActorFromClass::IsCompatibleWithGraph(const UEdGraph* TargetGraph) const 
 {
 	UBlueprint* Blueprint = FBlueprintEditorUtils::FindBlueprintForGraph(TargetGraph);
-	return Super::IsCompatibleWithGraph(TargetGraph) && (!Blueprint || (FBlueprintEditorUtils::FindUserConstructionScript(Blueprint) != TargetGraph && Blueprint->GeneratedClass->GetDefaultObject()->ImplementsGetWorld()));
+	const UClass* Class = Blueprint ? Blueprint->GeneratedClass : nullptr;
+	return Super::IsCompatibleWithGraph(TargetGraph) &&
+		(	!Blueprint || 
+			(FBlueprintEditorUtils::FindUserConstructionScript(Blueprint) != TargetGraph && Class && Class->GetDefaultObject()->ImplementsGetWorld()));
 }
 
 void UK2Node_SpawnActorFromClass::PostPlacedNewNode()
diff --git a/Engine/Source/Editor/Kismet/Private/BlueprintCompilationManager.cpp b/Engine/Source/Editor/Kismet/Private/BlueprintCompilationManager.cpp
--- a/Engine/Source/Editor/Kismet/Private/BlueprintCompilationManager.cpp
+++ b/Engine/Source/Editor/Kismet/Private/BlueprintCompilationManager.cpp
@@ -84,7 +84,7 @@ extern UNREALED_API UUnrealEdEngine* GUnrealEd;
 	STAGE III: SORT
 	STAGE IV: SET TEMPORARY BLUEPRINT FLAGS
 	STAGE V: VALIDATE
-	STAGE VI: PURGE (LOAD ONLY)
+	STAGE VI: MISC FIXUP (LOAD ONLY)
 	STAGE VII: DISCARD SKELETON CDO
 	STAGE VIII: RECOMPILE SKELETON
 	STAGE IX: RECONSTRUCT NODES, REPLACE DEPRECATED NODES (LOAD ONLY)
@@ -1023,13 +1023,12 @@ void FBlueprintCompilationManagerImpl::FlushCompilationQueueImpl(bool bSuppressB
 
 		SlowTask.EnterProgressFrame();
 
-		// STAGE VI: Purge null graphs, misc. data fixup
+		// STAGE VI: misc. data fixup
 		for (FCompilerData& CompilerData : CurrentlyCompilingBPs)
 		{
 			UBlueprint* BP = CompilerData.BP;
 			if(BP->bIsRegeneratingOnLoad)
 			{
-				FBlueprintEditorUtils::PurgeNullGraphs(BP);
 				BP->ConformNativeComponents();
 				if (FLinkerLoad* Linker = BP->GetLinker())
 				{
@@ -3416,6 +3415,8 @@ UClass* FBlueprintCompilationManagerImpl::FastGenerateSkeletonClass(UBlueprint*
 			}
 		}
 	}
+	
+	FBlueprintEditorUtils::PurgeNullGraphs(BP);
 
 	UBlueprintGeneratedClass* Ret = nullptr;
 	UBlueprintGeneratedClass* OriginalNewClass = CompilerContext.NewClass;
