Avoid reporting cook dependencies for blueprints in the error state, this means they will always recook during incremental cook. This is desirable because they may be in the error state because there is data that is missing (which they cannot report as a dependency)

#rnx
#Synced-CL 51516599

diff --git a/Engine/Source/Editor/Kismet/Private/Blueprints/BlueprintDependencies.cpp b/Engine/Source/Editor/Kismet/Private/Blueprints/BlueprintDependencies.cpp
--- a/Engine/Source/Editor/Kismet/Private/Blueprints/BlueprintDependencies.cpp
+++ b/Engine/Source/Editor/Kismet/Private/Blueprints/BlueprintDependencies.cpp
@@ -835,10 +835,13 @@ UE::Cook::FCookDependency UE::Private::BlueprintDependencies::RecordCookDependen
 {
 	if(!BP->ParentClass || 
 		!BP->GeneratedClass ||
-		!BP->GeneratedClass->IsChildOf(UObject::StaticClass()))
+		!BP->GeneratedClass->IsChildOf(UObject::StaticClass()) ||
+		(BP->Status == BS_Error || BP->Status == BS_Unknown))
 	{
 		// blueprints without classes - malformed or some
-		// kind of utility blueprint, ignore them
+		// kind of utility blueprint or in the error state, ignore them
+		// blueprints in error state may not be able to report all of their
+		// dependencies (e.g. because they've gone missing)
 		return UE::Cook::FCookDependency::Function(
 			UE_COOK_DEPENDENCY_FUNCTION_CALL(NotCacheable), FCbFieldIterator());
 	}
