[BP Editor] Fixed selection-related crash when closing an Actor blueprint where you've selected one or more objects in it's viewport

#Synced-CL 51660293

diff --git a/Engine/Source/Editor/Kismet/Private/BlueprintEditor.cpp b/Engine/Source/Editor/Kismet/Private/BlueprintEditor.cpp
--- a/Engine/Source/Editor/Kismet/Private/BlueprintEditor.cpp
+++ b/Engine/Source/Editor/Kismet/Private/BlueprintEditor.cpp
@@ -805,6 +805,13 @@ bool FBlueprintEditor::OnRequestClose(EAssetEditorCloseReason InCloseReason)
 
 void FBlueprintEditor::OnClose()
 {
+	UTypedElementSelectionSet* SelectionSet = GetEditorModeManager().GetEditorSelectionSet();
+	const bool bHasSelectionSet = SelectionSet != nullptr;
+	if (bHasSelectionSet)
+	{
+		SelectionSet->ClearSelection(FTypedElementSelectionOptions());
+	}
+
 	// Also close the Find Results tab if we're not in full edit mode.
 	TSharedPtr<SDockTab> FindResultsTab = TabManager->FindExistingLiveTab(FBlueprintEditorTabs::FindResultsID);
 	if (FindResultsTab.IsValid() && !IsInAScriptingMode())
@@ -1144,6 +1151,13 @@ void FBlueprintEditor::ClearSelectionStateFor(FName SelectionOwner)
 	{	
 		if ( SubobjectEditor.IsValid() )
 		{
+			UTypedElementSelectionSet* SelectionSet = GetEditorModeManager().GetEditorSelectionSet();
+			const bool bHasSelectionSet = SelectionSet != nullptr;
+			if (bHasSelectionSet)
+			{
+				SelectionSet->ClearSelection(FTypedElementSelectionOptions());
+			}
+
 			SubobjectEditor->ClearSelection();
 		}
 	}
