BlueprintEditor creates it's own ModeManager, vs. using the default/global FEditorModeTools.

#Synced-CL 51660293

diff --git a/Engine/Source/Editor/Kismet/Private/BlueprintEditor.cpp b/Engine/Source/Editor/Kismet/Private/BlueprintEditor.cpp
--- a/Engine/Source/Editor/Kismet/Private/BlueprintEditor.cpp
+++ b/Engine/Source/Editor/Kismet/Private/BlueprintEditor.cpp
@@ -798,6 +798,11 @@ bool FBlueprintEditor::IsInAScriptingMode() const
 	return IsModeCurrent(FBlueprintEditorApplicationModes::StandardBlueprintEditorMode) || IsModeCurrent(FBlueprintEditorApplicationModes::BlueprintMacroMode);
 }
 
+void FBlueprintEditor::CreateEditorModeManager()
+{
+	EditorModeManager = MakeShared<FAssetEditorModeManager>();
+}
+
 bool FBlueprintEditor::OnRequestClose(EAssetEditorCloseReason InCloseReason)
 {
 	return FWorkflowCentricApplication::OnRequestClose(InCloseReason);
diff --git a/Engine/Source/Editor/Kismet/Public/BlueprintEditor.h b/Engine/Source/Editor/Kismet/Public/BlueprintEditor.h
--- a/Engine/Source/Editor/Kismet/Public/BlueprintEditor.h
+++ b/Engine/Source/Editor/Kismet/Public/BlueprintEditor.h
@@ -256,7 +256,8 @@ public:
 	UE_API virtual void OnClose() override;
 	UE_API virtual void AddViewportOverlayWidget(TSharedRef<SWidget> InViewportOverlayWidget, int32 ZOrder = INDEX_NONE) override;
 	UE_API virtual void RemoveViewportOverlayWidget(TSharedRef<SWidget> InViewportOverlayWidget) override;
-	// End of FAssetEditorToolkit 
+	UE_API virtual void CreateEditorModeManager() override;
+	// End of FAssetEditorToolkit
 
 	//~ Begin IToolkit Interface
 	UE_API virtual FName GetToolkitContextFName() const override;
