[ITF] Disable viewport ITF interaction in the blueprint editor while issues are diagnosed.

#rnx
#Synced-CL 51825263

diff --git a/Engine/Source/Editor/Kismet/Private/BlueprintEditor.cpp b/Engine/Source/Editor/Kismet/Private/BlueprintEditor.cpp
--- a/Engine/Source/Editor/Kismet/Private/BlueprintEditor.cpp
+++ b/Engine/Source/Editor/Kismet/Private/BlueprintEditor.cpp
@@ -2575,8 +2575,8 @@ void FBlueprintEditor::InitBlueprintEditor(
 		{
 			if (FEditorModeTools* const ModeTools = ViewportClient->GetModeTools())
 			{
-				// Opt in to new ITF gizmos and viewport inputs
-				ModeTools->SetSupportsViewportITF(true);
+				// Opt out of new ITF gizmos and viewport inputs
+				ModeTools->SetSupportsViewportITF(false);
 			}
 		}
 	}
