[Blueprint] Support for Additional categories that always appear in the MyBlueprint panel and category dropdown, even when empty.
The categories are added only if the Blueprint is a child of the supported class. An empty action can be added when there is no action for the category.

#Synced-CL 51823286

diff --git a/Engine/Source/Editor/BehaviorTreeEditor/Private/SBehaviorTreeBlackboardView.h b/Engine/Source/Editor/BehaviorTreeEditor/Private/SBehaviorTreeBlackboardView.h
--- a/Engine/Source/Editor/BehaviorTreeEditor/Private/SBehaviorTreeBlackboardView.h
+++ b/Engine/Source/Editor/BehaviorTreeEditor/Private/SBehaviorTreeBlackboardView.h
@@ -67,6 +67,11 @@ public:
 
 	/** Temp flag for new items */
 	bool bIsNew;
+
+	virtual bool CanBeRenamed() const override
+	{
+		return true;
+	}
 };
 
 /** Displays blackboard entries */
diff --git a/Engine/Source/Editor/Kismet/Private/BlueprintDetailsCustomization.cpp b/Engine/Source/Editor/Kismet/Private/BlueprintDetailsCustomization.cpp
--- a/Engine/Source/Editor/Kismet/Private/BlueprintDetailsCustomization.cpp
+++ b/Engine/Source/Editor/Kismet/Private/BlueprintDetailsCustomization.cpp
@@ -1730,14 +1730,15 @@ void FBlueprintVarActionDetails::PopulateCategories(SMyBlueprint* MyBlueprint, T
 
 	CategorySource.Reset();
 	CategorySource.Add(MakeShared<FText>(UEdGraphSchema_K2::VR_DefaultCategory));
-	for (const FAdditionalBlueprintCategory& AdditionalBlueprintCategory : GetDefault<UBlueprintEditorSettings>()->AdditionalBlueprintCategories)
+	FBlueprintEditorModule& BlueprintEditorModule = FModuleManager::LoadModuleChecked<FBlueprintEditorModule>("Kismet");
+	for (const UE::BlueprintEditor::FAdditionalCategory& Category : BlueprintEditorModule.GetAdditionalBlueprintCategories(Blueprint))
 	{
-		if (!AdditionalBlueprintCategory.Name.IsEmpty() && (AdditionalBlueprintCategory.ClassFilter.IsNull() ||
-			(Blueprint->ParentClass && Blueprint->ParentClass->IsChildOf(AdditionalBlueprintCategory.ClassFilter.TryLoadClass<UObject>()))))
+		if (IsNewCategorySource(Category.Category))
 		{
-			CategorySource.Add(MakeShared<FText>(AdditionalBlueprintCategory.Name));
+			CategorySource.Add(MakeShared<FText>(Category.Category));
 		}
 	}
+
 	for (const FName& VariableName : VisibleVariables)
 	{
 		FText Category = FBlueprintEditorUtils::GetBlueprintVariableCategory(Blueprint, VariableName, nullptr);
diff --git a/Engine/Source/Editor/Kismet/Private/BlueprintEditorModule.cpp b/Engine/Source/Editor/Kismet/Private/BlueprintEditorModule.cpp
--- a/Engine/Source/Editor/Kismet/Private/BlueprintEditorModule.cpp
+++ b/Engine/Source/Editor/Kismet/Private/BlueprintEditorModule.cpp
@@ -4,6 +4,7 @@
 #include "BlueprintEditorModule.h"
 
 #include "BlueprintDebugger.h"
+#include "BlueprintEditorSettings.h"
 #include "BlueprintEditor.h"
 #include "BlueprintEditorSCSEditorUICustomization.h"
 #include "BlueprintGraphPanelPinFactory.h"
@@ -306,6 +307,92 @@ void FBlueprintEditorModule::ShutdownModule()
 	FBlueprintNamespaceRegistry::Get().Shutdown();
 }
 
+FBlueprintEditorModule::FAdditionalCategoryHandle FBlueprintEditorModule::RegisterAdditionalBlueprintCategory(UE::BlueprintEditor::FAdditionalCategory InCategory)
+{
+	const uint32 RawHandle = NextAdditionalCategoryHandle++;
+	const FAdditionalCategoryHandle Handle(RawHandle);
+	RegisteredAdditionalCategories.Emplace(Handle, MoveTemp(InCategory));
+	return Handle;
+}
+
+void FBlueprintEditorModule::UnregisterAdditionalBlueprintCategory(FAdditionalCategoryHandle Handle)
+{
+	if (Handle.IsValid())
+	{
+		RegisteredAdditionalCategories.RemoveAllSwap(
+			[Handle](const TTuple<FAdditionalCategoryHandle, UE::BlueprintEditor::FAdditionalCategory>& Category)
+			{
+				return Category.Get<0>() == Handle;
+			});
+	}
+}
+
+TArray<UE::BlueprintEditor::FAdditionalCategory> FBlueprintEditorModule::GetAdditionalBlueprintCategories() const
+{
+	TArray<UE::BlueprintEditor::FAdditionalCategory> Result;
+
+	for (const TTuple<FAdditionalCategoryHandle, UE::BlueprintEditor::FAdditionalCategory>& Category : RegisteredAdditionalCategories)
+	{
+		Result.Add(Category.Get<1>());
+	}
+
+	// Merge entries from UBlueprintEditorSettings
+	for (const FAdditionalBlueprintCategory& Settings : GetDefault<UBlueprintEditorSettings>()->AdditionalBlueprintCategories)
+	{
+		if (!Settings.Name.IsEmpty())
+		{
+			UE::BlueprintEditor::FAdditionalCategory Converted;
+			// Empty array is for all classes.
+			if (!Settings.ClassFilter.IsNull())
+			{
+				Converted.ClassFilters.Add(Settings.ClassFilter);
+			}
+			Converted.Category = Settings.Name;
+			Result.Add(MoveTemp(Converted));
+		}
+	}
+	return Result;
+}
+
+TArray<UE::BlueprintEditor::FAdditionalCategory> FBlueprintEditorModule::GetAdditionalBlueprintCategories(TNotNull<const UBlueprint*> Blueprint) const
+{
+	TArray<UE::BlueprintEditor::FAdditionalCategory> Result;
+
+	auto Filter = [Blueprint](const FSoftClassPath& Filter)
+		{
+			// If null, then the class is not loaded.
+			//The Blueprint can't be loaded if the parent is not loaded. Then no need to load the class to test IsChildOf.
+			const UClass* FilterClass = Filter.ResolveClass();
+			return Blueprint->ParentClass && FilterClass && Blueprint->ParentClass->IsChildOf(FilterClass);
+		};
+
+	for (const TTuple<FAdditionalCategoryHandle, UE::BlueprintEditor::FAdditionalCategory>& Category : RegisteredAdditionalCategories)
+	{
+		const bool bPassesFilter = Category.Get<1>().ClassFilters.IsEmpty() || Category.Get<1>().ClassFilters.ContainsByPredicate(Filter);
+		if (bPassesFilter)
+		{
+			Result.Add(Category.Get<1>());
+		}
+	}
+
+	// Merge entries from UBlueprintEditorSettings
+	for (const FAdditionalBlueprintCategory& Settings : GetDefault<UBlueprintEditorSettings>()->AdditionalBlueprintCategories)
+	{
+		if (!Settings.Name.IsEmpty())
+		{
+			if (Settings.ClassFilter.IsNull() || Filter(Settings.ClassFilter))
+			{
+				UE::BlueprintEditor::FAdditionalCategory Converted;
+				Converted.Category = Settings.Name;
+				Converted.ClassFilters.Add(Settings.ClassFilter);
+				Result.Add(MoveTemp(Converted));
+			}
+		}
+	}
+
+	return Result;
+}
+
 TSharedRef<IBlueprintEditor> FBlueprintEditorModule::CreateBlueprintEditor(const EToolkitMode::Type Mode, const TSharedPtr< IToolkitHost >& InitToolkitHost, UBlueprint* Blueprint, bool bShouldOpenInDefaultsMode)
 {
 	TArray<UBlueprint*> BlueprintsToEdit = { Blueprint };
diff --git a/Engine/Source/Editor/Kismet/Private/SBlueprintPalette.cpp b/Engine/Source/Editor/Kismet/Private/SBlueprintPalette.cpp
--- a/Engine/Source/Editor/Kismet/Private/SBlueprintPalette.cpp
+++ b/Engine/Source/Editor/Kismet/Private/SBlueprintPalette.cpp
@@ -1577,10 +1577,14 @@ TSharedRef<SWidget> SBlueprintPaletteItem::CreateTextSlotWidget(FCreateWidgetFor
 
 	TSharedPtr<SOverlay> DisplayWidget;
 	TSharedPtr<SInlineEditableTextBlock> EditableTextElement;
+	const FInlineEditableTextBlockStyle* Style = ActionTypeId == FEdGraphSchemaAction_Dummy::StaticGetTypeId()
+		? &FAppStyle::Get().GetWidgetStyle<FInlineEditableTextBlockStyle>("InlineEditableTextBlockHintStyle")
+		: SInlineEditableTextBlock::FArguments()._Style;
 	SAssignNew(DisplayWidget, SOverlay)
 		+SOverlay::Slot()
 		[
 			SAssignNew(EditableTextElement, SInlineEditableTextBlock)
+				.Style(Style)
 				.Text(this, &SBlueprintPaletteItem::GetDisplayText)
 				.HighlightText(InCreateData->HighlightText)
 				.ToolTip(ToolTipWidget)
diff --git a/Engine/Source/Editor/Kismet/Private/SMyBlueprint.cpp b/Engine/Source/Editor/Kismet/Private/SMyBlueprint.cpp
--- a/Engine/Source/Editor/Kismet/Private/SMyBlueprint.cpp
+++ b/Engine/Source/Editor/Kismet/Private/SMyBlueprint.cpp
@@ -247,6 +247,7 @@ public:
 		const int32 SortIndex = GetSortIndex(RootCategory) + Action->GetSectionID();
 
 		SortedActions.Add(SortIndex, Action);
+		SortedActionsByCategoryRoot.Add(RootCategory, Action);
 	}
 
 	void AddAction(TSharedPtr<FEdGraphSchemaAction> Action)
@@ -259,7 +260,7 @@ public:
 	{
 		SortedActions.KeySort(TLess<int32>());
 
-		for (const auto& Iter : SortedActions)
+		for (const TPair<int32, TSharedPtr<FEdGraphSchemaAction>>& Iter : SortedActions)
 		{
 			OutActions.AddAction(Iter.Value);
 		}
@@ -281,6 +282,24 @@ public:
 		}
 	}
 
+	bool HasCategory(const FString& Category, const NodeSectionID::Type SectionID) const
+	{
+		const int32 RootCategoryDelim = Category.Find(TEXT("|"));
+		const FName RootCategory = RootCategoryDelim == INDEX_NONE ? *Category : *Category.Left(RootCategoryDelim);
+		for (TMultiMap<FName, TSharedPtr<FEdGraphSchemaAction>>::TConstKeyIterator It = SortedActionsByCategoryRoot.CreateConstKeyIterator(RootCategory); It; ++It)
+		{
+			if (It.Value()->GetSectionID() == SectionID)
+			{
+				//It contains the root but might not contain the full path.
+				if (RootCategoryDelim == INDEX_NONE || It.Value()->GetCategory().ToString().StartsWith(Category))
+				{
+					return true;
+				}
+			}
+		}
+		return false;
+	}
+
 private:
 	const int32 GetSortIndex(FName Category)
 	{
@@ -308,6 +327,7 @@ private:
 	TArray<FName>& CategorySortIndices;
 	/** Map used to sort Graph actions */
 	TMultiMap<int32, TSharedPtr<FEdGraphSchemaAction>> SortedActions;
+	TMultiMap<FName, TSharedPtr<FEdGraphSchemaAction>> SortedActionsByCategoryRoot;
 };
 
 //////////////////////////////////////////////////////////////////////////
@@ -1777,6 +1797,32 @@ void SMyBlueprint::CollectAllActions(FGraphActionListBuilderBase& OutAllActions)
 		GetLocalVariables(SortList);
 	}
 
+	// Inject a placeholder for each pinned category that is missing.
+	FBlueprintEditorModule& BlueprintEditorModule = FModuleManager::LoadModuleChecked<FBlueprintEditorModule>("Kismet");
+	for (const UE::BlueprintEditor::FAdditionalCategory& Category : BlueprintEditorModule.GetAdditionalBlueprintCategories(BlueprintObj))
+	{
+		if (Category.Placeholders.IsEmpty())
+		{
+			continue;
+		}
+
+		for (const UE::BlueprintEditor::FAdditionalCategory::FPlaceholder& Placeholder : Category.Placeholders)
+		{
+			const FString CategoryDisplayName = FEditorCategoryUtils::GetCategoryDisplayString(Category.Category.ToString());
+			if (!SortList.HasCategory(CategoryDisplayName, Placeholder.SectionID))
+			{
+				TSharedPtr<FEdGraphSchemaAction_Dummy> Action = MakeShareable(new FEdGraphSchemaAction_Dummy(
+					Category.Category,
+					Placeholder.MenuDescription,
+					Placeholder.Tooltip,
+					0,
+					FText::GetEmpty(),
+					static_cast<int32>(Placeholder.SectionID)));
+				SortList.AddAction(CategoryDisplayName, Action);
+			}
+		}
+	}
+
 	// Add all the sorted variables, components, functions, etc...
 	SortList.CleanupCategories();
 	SortList.GetAllActions(OutAllActions);
diff --git a/Engine/Source/Editor/Kismet/Public/BlueprintEditor.h b/Engine/Source/Editor/Kismet/Public/BlueprintEditor.h
--- a/Engine/Source/Editor/Kismet/Public/BlueprintEditor.h
+++ b/Engine/Source/Editor/Kismet/Public/BlueprintEditor.h
@@ -3,6 +3,7 @@
 #pragma once
 
 #include "BlueprintEditorModule.h"
+#include "Blueprint/BlueprintSectionID.h"
 #include "Containers/Array.h"
 #include "Containers/Map.h"
 #include "Containers/Set.h"
@@ -107,29 +108,6 @@ struct FSubobjectData;
 struct FToolMenuContext;
 struct Rect;
 
-/* Enums to use when grouping the blueprint members in the list panel. The order here will determine the order in the list */
-namespace NodeSectionID
-{
-	enum Type
-	{
-		NONE = 0,
-		GRAPH,					// Graph
-		ANIMGRAPH,				// Anim Graph
-		ANIMLAYER,				// Anim Layer
-		FUNCTION,				// Functions
-		FUNCTION_OVERRIDABLE,	// Overridable functions
-		INTERFACE,				// Interface
-		MACRO,					// Macros
-		VARIABLE,				// Variables
-		COMPONENT,				// Components
-		DELEGATE,				// Delegate/Event
-		USER_ENUM,				// User defined enums
-		LOCAL_VARIABLE,			// Local variables
-		USER_STRUCT,			// User defined structs
-		USER_SORTED				// User sorted categories
-	};
-};
-
 /////////////////////////////////////////////////////
 // FCustomDebugObjectEntry - Used to pass a custom debug object override around
 
diff --git a/Engine/Source/Editor/Kismet/Public/BlueprintEditorModule.h b/Engine/Source/Editor/Kismet/Public/BlueprintEditorModule.h
--- a/Engine/Source/Editor/Kismet/Public/BlueprintEditorModule.h
+++ b/Engine/Source/Editor/Kismet/Public/BlueprintEditorModule.h
@@ -2,6 +2,7 @@
 
 #pragma once
 
+#include "Blueprint/BlueprintSectionID.h"
 #include "Containers/Array.h"
 #include "Containers/Map.h"
 #include "Containers/UnrealString.h"
@@ -12,6 +13,7 @@
 #include "Framework/Commands/UICommandList.h"
 #include "K2Node_EditablePinBase.h"
 #include "Math/Vector2D.h"
+#include "Misc/Optional.h"
 #include "Modules/ModuleInterface.h"
 #include "Templates/SharedPointer.h"
 #include "Templates/SubclassOf.h"
@@ -51,6 +53,42 @@ class UUserDefinedStruct;
 struct FBlueprintDebugger;
 struct Rect;
 
+namespace UE::BlueprintEditor
+{
+
+/**
+ * Describes a Blueprint category that should always be visible, even when it contains no entries.
+ * Registered programmatically via FBlueprintEditorModule.
+ */
+struct FAdditionalCategory
+{
+	/** Display name of the category. */
+	FText Category;
+
+	/**
+	 * Only apply this category for blueprints whose ParentClass is a child of one of these classes.
+	 * If empty, applies to all blueprints.
+	 */
+	TArray<FSoftClassPath> ClassFilters;
+
+	struct FPlaceholder
+	{
+		/** Display text of the placeholder entry. */
+		FText MenuDescription;
+
+		/** Tooltip of the placeholder entry. */
+		FText Tooltip;
+
+		/** Section of the placeholder entry. */
+		NodeSectionID::Type SectionID = NodeSectionID::NONE;
+	};
+
+	/** Show placeholder row when the category is empty. */
+	TArray<FPlaceholder> Placeholders;
+};
+
+} // namespace UE::BlueprintEditor
+
 /** Delegate used to customize variable display */
 DECLARE_DELEGATE_RetVal_OneParam(TSharedPtr<IDetailCustomization>, FOnGetVariableCustomizationInstance, TSharedPtr<IBlueprintEditor> /*BlueprintEditor*/);
 
@@ -359,6 +397,44 @@ public:
 	/** Returns a reference to the Blueprint Debugger state object */
 	const TUniquePtr<FBlueprintDebugger>& GetBlueprintDebugger() const { return BlueprintDebugger; }
 
+	/** Opaque handle returned by RegisterAdditionalBlueprintCategory. Pass to Unregister to remove. */
+	struct FAdditionalCategoryHandle
+	{
+		FAdditionalCategoryHandle() = default;
+
+		bool IsValid() const
+		{
+			return Value != InvalidValue;
+		}
+
+		bool operator==(const FAdditionalCategoryHandle& Other) const
+		{
+			return Other.Value == Value;
+		}
+
+	private:
+		friend class FBlueprintEditorModule;
+
+		static constexpr uint32 InvalidValue = MAX_uint32;
+		explicit FAdditionalCategoryHandle(uint32 InValue)
+			: Value(InValue)
+		{}
+
+		uint32 Value = InvalidValue;
+	};
+
+	/** Register an additional category. Returns a handle for unregistration. */
+	UE_API FAdditionalCategoryHandle RegisterAdditionalBlueprintCategory(UE::BlueprintEditor::FAdditionalCategory InCategory);
+
+	/** Unregister a previously registered category by handle. */
+	UE_API void UnregisterAdditionalBlueprintCategory(FAdditionalCategoryHandle Handle);
+
+	/** Returns all additional categories. */
+	UE_API TArray<UE::BlueprintEditor::FAdditionalCategory> GetAdditionalBlueprintCategories() const;
+
+	/** Returns the additional categories for the Blueprint. */
+	UE_API TArray<UE::BlueprintEditor::FAdditionalCategory> GetAdditionalBlueprintCategories(TNotNull<const UBlueprint*> Blueprint) const;
+
 private:
 	/** Loads from ini a list of all events that should be auto created for Blueprints of a specific class */
 	void PrepareAutoGeneratedDefaultEvents();
@@ -417,6 +493,10 @@ private:
 
 	/** Blueprint debugger state - refactor into SBlueprintDebugger if needed */
 	TUniquePtr<FBlueprintDebugger> BlueprintDebugger;
+
+	/** Programmatically registered additional Blueprint categories. */
+	TArray<TTuple<FAdditionalCategoryHandle, UE::BlueprintEditor::FAdditionalCategory>> RegisteredAdditionalCategories;
+	uint32 NextAdditionalCategoryHandle = 0;
 };
 
 #undef UE_API
diff --git a/Engine/Source/Runtime/Engine/Classes/EdGraph/EdGraphSchema.h b/Engine/Source/Runtime/Engine/Classes/EdGraph/EdGraphSchema.h
--- a/Engine/Source/Runtime/Engine/Classes/EdGraph/EdGraphSchema.h
+++ b/Engine/Source/Runtime/Engine/Classes/EdGraph/EdGraphSchema.h
@@ -402,9 +402,14 @@ struct FEdGraphSchemaAction_Dummy : public FEdGraphSchemaAction
 	: FEdGraphSchemaAction()
 	{}
 
-	FEdGraphSchemaAction_Dummy(FText InNodeCategory, FText InMenuDesc, FText InToolTip, const int32 InGrouping)
-		: FEdGraphSchemaAction(MoveTemp(InNodeCategory), MoveTemp(InMenuDesc), MoveTemp(InToolTip), InGrouping)
+	FEdGraphSchemaAction_Dummy(FText InNodeCategory, FText InMenuDesc, FText InToolTip, const int32 InGrouping, FText InKeywords = FText(), int32 InSectionID = 0)
+		: FEdGraphSchemaAction(MoveTemp(InNodeCategory), MoveTemp(InMenuDesc), MoveTemp(InToolTip), InGrouping, InKeywords, InSectionID)
 	{}
+
+	virtual bool CanBeRenamed() const override
+	{
+		return false;
+	}
 };
 
 /** This is a response from CanCreateConnection, indicating if the connecting action is legal and what the result will be */
diff --git a/Engine/Source/Runtime/Engine/Public/Blueprint/BlueprintSectionID.h b/Engine/Source/Runtime/Engine/Public/Blueprint/BlueprintSectionID.h
new file mode 100644
--- /dev/null
+++ b/Engine/Source/Runtime/Engine/Public/Blueprint/BlueprintSectionID.h
@@ -0,0 +1,26 @@
+// Copyright Epic Games, Inc. All Rights Reserved.
+
+#pragma once
+
+/* Enums to use when grouping the blueprint members in the list panel. The order here will determine the order in the list */
+namespace NodeSectionID
+{
+	enum Type
+	{
+		NONE = 0,
+		GRAPH,					// Graph
+		ANIMGRAPH,				// Anim Graph
+		ANIMLAYER,				// Anim Layer
+		FUNCTION,				// Functions
+		FUNCTION_OVERRIDABLE,	// Overridable functions
+		INTERFACE,				// Interface
+		MACRO,					// Macros
+		VARIABLE,				// Variables
+		COMPONENT,				// Components
+		DELEGATE,				// Delegate/Event
+		USER_ENUM,				// User defined enums
+		LOCAL_VARIABLE,			// Local variables
+		USER_STRUCT,			// User defined structs
+		USER_SORTED				// User sorted categories
+	};
+};
diff --git a/Engine/Source/Runtime/SlateCore/Private/Styling/StarshipCoreStyle.cpp b/Engine/Source/Runtime/SlateCore/Private/Styling/StarshipCoreStyle.cpp
--- a/Engine/Source/Runtime/SlateCore/Private/Styling/StarshipCoreStyle.cpp
+++ b/Engine/Source/Runtime/SlateCore/Private/Styling/StarshipCoreStyle.cpp
@@ -1401,6 +1401,11 @@ void FStarshipCoreStyle::SetupTextStyles(TSharedRef<FStyle>& Style)
 		);
 	}
 
+	const FTextBlockStyle HintText = FTextBlockStyle(NormalText)
+		.SetFont(DEFAULT_FONT("Italic", 8))
+		.SetColorAndOpacity(FSlateColor::UseSubduedForeground());
+	Style->Set("HintText", HintText);
+
 	// SInlineEditableTextBlock
 	{
 		// Normal Editable Text 
@@ -1424,12 +1429,14 @@ void FStarshipCoreStyle::SetupTextStyles(TSharedRef<FStyle>& Style)
 			.SetEditableTextBoxStyle(InlineEditableTextBlockSmallEditable);
 
 		Style->Set("InlineEditableTextBlockSmallStyle", InlineEditableTextBlockSmallStyle);
-	}
 
-	const FTextBlockStyle HintText = FTextBlockStyle(NormalText)
-		.SetFont(DEFAULT_FONT("Italic", 8))
-		.SetColorAndOpacity(FSlateColor::UseSubduedForeground());
-	Style->Set("HintText", HintText);
+		// Hint Editable Text 
+		FInlineEditableTextBlockStyle InlineEditableHintTextBlockStyle = FInlineEditableTextBlockStyle()
+			.SetTextStyle(HintText)
+			.SetEditableTextBoxStyle(NormalEditableTextBoxStyle);
+
+		Style->Set("InlineEditableTextBlockHintStyle", InlineEditableHintTextBlockStyle);
+	}
 }
 
 void FStarshipCoreStyle::SetupButtonStyles(TSharedRef<FStyle>& Style)
