Gameplay systems deprecation removal pass for 5.4 and earlier, I skipped anything that was still in use for data fixup or blueprint accessible.

diff --git a/Engine/Source/Editor/Kismet/Private/SSCSEditor.cpp b/Engine/Source/Editor/Kismet/Private/SSCSEditor.cpp
--- a/Engine/Source/Editor/Kismet/Private/SSCSEditor.cpp
+++ b/Engine/Source/Editor/Kismet/Private/SSCSEditor.cpp
@@ -81,9 +81,7 @@
 #include "Misc/Guid.h"
 #include "ObjectTools.h"
 #include "PropertyPath.h"
-#include "SCSEditorExtensionContext.h"
 #include "SPositiveActionButton.h"
-#include "SSCSEditorMenuContext.h"
 #include "ScopedTransaction.h"
 #include "Selection.h"
 #include "Settings/EditorStyleSettings.h"
@@ -140,336 +138,6 @@ struct FSlateBrush;
 
 #define LOCTEXT_NAMESPACE "SSCSEditor"
 
-DEFINE_LOG_CATEGORY_STATIC(LogSCSEditor, Log, All);
-
-static const FName SCS_ColumnName_ComponentClass( "ComponentClass" );
-static const FName SCS_ColumnName_Asset( "Asset" );
-static const FName SCS_ColumnName_Mobility( "Mobility" );
-
-static const FName SCS_ContextMenuName( "Kismet.SCSEditorContextMenu" );
-
-PRAGMA_DISABLE_DEPRECATION_WARNINGS
-
-//////////////////////////////////////////////////////////////////////////
-// SSCSEditorDragDropTree
-void SSCSEditorDragDropTree::Construct( const FArguments& InArgs )
-{
-	SCSEditor = InArgs._SCSEditor;
-
-	STreeView<FSCSEditorTreeNodePtrType>::FArguments BaseArgs;
-	BaseArgs.OnGenerateRow( InArgs._OnGenerateRow )
-			.OnItemScrolledIntoView( InArgs._OnItemScrolledIntoView )
-			.OnGetChildren( InArgs._OnGetChildren )
-			.OnSetExpansionRecursive( InArgs._OnSetExpansionRecursive )
-			.TreeItemsSource( InArgs._TreeItemsSource )
-			.ItemHeight( InArgs._ItemHeight )
-			.OnContextMenuOpening( InArgs._OnContextMenuOpening )
-			.OnMouseButtonDoubleClick( InArgs._OnMouseButtonDoubleClick )
-			.OnSelectionChanged( InArgs._OnSelectionChanged )
-			.OnExpansionChanged( InArgs._OnExpansionChanged )
-			.SelectionMode( InArgs._SelectionMode )
-			.HeaderRow( InArgs._HeaderRow )
-			.ClearSelectionOnClick( InArgs._ClearSelectionOnClick )
-			.ExternalScrollbar( InArgs._ExternalScrollbar )
-			.OnEnteredBadState( InArgs._OnTableViewBadState )
-			.HighlightParentNodesForSelection(true);
-
-	STreeView<FSCSEditorTreeNodePtrType>::Construct( BaseArgs );
-}
-
-FReply SSCSEditorDragDropTree::OnDragOver( const FGeometry& MyGeometry, const FDragDropEvent& DragDropEvent )
-{
-	FReply Handled = FReply::Unhandled();
-
-	if (SCSEditor != nullptr)
-	{
-		TSharedPtr<FDragDropOperation> Operation = DragDropEvent.GetOperation();
-		if (Operation.IsValid() && (Operation->IsOfType<FExternalDragOperation>() || Operation->IsOfType<FAssetDragDropOp>()))
-		{
-			Handled = AssetUtil::CanHandleAssetDrag(DragDropEvent);
-
-			if (!Handled.IsEventHandled())
-			{
-				if (Operation->IsOfType<FAssetDragDropOp>())
-				{
-					const TSharedPtr<FAssetDragDropOp> AssetDragDropOp = StaticCastSharedPtr<FAssetDragDropOp>(Operation);
-
-					for (const FAssetData& AssetData : AssetDragDropOp->GetAssets())
-					{
-						if (UClass* AssetClass = AssetData.GetClass())
-						{
-							if (AssetClass->IsChildOf(UClass::StaticClass()))
-							{
-								Handled = FReply::Handled();
-								break;
-							}
-						}
-					}
-				}
-			}
-		}
-	}
-
-	return Handled;
-}
-
-FReply SSCSEditor::TryHandleAssetDragDropOperation(const FDragDropEvent& DragDropEvent)
-{
-	TSharedPtr<FDragDropOperation> Operation = DragDropEvent.GetOperation();
-	if (Operation.IsValid() && (Operation->IsOfType<FExternalDragOperation>() || Operation->IsOfType<FAssetDragDropOp>()))
-	{
-		TArray< FAssetData > DroppedAssetData = AssetUtil::ExtractAssetDataFromDrag(Operation);
-		const int32 NumAssets = DroppedAssetData.Num();
-
-		if (NumAssets > 0)
-		{
-			GWarn->BeginSlowTask(LOCTEXT("LoadingAssets", "Loading Asset(s)"), true);
-			bool bMarkBlueprintAsModified = false;
-
-			for (int32 DroppedAssetIdx = 0; DroppedAssetIdx < NumAssets; ++DroppedAssetIdx)
-			{
-				const FAssetData& AssetData = DroppedAssetData[DroppedAssetIdx];
-
-				if (!AssetData.IsAssetLoaded())
-				{
-					GWarn->StatusUpdate(DroppedAssetIdx, NumAssets, FText::Format(LOCTEXT("LoadingAsset", "Loading Asset {0}"), FText::FromName(AssetData.AssetName)));
-				}
-
-				UClass* AssetClass = AssetData.GetClass();
-				UObject* Asset = AssetData.GetAsset();
-
-				UBlueprint* BPClass = Cast<UBlueprint>(Asset);
-				UClass* PotentialComponentClass = nullptr;
-				UClass* PotentialActorClass = nullptr;
-
-				if ((BPClass != nullptr) && (BPClass->GeneratedClass != nullptr))
-				{
-					if (BPClass->GeneratedClass->IsChildOf(UActorComponent::StaticClass()))
-					{
-						PotentialComponentClass = BPClass->GeneratedClass;
-					}
-					else if (BPClass->GeneratedClass->IsChildOf(AActor::StaticClass()))
-					{
-						PotentialActorClass = BPClass->GeneratedClass;
-					}
-				}
-				else if (AssetClass && AssetClass->IsChildOf(UClass::StaticClass()))
-				{
-					UClass* AssetAsClass = CastChecked<UClass>(Asset);
-					if (AssetAsClass->IsChildOf(UActorComponent::StaticClass()))
-					{
-						PotentialComponentClass = AssetAsClass;
-					}
-					else if (AssetAsClass->IsChildOf(AActor::StaticClass()))
-					{
-						PotentialActorClass = AssetAsClass;
-					}
-				}
-
-				FAddNewComponentParams NewComponentParams;
-				NewComponentParams.bSkipMarkBlueprintModified = true;
-				NewComponentParams.bSetFocusToNewItem = (DroppedAssetIdx == NumAssets - 1); // Only set focus to the last item created
-
-				TSubclassOf<UActorComponent> MatchingComponentClassForAsset = FComponentAssetBrokerage::GetPrimaryComponentForAsset(AssetClass);
-				if (MatchingComponentClassForAsset != nullptr)
-				{
-					AddNewComponent(MatchingComponentClassForAsset, Asset, NewComponentParams);
-					bMarkBlueprintAsModified = true;
-				}
-				else if ((PotentialComponentClass != nullptr) && !PotentialComponentClass->HasAnyClassFlags(CLASS_Deprecated | CLASS_Abstract | CLASS_NewerVersionExists))
-				{
-					if (PotentialComponentClass->HasMetaData(FBlueprintMetadata::MD_BlueprintSpawnableComponent))
-					{
-						AddNewComponent(PotentialComponentClass, nullptr, NewComponentParams);
-						bMarkBlueprintAsModified = true;
-					}
-				}
-				else if ((PotentialActorClass != nullptr) && !PotentialActorClass->HasAnyClassFlags(CLASS_Deprecated | CLASS_Abstract | CLASS_NewerVersionExists | CLASS_NotPlaceable))
-				{
-					AddNewComponent(UChildActorComponent::StaticClass(), PotentialActorClass, NewComponentParams);
-					bMarkBlueprintAsModified = true;
-				}
-			}
-
-			// Optimization: Only mark the blueprint as modified at the end
-			if (bMarkBlueprintAsModified && EditorMode == EComponentEditorMode::BlueprintSCS)
-			{
-				UBlueprint* Blueprint = GetBlueprint();
-				check(Blueprint != nullptr && Blueprint->SimpleConstructionScript != nullptr);
-
-				Blueprint->Modify();
-				SaveSCSCurrentState(Blueprint->SimpleConstructionScript);
-
-				bAllowTreeUpdates = true;
-				FBlueprintEditorUtils::MarkBlueprintAsStructurallyModified(Blueprint);
-			}
-
-			UpdateTree();
-
-			GWarn->EndSlowTask();
-		}
-
-		return FReply::Handled();
-	}
-
-	return FReply::Unhandled();
-}
-
-FReply SSCSEditorDragDropTree::OnDrop( const FGeometry& MyGeometry, const FDragDropEvent& DragDropEvent ) 
-{
-	if (SCSEditor != nullptr)
-	{
-		return SCSEditor->TryHandleAssetDragDropOperation(DragDropEvent);
-	}
-	else
-	{
-		return FReply::Unhandled();
-	}
-}
-
-
-
-//////////////////////////////////////////////////////////////////////////
-// FSCSRowDragDropOp - The drag-drop operation triggered when dragging a row in the components tree
-
-class FSCSRowDragDropOp : public FKismetVariableDragDropAction
-{
-public:
-	DRAG_DROP_OPERATOR_TYPE(FSCSRowDragDropOp, FKismetVariableDragDropAction)
-
-	/** Available drop actions */
-	enum EDropActionType
-	{
-		DropAction_None,
-		DropAction_AttachTo,
-		DropAction_DetachFrom,
-		DropAction_MakeNewRoot,
-		DropAction_AttachToOrMakeNewRoot
-	};
-
-	// FGraphEditorDragDropAction interface
-	virtual void HoverTargetChanged() override;
-	virtual FReply DroppedOnNode(const FVector2f& ScreenPosition, const FVector2f& GraphPosition) override;
-	virtual FReply DroppedOnPanel(const TSharedRef< class SWidget >& Panel, const FVector2f& ScreenPosition, const FVector2f& GraphPosition, UEdGraph& Graph) override;
-	// End of FGraphEditorDragDropAction
-
-	/** Node(s) that we started the drag from */
-	TArray<FSCSEditorTreeNodePtrType> SourceNodes;
-
-	/** The type of drop action that's pending while dragging */
-	EDropActionType PendingDropAction;
-
-	static TSharedRef<FSCSRowDragDropOp> New(FName InVariableName, UStruct* InVariableSource, FNodeCreationAnalytic AnalyticCallback);
-};
-
-TSharedRef<FSCSRowDragDropOp> FSCSRowDragDropOp::New(FName InVariableName, UStruct* InVariableSource, FNodeCreationAnalytic AnalyticCallback)
-{
-	TSharedPtr<FSCSRowDragDropOp> Operation = MakeShareable(new FSCSRowDragDropOp);
-	Operation->VariableName = InVariableName;
-	Operation->VariableSource = InVariableSource;
-	Operation->AnalyticCallback = AnalyticCallback;
-	Operation->Construct();
-	return Operation.ToSharedRef();
-}
-
-void FSCSRowDragDropOp::HoverTargetChanged()
-{
-	bool bHoverHandled = false;
-
-	FSlateColor IconTint = FLinearColor::White;
-	const FSlateBrush* ErrorSymbol = FAppStyle::GetBrush(TEXT("Graph.ConnectorFeedback.Error"));
-
-	if(SourceNodes.Num() > 1)
-	{
-		// Display an error message if attempting to drop multiple source items onto a node
-		UEdGraphNode* VarNodeUnderCursor = Cast<UK2Node_Variable>(GetHoveredNode());
-		if (VarNodeUnderCursor != NULL)
-		{
-			// Icon/text to draw on tooltip
-			FText Message = LOCTEXT("InvalidMultiDropTarget", "Cannot replace node with multiple nodes");
-			SetSimpleFeedbackMessage(ErrorSymbol, IconTint, Message);
-
-			bHoverHandled = true;
-		}
-	}
-
-	if (!bHoverHandled)
-	{
-		if (FChildActorComponentEditorUtils::ContainsChildActorSubtreeNode(SourceNodes))
-		{
-			// @todo - Add support for drag/drop of child actor template components to create variable get nodes in a local Blueprint event/function graph
-			FText Message = LOCTEXT("ChildActorDragDropAddVariableNode_Unsupported", "This operation is not currently supported for one or more of the selected components.");
-			SetSimpleFeedbackMessage(ErrorSymbol, IconTint, Message);
-		}
-		else if (FProperty* VariableProperty = GetVariableProperty())
-		{
-			const FSlateBrush* PrimarySymbol = nullptr;
-			const FSlateBrush* SecondarySymbol = nullptr;
-			FSlateColor PrimaryColor;
-			FSlateColor SecondaryColor;
-			GetDefaultStatusSymbol(/*out*/ PrimarySymbol, /*out*/ PrimaryColor, /*out*/ SecondarySymbol, /*out*/ SecondaryColor);
-
-			//Create feedback message with the function name.
-			SetSimpleFeedbackMessage(PrimarySymbol, PrimaryColor, VariableProperty->GetDisplayNameText(), SecondarySymbol, SecondaryColor);
-		}
-		else
-		{
-			FText Message = LOCTEXT("CannotFindProperty", "Cannot find corresponding variable (make sure component has been assigned to one)");
-			SetSimpleFeedbackMessage(ErrorSymbol, IconTint, Message);
-		}
-		bHoverHandled = true;
-	}
-
-	if(!bHoverHandled)
-	{
-		FKismetVariableDragDropAction::HoverTargetChanged();
-	}
-}
-
-FReply FSCSRowDragDropOp::DroppedOnNode(const FVector2f& ScreenPosition, const FVector2f& GraphPosition)
-{
-	// Only allow dropping on another node if there is only a single source item
-	if(SourceNodes.Num() == 1)
-	{
-		FKismetVariableDragDropAction::DroppedOnNode(ScreenPosition, GraphPosition);
-	}
-	return FReply::Handled();
-}
-
-FReply FSCSRowDragDropOp::DroppedOnPanel(const TSharedRef< class SWidget >& Panel, const FVector2f& ScreenPosition, const FVector2f& GraphPosition, UEdGraph& Graph)
-{
-	const FScopedTransaction Transaction(LOCTEXT("SCSEditorAddMultipleNodes", "Add Component Nodes"));
-
-	TArray<UK2Node_VariableGet*> OriginalVariableNodes;
-	Graph.GetNodesOfClass<UK2Node_VariableGet>(OriginalVariableNodes);
-
-	FVector2f GraphPositionToPlace = GraphPosition;
-	// Add source items to the graph in turn
-	for (FSCSEditorTreeNodePtrType& SourceNode : SourceNodes)
-	{
-		VariableName = SourceNode->GetVariableName();
-		FKismetVariableDragDropAction::DroppedOnPanel(Panel, ScreenPosition, GraphPositionToPlace, Graph);
-
-		GraphPositionToPlace.Y += 50;
-	}
-
-	TArray<UK2Node_VariableGet*> ResultVariableNodes;
-	Graph.GetNodesOfClass<UK2Node_VariableGet>(ResultVariableNodes);
-
-	if (ResultVariableNodes.Num() - OriginalVariableNodes.Num() > 1)
-	{
-		TSet<const UEdGraphNode*> NodeSelection;
-
-		// Because there is more than one new node, lets grab all the nodes at the bottom of the list and add them to a set for selection
-		for (int32 NodeIdx = ResultVariableNodes.Num() - 1; NodeIdx >= OriginalVariableNodes.Num(); --NodeIdx)
-		{
-			NodeSelection.Add(ResultVariableNodes[NodeIdx]);
-		}
-		Graph.SelectNodeSet(NodeSelection);
-	}
-	return FReply::Handled();
-}
-
 //////////////////////////////////////////////////////////////////////////
 // FSCSEditorTreeNode
 
@@ -911,6 +579,34 @@ FSCSEditorTreeNodePtrType FSCSEditorTreeNode::FactoryNodeFromComponent(UActorCom
 	return MakeShareable(new FSCSEditorTreeNodeComponent(InComponent));
 }
 
+FSlateColor FSCSEditorTreeNode::GetColorTintForIcon(FSCSEditorTreeNodePtrType InNode)
+{
+	const FLinearColor InheritedBlueprintComponentColor(0.08f, 0.35f, 0.6f);
+	const FLinearColor InstancedInheritedBlueprintComponentColor(0.08f, 0.35f, 0.6f);
+	const FLinearColor InheritedNativeComponentColor(0.7f, 0.9f, 0.7f);
+	const FLinearColor IntroducedHereColor(FLinearColor::White);
+
+	if (InNode->IsInheritedComponent())
+	{
+		if (InNode->IsNativeComponent())
+		{
+			return InheritedNativeComponentColor;
+		}
+		else if (InNode->IsInstancedComponent())
+		{
+			return InstancedInheritedBlueprintComponentColor;
+		}
+		else
+		{
+			return InheritedBlueprintComponentColor;
+		}
+	}
+	else
+	{
+		return IntroducedHereColor;
+	}
+}
+
 void FSCSEditorTreeNode::CloseOngoingCreateTransaction()
 {
 	OngoingCreateTransaction.Reset();
@@ -1816,5592 +1512,5 @@ void FSCSEditorTreeNodeSeparator::AddFilteredComponentType(const TSubclassOf<UAc
 	}
 }
 
-//////////////////////////////////////////////////////////////////////////
-// SSCS_RowWidget
-
-void SSCS_RowWidget::Construct( const FArguments& InArgs, TSharedPtr<SSCSEditor> InSCSEditor, FSCSEditorTreeNodePtrType InNodePtr, TSharedPtr<STableViewBase> InOwnerTableView  )
-{
-	check(InNodePtr.IsValid());
-
-	SCSEditor = InSCSEditor;
-	TreeNodePtr = InNodePtr;
-
-	bool bIsSeparator = InNodePtr->GetNodeType() == FSCSEditorTreeNode::SeparatorNode;
-	
-	FSuperRowType::FArguments Args = FSuperRowType::FArguments()
-		.Style(bIsSeparator ?
-			&FAppStyle::Get().GetWidgetStyle<FTableRowStyle>("TableView.NoHoverTableRow") :
-			&FAppStyle::Get().GetWidgetStyle<FTableRowStyle>("SceneOutliner.TableViewRow")) //@todo create editor style for the SCS tree
-		.Padding(FMargin(0.f, 4.f, 0.f, 4.f))
-		.ShowSelection(!bIsSeparator)
-		.OnDragDetected(this, &SSCS_RowWidget::HandleOnDragDetected)
-		.OnDragEnter(this, &SSCS_RowWidget::HandleOnDragEnter)
-		.OnDragLeave(this, &SSCS_RowWidget::HandleOnDragLeave)
-		.OnCanAcceptDrop(this, &SSCS_RowWidget::HandleOnCanAcceptDrop)
-		.OnAcceptDrop(this, &SSCS_RowWidget::HandleOnAcceptDrop);
-
-	SMultiColumnTableRow<FSCSEditorTreeNodePtrType>::Construct( Args, InOwnerTableView.ToSharedRef() );
-}
-
-SSCS_RowWidget::~SSCS_RowWidget()
-{
-	TSharedPtr<SSCSEditor> Editor = SCSEditor.Pin();
-	if(Editor.IsValid())
-	{
-		// Ensure to end any owned ongoing transaction if the node is still in 'editing initial name' mode when deleted.
-		GetNode()->CloseOngoingCreateTransaction();
-
-		// Ask SCSEditor if Node is still active, if it isn't it might have been collected so we can't do anything to it
-		USCS_Node* SCS_Node = GetNode()->GetSCSNode();
-		if(SCS_Node != NULL && Editor->IsNodeInSimpleConstructionScript(SCS_Node))
-		{
-			// Clear the delegate when widget goes away
-			SCS_Node->SetOnNameChanged(FSCSNodeNameChanged());
-		}
-	}
-}
-
-BEGIN_SLATE_FUNCTION_BUILD_OPTIMIZATION
-TSharedRef<SWidget> SSCS_RowWidget::GenerateWidgetForColumn( const FName& ColumnName )
-{
-	FSCSEditorTreeNodePtrType NodePtr = GetNode();
-	if(ColumnName == SCS_ColumnName_ComponentClass)
-	{
-		InlineWidget =
-			SNew(SInlineEditableTextBlock)
-				.Text(this, &SSCS_RowWidget::GetNameLabel)
-				.OnVerifyTextChanged( this, &SSCS_RowWidget::OnNameTextVerifyChanged )
-				.OnTextCommitted( this, &SSCS_RowWidget::OnNameTextCommit )
-				.IsSelected( this, &SSCS_RowWidget::IsSelectedExclusively )
-				.IsReadOnly(!NodePtr->CanRename() || (SCSEditor.IsValid() && !SCSEditor.Pin()->IsEditingAllowed()) || FChildActorComponentEditorUtils::IsChildActorSubtreeNode(NodePtr));
-
-		NodePtr->SetRenameRequestedDelegate(FSCSEditorTreeNode::FOnRenameRequested::CreateSP(InlineWidget.Get(), &SInlineEditableTextBlock::EnterEditingMode));
-		
-		TSharedRef<SToolTip> Tooltip = CreateToolTipWidget();
-
-		return SNew(SHorizontalBox)
-				.ToolTip(Tooltip)
-				+SHorizontalBox::Slot()
-					.AutoWidth()
-					.VAlign(VAlign_Center)
-					[
-						SNew(SExpanderArrow, SharedThis(this))
-					]
-				+SHorizontalBox::Slot()
-					.AutoWidth()
-					.VAlign(VAlign_Center)
-					[
-						SNew(SImage)
-						.Image(GetIconBrush())
-						.ColorAndOpacity(this, &SSCS_RowWidget::GetColorTintForIcon)
-					]
-				+SHorizontalBox::Slot()
-					.VAlign(VAlign_Center)
-					.Padding(4.0f, 0.0f)
-					[
-						InlineWidget.ToSharedRef()
-					];
-	}
-	else if(ColumnName == SCS_ColumnName_Asset)
-	{
-		return
-			SNew(SHorizontalBox)
-			+SHorizontalBox::Slot()
-			.VAlign(VAlign_Center)
-			[
-				SNew(STextBlock)
-				.Visibility(this, &SSCS_RowWidget::GetAssetVisibility)
-				.Text(this, &SSCS_RowWidget::GetAssetName)
-				.ToolTipText(this, &SSCS_RowWidget::GetAssetPath)
-			];
-	}
-	else if (ColumnName == SCS_ColumnName_Mobility)
-	{
-		if (NodePtr->GetNodeType() == FSCSEditorTreeNode::ComponentNode)
-		{
-			TSharedPtr<SToolTip> MobilityTooltip = SNew(SToolTip)
-				.Text(this, &SSCS_RowWidget::GetMobilityToolTipText);
-
-			return SNew(SHorizontalBox)
-				.ToolTip(MobilityTooltip)
-				.Visibility(EVisibility::Visible) // so we still get tooltip text for an empty SHorizontalBox
-				+ SHorizontalBox::Slot()
-				.VAlign(VAlign_Center)
-				.FillWidth(1.0f)
-				[
-					SNew(SImage)
-					.Image(this, &SSCS_RowWidget::GetMobilityIconImage)
-					.ToolTip(MobilityTooltip)
-				];
-		}
-		else
-		{
-			return SNew(SSpacer);
-		}
-	}
-	else
-	{
-		return	SNew(STextBlock)
-				.Text( LOCTEXT("UnknownColumn", "Unknown Column") );
-	}
-}
-END_SLATE_FUNCTION_BUILD_OPTIMIZATION
-
-BEGIN_SLATE_FUNCTION_BUILD_OPTIMIZATION
-void SSCS_RowWidget::AddToToolTipInfoBox(const TSharedRef<SVerticalBox>& InfoBox, const FText& Key, TSharedRef<SWidget> ValueIcon, const TAttribute<FText>& Value, bool bImportant)
-{
-	InfoBox->AddSlot()
-		.AutoHeight()
-		.Padding(0.0f, 1.0f)
-		[
-			SNew(SHorizontalBox)
-
-			+ SHorizontalBox::Slot()
-			.AutoWidth()
-			.Padding(0.0f, 0.0f, 4.0f, 0.0f)
-			[
-				SNew(STextBlock)
-				.TextStyle(FAppStyle::Get(), bImportant ? "SCSEditor.ComponentTooltip.ImportantLabel" : "SCSEditor.ComponentTooltip.Label")
-				.Text(FText::Format(LOCTEXT("AssetViewTooltipFormat", "{0}:"), Key))
-			]
-
-			+ SHorizontalBox::Slot()
-			.AutoWidth()
-			[
-				ValueIcon
-			]
-
-			+ SHorizontalBox::Slot()
-			.AutoWidth()
-			[
-				SNew(STextBlock)
-				.TextStyle(FAppStyle::Get(), bImportant ? "SCSEditor.ComponentTooltip.ImportantValue" : "SCSEditor.ComponentTooltip.Value")
-				.Text(Value)
-			]
-		];
-}
-END_SLATE_FUNCTION_BUILD_OPTIMIZATION
-
-BEGIN_SLATE_FUNCTION_BUILD_OPTIMIZATION
-TSharedRef<SToolTip> SSCS_RowWidget::CreateToolTipWidget() const
-{
-	// Create a box to hold every line of info in the body of the tooltip
-	TSharedRef<SVerticalBox> InfoBox = SNew(SVerticalBox);
-
-	// 
-	if (FSCSEditorTreeNode* TreeNode = GetNode().Get())
-	{
-		if (TreeNode->IsComponentNode())
-		{
-			// Add the tooltip
-			if (UActorComponent* Template = TreeNode->GetComponentTemplate())
-			{
-				UClass* TemplateClass = Template->GetClass();
-				FText ClassTooltip = TemplateClass->GetToolTipText(/*bShortTooltip=*/ true);
-
-				InfoBox->AddSlot()
-					.AutoHeight()
-					.HAlign(HAlign_Center)
-					.Padding(FMargin(0.0f, 2.0f, 0.0f, 4.0f))
-					[
-						SNew(STextBlock)
-						.TextStyle(FAppStyle::Get(), "SCSEditor.ComponentTooltip.ClassDescription")
-						.Text(ClassTooltip)
-						.WrapTextAt(400.0f)
-					];
-			}
-
-			// Add introduction point
-			AddToToolTipInfoBox(InfoBox, LOCTEXT("TooltipAddType", "Source"), SNullWidget::NullWidget, TAttribute<FText>::Create(TAttribute<FText>::FGetter::CreateSP(this, &SSCS_RowWidget::GetComponentAddSourceToolTipText)), false);
-			if (TreeNode->IsInheritedComponent())
-			{
-				AddToToolTipInfoBox(InfoBox, LOCTEXT("TooltipIntroducedIn", "Introduced in"), SNullWidget::NullWidget, TAttribute<FText>::Create(TAttribute<FText>::FGetter::CreateSP(this, &SSCS_RowWidget::GetIntroducedInToolTipText)), false);
-			}
-
-			// Add Underlying Component Name for Native Components
-			if (TreeNode->IsNativeComponent())
-			{
-				AddToToolTipInfoBox(InfoBox, LOCTEXT("TooltipNativeComponentName", "Native Component Name"), SNullWidget::NullWidget, TAttribute<FText>::Create(TAttribute<FText>::FGetter::CreateSP(this, &SSCS_RowWidget::GetNativeComponentNameToolTipText)), false);
-			}
-
-			// Add mobility
-			TSharedRef<SImage> MobilityIcon = SNew(SImage).Image(this, &SSCS_RowWidget::GetMobilityIconImage);
-			AddToToolTipInfoBox(InfoBox, LOCTEXT("TooltipMobility", "Mobility"), MobilityIcon, TAttribute<FText>::Create(TAttribute<FText>::FGetter::CreateSP(this, &SSCS_RowWidget::GetMobilityToolTipText)), false);
-
-			// Add asset if applicable to this node
-			if (GetAssetVisibility() == EVisibility::Visible)
-			{
-				InfoBox->AddSlot()[SNew(SSpacer).Size(FVector2D(1.0f, 8.0f))];
-				AddToToolTipInfoBox(InfoBox, LOCTEXT("TooltipAsset", "Asset"), SNullWidget::NullWidget, TAttribute<FText>(this, &SSCS_RowWidget::GetAssetName), false);
-			}
-
-			// If the component is marked as editor only, then display that info here
-			AddToToolTipInfoBox(InfoBox, LOCTEXT("TooltipEditorOnly", "Editor Only"), SNullWidget::NullWidget, TAttribute<FText>::Create(TAttribute<FText>::FGetter::CreateSP(this, &SSCS_RowWidget::GetComponentEditorOnlyTooltipText)), false);
-		}
-	}
-
-	TSharedRef<SBorder> TooltipContent = SNew(SBorder)
-		.BorderImage(FAppStyle::GetBrush("NoBorder"))
-		.Padding(0.0f)
-		[
-			SNew(SVerticalBox)
-			
-			+SVerticalBox::Slot()
-			.AutoHeight()
-			.Padding(0.0f, 0.0f, 0.0f, 4.0f)
-			[
-				SNew(SVerticalBox)
-
-				+ SVerticalBox::Slot()
-				.AutoHeight()
-				[
-					SNew(SHorizontalBox)
-
-					+ SHorizontalBox::Slot()
-					.AutoWidth()
-					.VAlign(VAlign_Center)
-					.Padding(2.0f)
-					[
-						SNew(STextBlock)
-						.TextStyle(FAppStyle::Get(), "SCSEditor.ComponentTooltip.Title")
-						.Text(this, &SSCS_RowWidget::GetTooltipText)
-					]
-				]
-			]
-
-			+ SVerticalBox::Slot()
-			.AutoHeight()
-			[
-				SNew(SBorder)
-				.BorderImage(FAppStyle::GetBrush("NoBorder"))
-				.Padding(2.0f)
-				[
-					InfoBox
-				]
-			]
-		];
-
-	return IDocumentation::Get()->CreateToolTip(TAttribute<FText>(this, &SSCS_RowWidget::GetTooltipText), TooltipContent, InfoBox, GetDocumentationLink(), GetDocumentationExcerptName());
-}
-END_SLATE_FUNCTION_BUILD_OPTIMIZATION
-
-FSlateBrush const* SSCS_RowWidget::GetMobilityIconImage() const
-{
-	if (FSCSEditorTreeNode* TreeNode = GetNode().Get())
-	{
-		if (USceneComponent* SceneComponentTemplate = Cast<USceneComponent>(TreeNode->GetComponentTemplate()))
-		{
-			if (SceneComponentTemplate->Mobility == EComponentMobility::Movable)
-			{
-				return FAppStyle::GetBrush(TEXT("ClassIcon.MovableMobilityIcon"));
-			}
-			else if (SceneComponentTemplate->Mobility == EComponentMobility::Stationary)
-			{
-				return FAppStyle::GetBrush(TEXT("ClassIcon.StationaryMobilityIcon"));
-			}
-
-			// static components don't get an icon (because static is the most common
-			// mobility type, and we'd like to keep the icon clutter to a minimum)
-		}
-	}
-
-	return nullptr;
-}
-
-FText SSCS_RowWidget::GetMobilityToolTipText() const
-{
-	FText MobilityToolTip = LOCTEXT("ErrorNoMobilityTooltip", "Invalid component");
-
-	if (FSCSEditorTreeNode* TreeNode = TreeNodePtr.Get())
-	{
-		if (USceneComponent* SceneComponentTemplate = Cast<USceneComponent>(TreeNode->GetComponentTemplate()))
-		{
-			if (SceneComponentTemplate->Mobility == EComponentMobility::Movable)
-			{
-				MobilityToolTip = LOCTEXT("MovableMobilityTooltip", "Movable");
-			}
-			else if (SceneComponentTemplate->Mobility == EComponentMobility::Stationary)
-			{
-				MobilityToolTip = LOCTEXT("StationaryMobilityTooltip", "Stationary");
-			}
-			else if (SceneComponentTemplate->Mobility == EComponentMobility::Static)
-			{
-				MobilityToolTip = LOCTEXT("StaticMobilityTooltip", "Static");
-			}
-			else
-			{
-				// make sure we're the mobility type we're expecting (we've handled Movable & Stationary)
-				ensureMsgf(false, TEXT("Unhandled mobility type [%d], is this a new type that we don't handle here?"), SceneComponentTemplate->Mobility.GetValue());
-				MobilityToolTip = LOCTEXT("UnknownMobilityTooltip", "Component with unknown mobility");
-			}
-		}
-		else
-		{
-			MobilityToolTip = LOCTEXT("NoMobilityTooltip", "Non-scene component");
-		}
-	}
-
-	return MobilityToolTip;
-}
-
-FText SSCS_RowWidget::GetComponentAddSourceToolTipText() const
-{
-	FText NodeType;
-	
-	if (FSCSEditorTreeNode* TreeNode = TreeNodePtr.Get())
-	{
-		if (TreeNode->IsInheritedComponent())
-		{
-			if (TreeNode->IsNativeComponent())
-			{
-				NodeType = LOCTEXT("InheritedNativeComponent", "Inherited (C++)");
-			}
-			else
-			{
-				NodeType = LOCTEXT("InheritedBlueprintComponent", "Inherited (Blueprint)");
-			}
-		}
-		else
-		{
-			if (TreeNode->IsInstancedComponent())
-			{
-				NodeType = LOCTEXT("ThisInstanceAddedComponent", "This actor instance");
-			}
-			else
-			{
-				NodeType = LOCTEXT("ThisBlueprintAddedComponent", "This Blueprint");
-			}
-		}
-	}
-
-	return NodeType;
-}
-
-FText SSCS_RowWidget::GetNativeComponentNameToolTipText() const
-{
-	const FSCSEditorTreeNode* TreeNode = TreeNodePtr.Get();
-	const UActorComponent* Template = TreeNode ? TreeNode->GetComponentTemplate() : nullptr;
-
-	if (Template)
-	{
-		return FText::FromName(Template->GetFName());
-	}
-	else
-	{
-		return FText::GetEmpty();
-	}
-}
-
-FText SSCS_RowWidget::GetComponentEditorOnlyTooltipText() const
-{
-	FText ComponentType = LOCTEXT("ComponentEditorOnlyFalse", "False");
-	if (FSCSEditorTreeNode* TreeNode = GetNode().Get())
-	{
-		if (TreeNode->IsComponentNode())
-		{
-			if (const UActorComponent* Template = TreeNode->GetComponentTemplate())
-			{
-				UBlueprint* Blueprint = GetBlueprint();
-				FObjectProperty* Prop = Blueprint ? FindFProperty<FObjectProperty>(Blueprint->SkeletonGeneratedClass, TreeNode->GetVariableName()) : nullptr;
-
-				if(Template->bIsEditorOnly || (Prop && Prop->HasAnyPropertyFlags(CPF_EditorOnly)))
-				{
-					ComponentType = LOCTEXT("ComponentEditorOnlyTrue", "True");
-				}
-			}
-		}
-	}
-	
-	return ComponentType;
-}
-
-FText SSCS_RowWidget::GetIntroducedInToolTipText() const
-{
-	FText IntroducedInTooltip = LOCTEXT("IntroducedInThisBPTooltip", "this class");
-
-	if (FSCSEditorTreeNode* TreeNode = TreeNodePtr.Get())
-	{
-		if (TreeNode->IsInheritedComponent())
-		{
-			if (UActorComponent* ComponentTemplate = TreeNode->GetComponentTemplate())
-			{
-				UClass* BestClass = nullptr;
-				AActor* OwningActor = ComponentTemplate->GetOwner();
-
-				if (TreeNode->IsNativeComponent() && (OwningActor != nullptr))
-				{
-					for (UClass* TestClass = OwningActor->GetClass(); TestClass != AActor::StaticClass(); TestClass = TestClass->GetSuperClass())
-					{
-						if (TreeNode->FindComponentInstanceInActor(Cast<AActor>(TestClass->GetDefaultObject())))
-						{
-							BestClass = TestClass;
-						}
-						else
-						{
-							break;
-						}
-					}
-				}
-				else if (!TreeNode->IsNativeComponent())
-				{
-					USCS_Node* SCSNode = TreeNode->GetSCSNode();
-
-					if ((SCSNode == nullptr) && (OwningActor != nullptr))
-					{
-						SCSNode = FSCSEditorTreeNode::FindSCSNodeForInstance(ComponentTemplate, OwningActor->GetClass());
-					}
-
-					if (SCSNode != nullptr)
-					{
-						if (UBlueprint* OwningBP = SCSNode->GetSCS()->GetBlueprint())
-						{
-							BestClass = OwningBP->GeneratedClass;
-						}
-					}
-					else if (OwningActor != nullptr)
-					{
-						if (UBlueprint* OwningBP = UBlueprint::GetBlueprintFromClass(OwningActor->GetClass()))
-						{
-							BestClass = OwningBP->GeneratedClass;
-						}
-					}
-				}
-
-				if (BestClass == nullptr)
-				{
-					if (ComponentTemplate->IsCreatedByConstructionScript()) 
-					{
-						IntroducedInTooltip = LOCTEXT("IntroducedInUnknownError", "Unknown Blueprint Class (via an Add Component call)");
-					} 
-					else 
-					{
-						IntroducedInTooltip = LOCTEXT("IntroducedInNativeError", "Unknown native source (via C++ code)");
-					}
-				}
-				else if (TreeNode->IsInstancedComponent() && ComponentTemplate->CreationMethod == EComponentCreationMethod::Native && !ComponentTemplate->HasAnyFlags(RF_DefaultSubObject))
-				{
-					IntroducedInTooltip = FText::Format(LOCTEXT("IntroducedInCPPErrorFmt", "{0} (via C++ code)"), FBlueprintEditorUtils::GetFriendlyClassDisplayName(BestClass));
-				}
-				else if (TreeNode->IsInstancedComponent() && ComponentTemplate->CreationMethod == EComponentCreationMethod::UserConstructionScript)
-				{
-					IntroducedInTooltip = FText::Format(LOCTEXT("IntroducedInUCSErrorFmt", "{0} (via an Add Component call)"), FBlueprintEditorUtils::GetFriendlyClassDisplayName(BestClass));
-				}
-				else
-				{
-					IntroducedInTooltip = FBlueprintEditorUtils::GetFriendlyClassDisplayName(BestClass);
-				}
-			}
-			else
-			{
-				IntroducedInTooltip = LOCTEXT("IntroducedInNoTemplateError", "[no component template found]");
-			}
-		}
-		else if (TreeNode->IsInstancedComponent())
-		{
-			IntroducedInTooltip = LOCTEXT("IntroducedInThisActorInstanceTooltip", "this actor instance");
-		}
-	}
-
-	return IntroducedInTooltip;
-}
-
-FText SSCS_RowWidget::GetAssetName() const
-{
-	FSCSEditorTreeNodePtrType NodePtr = GetNode();
-
-	FText AssetName = LOCTEXT("None", "None");
-	if(NodePtr.IsValid() && NodePtr->GetComponentTemplate())
-	{
-		UObject* Asset = FComponentAssetBrokerage::GetAssetFromComponent(NodePtr->GetComponentTemplate());
-		if(Asset != NULL)
-		{
-			AssetName = FText::FromString(Asset->GetName());
-		}
-	}
-
-	return AssetName;
-}
-
-FText SSCS_RowWidget::GetAssetPath() const
-{
-	FSCSEditorTreeNodePtrType NodePtr = GetNode();
-
-	FText AssetName = LOCTEXT("None", "None");
-	if(NodePtr.IsValid() && NodePtr->GetComponentTemplate())
-	{
-		UObject* Asset = FComponentAssetBrokerage::GetAssetFromComponent(NodePtr->GetComponentTemplate());
-		if(Asset != NULL)
-		{
-			AssetName = FText::FromString(Asset->GetPathName());
-		}
-	}
-
-	return AssetName;
-}
-
-
-EVisibility SSCS_RowWidget::GetAssetVisibility() const
-{
-	FSCSEditorTreeNodePtrType NodePtr = GetNode();
-
-	if(NodePtr.IsValid() && NodePtr->GetComponentTemplate() && FComponentAssetBrokerage::SupportsAssets(NodePtr->GetComponentTemplate()))
-	{
-		return EVisibility::Visible;
-	}
-	else
-	{
-		return EVisibility::Hidden;
-	}
-}
-
-const FSlateBrush* SSCS_RowWidget::GetIconBrush() const
-{
-	const FSlateBrush* ComponentIcon = FAppStyle::GetBrush("SCS.NativeComponent");
-	FSCSEditorTreeNodePtrType NodePtr = GetNode();
-	if (NodePtr.IsValid())
-	{
-		if (UActorComponent* ComponentTemplate = NodePtr->GetComponentTemplate())
-		{
-			ComponentIcon = FSlateIconFinder::FindIconBrushForClass(ComponentTemplate->GetClass(), TEXT("SCS.Component"));
-		}
-	}
-
-	return ComponentIcon;
-}
-
-FSlateColor SSCS_RowWidget::GetColorTintForIcon() const
-{
-	return GetColorTintForIcon(GetNode());
-}
-
-FSlateColor SSCS_RowWidget::GetColorTintForIcon(FSCSEditorTreeNodePtrType InNode)
-{
-	const FLinearColor InheritedBlueprintComponentColor(0.08f, 0.35f, 0.6f);
-	const FLinearColor InstancedInheritedBlueprintComponentColor(0.08f, 0.35f, 0.6f);
-	const FLinearColor InheritedNativeComponentColor(0.7f, 0.9f, 0.7f);
-	const FLinearColor IntroducedHereColor(FLinearColor::White);
-	
-	if (InNode->IsInheritedComponent())
-	{
-		if (InNode->IsNativeComponent())
-		{
-			return InheritedNativeComponentColor;
-		}
-		else if (InNode->IsInstancedComponent())
-		{
-			return InstancedInheritedBlueprintComponentColor;
-		}
-		else
-		{
-			return InheritedBlueprintComponentColor;
-		}
-	}
-	else
-	{
-		return IntroducedHereColor;
-	}
-}
-
-TSharedPtr<SWidget> SSCS_RowWidget::BuildSceneRootDropActionMenu(FSCSEditorTreeNodePtrType DroppedNodePtr)
-{
-	check(SCSEditor.IsValid());
-	FMenuBuilder MenuBuilder(true, SCSEditor.Pin()->CommandList);
-
-	MenuBuilder.BeginSection("SceneRootNodeDropActions", LOCTEXT("SceneRootNodeDropActionContextMenu", "Drop Actions"));
-	{
-		const FText DroppedVariableNameText = FText::FromName( DroppedNodePtr->GetVariableName() );
-		const FText NodeVariableNameText = FText::FromName( GetNode()->GetVariableName() );
-
-		bool bDroppedInSameBlueprint = true;
-		if (SCSEditor.Pin()->GetEditorMode() == EComponentEditorMode::BlueprintSCS)
-		{
-			bDroppedInSameBlueprint = DroppedNodePtr->GetBlueprint() == GetBlueprint();
-		}
-
-		MenuBuilder.AddMenuEntry(
-			LOCTEXT("DropActionLabel_AttachToRootNode", "Attach"),
-			bDroppedInSameBlueprint 
-			? FText::Format( LOCTEXT("DropActionToolTip_AttachToRootNode", "Attach {0} to {1}."), DroppedVariableNameText, NodeVariableNameText )
-			: FText::Format( LOCTEXT("DropActionToolTip_AttachToRootNodeFromCopy", "Copy {0} to a new variable and attach it to {1}."), DroppedVariableNameText, NodeVariableNameText ),
-			FSlateIcon(),
-			FUIAction(
-				FExecuteAction::CreateSP(this, &SSCS_RowWidget::OnAttachToDropAction, DroppedNodePtr),
-				FCanExecuteAction()));
-
-		FSCSEditorTreeNodePtrType NodePtr = GetNode();
-		const bool bIsDefaultSceneRoot = NodePtr->IsDefaultSceneRoot();
-
-		FText NewRootNodeText = bIsDefaultSceneRoot
-			? FText::Format(LOCTEXT("DropActionToolTip_MakeNewRootNodeAndDelete", "Make {0} the new root. The default root will be deleted."), DroppedVariableNameText)
-			: FText::Format(LOCTEXT("DropActionToolTip_MakeNewRootNode", "Make {0} the new root."), DroppedVariableNameText);
-
-		FText NewRootNodeFromCopyText = bIsDefaultSceneRoot
-			? FText::Format(LOCTEXT("DropActionToolTip_MakeNewRootNodeFromCopyAndDelete", "Copy {0} to a new variable and make it the new root. The default root will be deleted."), DroppedVariableNameText)
-			: FText::Format(LOCTEXT("DropActionToolTip_MakeNewRootNodeFromCopy", "Copy {0} to a new variable and make it the new root."), DroppedVariableNameText);
-
-		MenuBuilder.AddMenuEntry(
-			LOCTEXT("DropActionLabel_MakeNewRootNode", "Make New Root"),
-			bDroppedInSameBlueprint ? NewRootNodeText : NewRootNodeFromCopyText,
-			FSlateIcon(),
-			FUIAction(
-				FExecuteAction::CreateSP(this, &SSCS_RowWidget::OnMakeNewRootDropAction, DroppedNodePtr),
-				FCanExecuteAction()));
-	}
-	MenuBuilder.EndSection();
-
-	return MenuBuilder.MakeWidget();
-}
-
-FReply SSCS_RowWidget::OnMouseButtonDown(const FGeometry& MyGeometry, const FPointerEvent& MouseEvent)
-{
-	if (MouseEvent.GetEffectingButton() == EKeys::LeftMouseButton && GetNode()->GetNodeType() != FSCSEditorTreeNode::SeparatorNode)
-	{
-		FReply Reply = SMultiColumnTableRow<FSCSEditorTreeNodePtrType>::OnMouseButtonDown( MyGeometry, MouseEvent );
-		return Reply.DetectDrag( SharedThis(this) , EKeys::LeftMouseButton );
-	}
-	else
-	{
-		return FReply::Unhandled();
-	}
-}
-
-FReply SSCS_RowWidget::HandleOnDragDetected( const FGeometry& MyGeometry, const FPointerEvent& MouseEvent )
-{
-	TSharedPtr<SSCSEditor> SCSEditorPtr = SCSEditor.Pin();
-	if (MouseEvent.IsMouseButtonDown(EKeys::LeftMouseButton)
-		&& SCSEditorPtr.IsValid()
-		&& SCSEditorPtr->IsEditingAllowed()) //can only drag when editing
-	{
-		TArray<TSharedPtr<FSCSEditorTreeNode>> SelectedNodePtrs = SCSEditorPtr->GetSelectedNodes();
-		if (SelectedNodePtrs.Num() == 0)
-		{
-			SelectedNodePtrs.Add(GetNode());
-		}
-
-		TSharedPtr<FSCSEditorTreeNode> FirstNode = SelectedNodePtrs[0];
-		if (FirstNode->IsComponentNode())
-		{
-			// Do not use the Blueprint from FirstNode, it may still be referencing the parent.
-			UBlueprint* Blueprint = GetBlueprint();
-			const FName VariableName = FirstNode->GetVariableName();
-			UStruct* VariableScope = (Blueprint != nullptr) ? *Blueprint->SkeletonGeneratedClass : nullptr;
-
-			TSharedRef<FSCSRowDragDropOp> Operation = FSCSRowDragDropOp::New(VariableName, VariableScope, FNodeCreationAnalytic());
-			Operation->SetCtrlDrag(true); // Always put a getter
-			Operation->PendingDropAction = FSCSRowDragDropOp::DropAction_None;
-			Operation->SourceNodes = SelectedNodePtrs;
-
-			return FReply::Handled().BeginDragDrop(Operation);
-		}
-	}
-	
-	return FReply::Unhandled();
-}
-
-void SSCS_RowWidget::HandleOnDragEnter( const FDragDropEvent& DragDropEvent )
-{
-	TSharedPtr<FDragDropOperation> Operation = DragDropEvent.GetOperation();
-	if (!Operation.IsValid())
-	{
-		return;
-	}
-
-	TSharedPtr<FSCSRowDragDropOp> DragRowOp = DragDropEvent.GetOperationAs<FSCSRowDragDropOp>();
-	if (DragRowOp.IsValid())
-	{
-		check(SCSEditor.IsValid());
-		
-		FText Message;
-		FSlateColor IconColor = FLinearColor::White;
-		
-		for (const FSCSEditorTreeNodePtrType& SelectedNodePtr : DragRowOp->SourceNodes)
-		{
-			if (!SelectedNodePtr->CanReparent())
-			{
-				// We set the tooltip text here because it won't change across entry/leave events
-				if (DragRowOp->SourceNodes.Num() == 1)
-				{
-					if (FChildActorComponentEditorUtils::IsChildActorSubtreeNode(SelectedNodePtr))
-					{
-						Message = LOCTEXT("DropActionToolTip_Error_CannotReparent_ChildActorSubTreeNodes", "The selected component is part of a child actor template and cannot be reordered here.");
-					}
-					else if (!SelectedNodePtr->IsSceneComponent())
-					{
-						Message = LOCTEXT("DropActionToolTip_Error_CannotReparent_NotSceneComponent", "The selected component is not a scene component and cannot be attached to other components.");
-					}
-					else if (SelectedNodePtr->IsInheritedComponent())
-					{
-						Message = LOCTEXT("DropActionToolTip_Error_CannotReparent_Inherited", "The selected component is inherited and cannot be reordered here.");
-					}
-					else
-					{
-						Message = LOCTEXT("DropActionToolTip_Error_CannotReparent", "The selected component cannot be moved.");
-					}
-				}
-				else
-				{
-					Message = LOCTEXT("DropActionToolTip_Error_CannotReparentMultiple", "One or more of the selected components cannot be attached.");
-				}
-				break;
-			}
-		}
-
-		if (Message.IsEmpty())
-		{
-			FSCSEditorTreeNodePtrType SceneRootNodePtr = SCSEditor.Pin()->GetSceneRootNode();
-			check(SceneRootNodePtr.IsValid());
-
-			FSCSEditorTreeNodePtrType NodePtr = GetNode();
-			if (!NodePtr->IsComponentNode())
-			{
-				// Don't show a feedback message if over a node that makes no sense, such as a separator or the instance node
-				Message = LOCTEXT("DropActionToolTip_FriendlyError_DragToAComponent", "Drag to another component in order to attach to that component or become the root component.\nDrag to a Blueprint graph in order to drop a reference.");
-			}
-			else if (FChildActorComponentEditorUtils::IsChildActorSubtreeNode(NodePtr))
-			{
-				// Can't drag onto components within a child actor node's subtree
-				Message = LOCTEXT("DropActionToolTip_Error_ChildActorSubTree", "Cannot attach to this component as it is part of a child actor template.");
-			}
-
-			// Validate each selected node being dragged against the node that belongs to this row. Exit the loop if we have a valid tooltip OR a valid pending drop action once all nodes in the selection have been validated.
-			for (auto SourceNodeIter = DragRowOp->SourceNodes.CreateConstIterator(); SourceNodeIter && (Message.IsEmpty() || DragRowOp->PendingDropAction != FSCSRowDragDropOp::DropAction_None); ++SourceNodeIter)
-			{
-				FSCSEditorTreeNodePtrType DraggedNodePtr = *SourceNodeIter;
-				check(DraggedNodePtr.IsValid());
-
-				// Reset the pending drop action each time through the loop
-				DragRowOp->PendingDropAction = FSCSRowDragDropOp::DropAction_None;
-
-				// Get the component template objects associated with each node
-				USceneComponent* HoveredTemplate = Cast<USceneComponent>(NodePtr->GetComponentTemplate());
-				USceneComponent* DraggedTemplate = Cast<USceneComponent>(DraggedNodePtr->GetComponentTemplate());
-
-				if (DraggedNodePtr == NodePtr)
-				{
-					// Attempted to drag and drop onto self
-					if (DragRowOp->SourceNodes.Num() > 1)
-					{
-						Message = FText::Format(LOCTEXT("DropActionToolTip_Error_CannotAttachToSelfWithMultipleSelection", "Cannot attach the selected components here because it would result in {0} being attached to itself. Remove it from the selection and try again."), DraggedNodePtr->GetDisplayName());
-					}
-					else
-					{
-						Message = FText::Format(LOCTEXT("DropActionToolTip_Error_CannotAttachToSelf", "Cannot attach {0} to itself."), DraggedNodePtr->GetDisplayName());
-					}
-				}
-				else if (NodePtr->IsAttachedTo(DraggedNodePtr))
-				{
-					// Attempted to drop a parent onto a child
-					if (DragRowOp->SourceNodes.Num() > 1)
-					{
-						Message = FText::Format(LOCTEXT("DropActionToolTip_Error_CannotAttachToChildWithMultipleSelection", "Cannot attach the selected components here because it would result in {0} being attached to one of its children. Remove it from the selection and try again."), DraggedNodePtr->GetDisplayName());
-					}
-					else
-					{
-						Message = FText::Format(LOCTEXT("DropActionToolTip_Error_CannotAttachToChild", "Cannot attach {0} to one of its children."), DraggedNodePtr->GetDisplayName());
-					}
-				}
-				else if (HoveredTemplate == NULL || DraggedTemplate == NULL)
-				{
-					if (HoveredTemplate == nullptr)
-					{
-						// Can't attach to non-USceneComponent types
-						Message = LOCTEXT("DropActionToolTip_Error_NotAttachable_NotSceneComponent", "Cannot attach to this component as it is not a scene component.");
-					}
-					else
-					{
-						// Can't attach non-USceneComponent types
-						Message = LOCTEXT("DropActionToolTip_Error_NotAttachable", "Cannot attach to this component.");
-					}
-				}
-				else if (NodePtr == SceneRootNodePtr)
-				{
-					bool bCanMakeNewRoot = false;
-					bool bCanAttachToRoot = !DraggedNodePtr->IsDirectlyAttachedTo(NodePtr)
-						&& HoveredTemplate->CanAttachAsChild(DraggedTemplate, NAME_None)
-						&& DraggedTemplate->Mobility >= HoveredTemplate->Mobility
-						&& (!HoveredTemplate->IsEditorOnly() || DraggedTemplate->IsEditorOnly());
-
-					if (!NodePtr->CanReparent() && (!NodePtr->IsDefaultSceneRoot() || NodePtr->IsInheritedComponent()))
-					{
-						// Cannot make the dropped node the new root if we cannot reparent the current root
-						Message = LOCTEXT("DropActionToolTip_Error_CannotReparentRootNode", "The root component in this Blueprint is inherited and cannot be replaced.");
-					}
-					else if (DraggedTemplate->IsEditorOnly() && !HoveredTemplate->IsEditorOnly())
-					{
-						// can't have a new root that's editor-only (when children would be around in-game)
-						Message = LOCTEXT("DropActionToolTip_Error_CannotReparentEditorOnly", "Cannot re-parent game components under editor-only ones.");
-					}
-					else if (DraggedTemplate->Mobility > HoveredTemplate->Mobility)
-					{
-						// can't have a new root that's movable if the existing root is static or stationary
-						Message = LOCTEXT("DropActionToolTip_Error_CannotReparentNonMovable", "Cannot replace a non-movable scene root with a movable component.");
-					}
-					else if (DragRowOp->SourceNodes.Num() > 1)
-					{
-						Message = LOCTEXT("DropActionToolTip_Error_CannotAssignMultipleRootNodes", "Cannot replace the scene root with multiple components. Please select only a single component and try again.");
-					}
-					else
-					{
-						bCanMakeNewRoot = true;
-					}
-
-					if (bCanMakeNewRoot && bCanAttachToRoot)
-					{
-						// User can choose to either attach to the current root or make the dropped node the new root
-						Message = LOCTEXT("DropActionToolTip_AttachToOrMakeNewRoot", "Drop here to see available actions.");
-						DragRowOp->PendingDropAction = FSCSRowDragDropOp::DropAction_AttachToOrMakeNewRoot;
-					}
-					else if (SCSEditor.Pin()->GetEditorMode() == EComponentEditorMode::BlueprintSCS && DraggedNodePtr->GetBlueprint() != GetBlueprint())
-					{
-						if (bCanMakeNewRoot)
-						{
-							if (NodePtr->IsDefaultSceneRoot())
-							{
-								// Only available action is to copy the dragged node to the other Blueprint and make it the new root
-								// Default root will be deleted
-								Message = FText::Format(LOCTEXT("DropActionToolTip_DropMakeNewRootNodeFromCopyAndDelete", "Drop here to copy {0} to a new variable and make it the new root. The default root will be deleted."), DraggedNodePtr->GetDisplayName());
-							}
-							else
-							{
-								// Only available action is to copy the dragged node to the other Blueprint and make it the new root
-								Message = FText::Format(LOCTEXT("DropActionToolTip_DropMakeNewRootNodeFromCopy", "Drop here to copy {0} to a new variable and make it the new root."), DraggedNodePtr->GetDisplayName());
-							}
-							DragRowOp->PendingDropAction = FSCSRowDragDropOp::DropAction_MakeNewRoot;
-						}
-						else if (bCanAttachToRoot)
-						{
-							// Only available action is to copy the dragged node(s) to the other Blueprint and attach it to the root
-							if (DragRowOp->SourceNodes.Num() > 1)
-							{
-								Message = FText::Format(LOCTEXT("DropActionToolTip_AttachComponentsToThisNodeFromCopyWithMultipleSelection", "Drop here to copy the selected components to new variables and attach them to {0}."), NodePtr->GetDisplayName());
-							}
-							else
-							{
-								Message = FText::Format(LOCTEXT("DropActionToolTip_AttachToThisNodeFromCopy", "Drop here to copy {0} to a new variable and attach it to {1}."), DraggedNodePtr->GetDisplayName(), NodePtr->GetDisplayName());
-							}
-
-							DragRowOp->PendingDropAction = FSCSRowDragDropOp::DropAction_AttachTo;
-						}
-					}
-					else if (bCanMakeNewRoot)
-					{
-						if (NodePtr->IsDefaultSceneRoot())
-						{
-							// Only available action is to make the dragged node the new root
-							// Default root will be deleted
-							Message = FText::Format(LOCTEXT("DropActionToolTip_DropMakeNewRootNodeAndDelete", "Drop here to make {0} the new root. The default root will be deleted."), DraggedNodePtr->GetDisplayName());
-						}
-						else
-						{
-							// Only available action is to make the dragged node the new root
-							Message = FText::Format(LOCTEXT("DropActionToolTip_DropMakeNewRootNode", "Drop here to make {0} the new root."), DraggedNodePtr->GetDisplayName());
-						}
-						DragRowOp->PendingDropAction = FSCSRowDragDropOp::DropAction_MakeNewRoot;
-					}
-					else if (bCanAttachToRoot)
-					{
-						// Only available action is to attach the dragged node(s) to the root
-						if (DragRowOp->SourceNodes.Num() > 1)
-						{
-							Message = FText::Format(LOCTEXT("DropActionToolTip_AttachToThisNodeWithMultipleSelection", "Drop here to attach the selected components to {0}."), NodePtr->GetDisplayName());
-						}
-						else
-						{
-							Message = FText::Format(LOCTEXT("DropActionToolTip_AttachToThisNode", "Drop here to attach {0} to {1}."), DraggedNodePtr->GetDisplayName(), NodePtr->GetDisplayName());
-						}
-
-						DragRowOp->PendingDropAction = FSCSRowDragDropOp::DropAction_AttachTo;
-					}
-				}
-				else if (DraggedNodePtr->IsDirectlyAttachedTo(NodePtr)) // if dropped onto parent
-				{
-					// Detach the dropped node(s) from the current node and reattach to the root node
-					if (DragRowOp->SourceNodes.Num() > 1)
-					{
-						Message = FText::Format(LOCTEXT("DropActionToolTip_DetachFromThisNodeWithMultipleSelection", "Drop here to detach the selected components from {0}."), NodePtr->GetDisplayName());
-					}
-					else
-					{
-						Message = FText::Format(LOCTEXT("DropActionToolTip_DetachFromThisNode", "Drop here to detach {0} from {1}."), DraggedNodePtr->GetDisplayName(), NodePtr->GetDisplayName());
-					}
-
-					DragRowOp->PendingDropAction = FSCSRowDragDropOp::DropAction_DetachFrom;
-				}
-				else if (!DraggedTemplate->IsEditorOnly() && HoveredTemplate->IsEditorOnly())
-				{
-					// can't have a game component child nested under an editor-only one
-					Message = LOCTEXT("DropActionToolTip_Error_CannotAttachToEditorOnly", "Cannot attach game components to editor-only ones.");
-				}
-				else if ((DraggedTemplate->Mobility == EComponentMobility::Static) && ((HoveredTemplate->Mobility == EComponentMobility::Movable) || (HoveredTemplate->Mobility == EComponentMobility::Stationary)))
-				{
-					// Can't attach Static components to mobile ones
-					Message = LOCTEXT("DropActionToolTip_Error_CannotAttachStatic", "Cannot attach Static components to movable ones.");
-				}
-				else if ((DraggedTemplate->Mobility == EComponentMobility::Stationary) && (HoveredTemplate->Mobility == EComponentMobility::Movable))
-				{
-					// Can't attach Static components to mobile ones
-					Message = LOCTEXT("DropActionToolTip_Error_CannotAttachStationary", "Cannot attach Stationary components to movable ones.");
-				}
-				else if ((NodePtr->IsInstancedComponent() && HoveredTemplate->CreationMethod == EComponentCreationMethod::Native && !HoveredTemplate->HasAnyFlags(RF_DefaultSubObject)))
-				{
-					// Can't attach to post-construction C++-added components as they exist outside of the CDO and are not known at SCS execution time
-					Message = LOCTEXT("DropActionToolTip_Error_CannotAttachCPPAdded", "Cannot attach to components added in post-construction C++ code.");
-				}
-				else if (NodePtr->IsInstancedComponent() && HoveredTemplate->CreationMethod == EComponentCreationMethod::UserConstructionScript)
-				{
-					// Can't attach to UCS-added components as they exist outside of the CDO and are not known at SCS execution time
-					Message = LOCTEXT("DropActionToolTip_Error_CannotAttachUCSAdded", "Cannot attach to components added in the Construction Script.");
-				}
-				else if (HoveredTemplate->CanAttachAsChild(DraggedTemplate, NAME_None))
-				{
-					// Attach the dragged node(s) to this node
-					if (DraggedNodePtr->GetBlueprint() != GetBlueprint())
-					{
-						if (DragRowOp->SourceNodes.Num() > 1)
-						{
-							Message = FText::Format(LOCTEXT("DropActionToolTip_AttachToThisNodeFromCopyWithMultipleSelection", "Drop here to copy the selected nodes to new variables and attach them to {0}."), NodePtr->GetDisplayName());
-						}
-						else
-						{
-							Message = FText::Format(LOCTEXT("DropActionToolTip_AttachToThisNodeFromCopy", "Drop here to copy {0} to a new variable and attach it to {1}."), DraggedNodePtr->GetDisplayName(), NodePtr->GetDisplayName());
-						}
-					}
-					else if (DragRowOp->SourceNodes.Num() > 1)
-					{
-						Message = FText::Format(LOCTEXT("DropActionToolTip_AttachToThisNodeWithMultipleSelection", "Drop here to attach the selected components to {0}."), NodePtr->GetDisplayName());
-					}
-					else
-					{
-						Message = FText::Format(LOCTEXT("DropActionToolTip_AttachToThisNode", "Drop here to attach {0} to {1}."), DraggedNodePtr->GetDisplayName(), NodePtr->GetDisplayName());
-					}
-
-					DragRowOp->PendingDropAction = FSCSRowDragDropOp::DropAction_AttachTo;
-				}
-				else
-				{
-					// The dropped node cannot be attached to the current node
-					Message = FText::Format(LOCTEXT("DropActionToolTip_Error_TooManyAttachments", "Unable to attach {0} to {1}."), DraggedNodePtr->GetDisplayName(), NodePtr->GetDisplayName());
-				}
-			}
-		}
-
-		const FSlateBrush* StatusSymbol = DragRowOp->PendingDropAction != FSCSRowDragDropOp::DropAction_None
-			? FAppStyle::GetBrush(TEXT("Graph.ConnectorFeedback.OK"))
-			: FAppStyle::GetBrush(TEXT("Graph.ConnectorFeedback.Error"));
-
-		if (Message.IsEmpty())
-		{
-			DragRowOp->SetFeedbackMessage(nullptr);
-		}
-		else
-		{
-			DragRowOp->SetSimpleFeedbackMessage(StatusSymbol, FLinearColor::White, Message);
-		}
-	}
-	else if ( Operation->IsOfType<FExternalDragOperation>() || Operation->IsOfType<FAssetDragDropOp>() )
-	{
-		// defer to the tree widget's handler for this type of operation
-		TSharedPtr<SSCSEditor> PinnedEditor = SCSEditor.Pin();
-		if ( PinnedEditor.IsValid() && PinnedEditor->SCSTreeWidget.IsValid() )
-		{
-			// The widget geometry is irrelevant to the tree widget's OnDragEnter
-			PinnedEditor->SCSTreeWidget->OnDragEnter( FGeometry(), DragDropEvent );
-		}
-	}
-}
-
-void SSCS_RowWidget::HandleOnDragLeave(const FDragDropEvent& DragDropEvent)
-{
-	TSharedPtr<FSCSRowDragDropOp> DragRowOp = DragDropEvent.GetOperationAs<FSCSRowDragDropOp>();
-	if (DragRowOp.IsValid())
-	{
-		bool bCanReparentAllNodes = true;
-		for(auto SourceNodeIter = DragRowOp->SourceNodes.CreateConstIterator(); SourceNodeIter && bCanReparentAllNodes; ++SourceNodeIter)
-		{
-			FSCSEditorTreeNodePtrType DraggedNodePtr = *SourceNodeIter;
-			check(DraggedNodePtr.IsValid());
-
-			bCanReparentAllNodes = DraggedNodePtr->CanReparent();
-		}
-
-		// Only clear the tooltip text if all dragged nodes support it
-		if(bCanReparentAllNodes)
-		{
-			TSharedPtr<SWidget> NoWidget;
-			DragRowOp->SetFeedbackMessage(NoWidget);
-			DragRowOp->PendingDropAction = FSCSRowDragDropOp::DropAction_None;
-		}
-	}
-}
-
-TOptional<EItemDropZone> SSCS_RowWidget::HandleOnCanAcceptDrop(const FDragDropEvent& DragDropEvent, EItemDropZone DropZone, FSCSEditorTreeNodePtrType TargetItem)
-{
-	TOptional<EItemDropZone> ReturnDropZone;
-
-	TSharedPtr<FDragDropOperation> Operation = DragDropEvent.GetOperation();
-	if (Operation.IsValid())
-	{
-		if (Operation->IsOfType<FSCSRowDragDropOp>() && ( Cast<USceneComponent>(GetNode()->GetComponentTemplate()) != nullptr ))
-		{
-			TSharedPtr<FSCSRowDragDropOp> DragRowOp = StaticCastSharedPtr<FSCSRowDragDropOp>(Operation);
-			check(DragRowOp.IsValid());
-
-			if (DragRowOp->PendingDropAction != FSCSRowDragDropOp::DropAction_None)
-			{
-				ReturnDropZone = EItemDropZone::OntoItem;
-			}
-		}
-		else if (Operation->IsOfType<FExternalDragOperation>() || Operation->IsOfType<FAssetDragDropOp>())
-		{
-			ReturnDropZone = EItemDropZone::OntoItem;
-		}
-	}
-
-	return ReturnDropZone;
-}
-
-FReply SSCS_RowWidget::HandleOnAcceptDrop( const FDragDropEvent& DragDropEvent, EItemDropZone DropZone, FSCSEditorTreeNodePtrType TargetItem )
-{
-	TSharedPtr<FDragDropOperation> Operation = DragDropEvent.GetOperation();
-	if (!Operation.IsValid())
-	{
-		return FReply::Handled();
-	}
-	
-	if (Operation->IsOfType<FSCSRowDragDropOp>() && (Cast<USceneComponent>(GetNode()->GetComponentTemplate()) != nullptr))
-	{
-		TSharedPtr<FSCSRowDragDropOp> DragRowOp = StaticCastSharedPtr<FSCSRowDragDropOp>( Operation );	
-		check(DragRowOp.IsValid());
-
-		switch(DragRowOp->PendingDropAction)
-		{
-		case FSCSRowDragDropOp::DropAction_AttachTo:
-			OnAttachToDropAction(DragRowOp->SourceNodes);
-			break;
-			
-		case FSCSRowDragDropOp::DropAction_DetachFrom:
-			OnDetachFromDropAction(DragRowOp->SourceNodes);
-			break;
-
-		case FSCSRowDragDropOp::DropAction_MakeNewRoot:
-			check(DragRowOp->SourceNodes.Num() == 1);
-			OnMakeNewRootDropAction(DragRowOp->SourceNodes[0]);
-			break;
-
-		case FSCSRowDragDropOp::DropAction_AttachToOrMakeNewRoot:
-			{
-				check(DragRowOp->SourceNodes.Num() == 1);
-				FWidgetPath WidgetPath = DragDropEvent.GetEventPath() != nullptr ? *DragDropEvent.GetEventPath() : FWidgetPath();
-				FSlateApplication::Get().PushMenu(
-					SharedThis(this),
-					WidgetPath,
-					BuildSceneRootDropActionMenu(DragRowOp->SourceNodes[0]).ToSharedRef(),
-					FSlateApplication::Get().GetCursorPos(),
-					FPopupTransitionEffect(FPopupTransitionEffect::TypeInPopup)
-				);
-			}
-			break;
-
-		case FSCSRowDragDropOp::DropAction_None:
-		default:
-			break;
-		}
-	}
-	else if (Operation->IsOfType<FExternalDragOperation>() || Operation->IsOfType<FAssetDragDropOp>())
-	{
-		// defer to the tree widget's handler for this type of operation
-		TSharedPtr<SSCSEditor> PinnedEditor = SCSEditor.Pin();
-		if ( PinnedEditor.IsValid() && PinnedEditor->SCSTreeWidget.IsValid() )
-		{
-			// The widget geometry is irrelevant to the tree widget's OnDrop
-			PinnedEditor->SCSTreeWidget->OnDrop( FGeometry(), DragDropEvent );
-		}
-	}
-
-	return FReply::Handled();
-}
-
-void ConformTransformRelativeToParent(USceneComponent* SceneComponentTemplate, USceneComponent* ParentSceneComponent)
-{
-	// If we find a match, calculate its new position relative to the scene root component instance in its current scene
-	FTransform ComponentToWorld(SceneComponentTemplate->GetRelativeRotation(), SceneComponentTemplate->GetRelativeLocation(), SceneComponentTemplate->GetRelativeScale3D());
-	FTransform ParentToWorld = (SceneComponentTemplate->GetAttachSocketName() != NAME_None) ? ParentSceneComponent->GetSocketTransform(SceneComponentTemplate->GetAttachSocketName(), RTS_World) : ParentSceneComponent->GetComponentToWorld();
-	FTransform RelativeTM = ComponentToWorld.GetRelativeTransform(ParentToWorld);
-
-	// Store new relative location value (if not set to absolute)
-	if (!SceneComponentTemplate->IsUsingAbsoluteLocation())
-	{
-		SceneComponentTemplate->SetRelativeLocation_Direct(RelativeTM.GetTranslation());
-	}
-
-	// Store new relative rotation value (if not set to absolute)
-	if (!SceneComponentTemplate->IsUsingAbsoluteRotation())
-	{
-		SceneComponentTemplate->SetRelativeRotation_Direct(RelativeTM.Rotator());
-	}
-
-	// Store new relative scale value (if not set to absolute)
-	if (!SceneComponentTemplate->IsUsingAbsoluteScale())
-	{
-		SceneComponentTemplate->SetRelativeScale3D_Direct(RelativeTM.GetScale3D());
-	}
-}
-
-void SSCS_RowWidget::OnAttachToDropAction(const TArray<FSCSEditorTreeNodePtrType>& DroppedNodePtrs)
-{
-	FSCSEditorTreeNodePtrType NodePtr = GetNode();
-
-	check(NodePtr.IsValid());
-	check(DroppedNodePtrs.Num() > 0);
-
-	TSharedPtr<SSCSEditor> SCSEditorPtr = SCSEditor.Pin();
-	check(SCSEditorPtr.IsValid());
-
-	bool bRegenerateTreeNodes = false;
-	const FScopedTransaction TransactionContext(DroppedNodePtrs.Num() > 1 ? LOCTEXT("AttachComponents", "Attach Components") : LOCTEXT("AttachComponent", "Attach Component"));
-
-	if (SCSEditorPtr->GetEditorMode() == EComponentEditorMode::BlueprintSCS)
-	{
-		// Get the current Blueprint context
-		UBlueprint* Blueprint = GetBlueprint();
-		check(Blueprint);
-
-		// Get the current "preview" Actor instance
-		AActor* PreviewActor = SCSEditorPtr->PreviewActor.Get();
-		check(PreviewActor);
-
-		for(const FSCSEditorTreeNodePtrType& DroppedNodePtr : DroppedNodePtrs)
-		{
-			// Clone the component if it's being dropped into a different SCS
-			if(DroppedNodePtr->GetBlueprint() != Blueprint)
-			{
-				bRegenerateTreeNodes = true;
-
-				check(DroppedNodePtr.IsValid());
-				UActorComponent* ComponentTemplate = DroppedNodePtr->GetComponentTemplate();
-				check(ComponentTemplate);
-
-				// Note: This will mark the Blueprint as structurally modified
-				UActorComponent* ClonedComponent = SCSEditorPtr->AddNewComponent(ComponentTemplate->GetClass(), nullptr);
-				check(ClonedComponent);
-
-				UEngine::CopyPropertiesForUnrelatedObjects(ComponentTemplate, ClonedComponent);
-
-				// Attach the copied node to the target node (this will also detach it from the root if necessary)
-				FSCSEditorTreeNodePtrType NewNodePtr = SCSEditorPtr->GetNodeFromActorComponent(ClonedComponent);
-				if(NewNodePtr.IsValid())
-				{
-					NodePtr->AddChild(NewNodePtr);
-				}
-			}
-			else
-			{
-				// Get the associated component template if it is a scene component, so we can adjust the transform
-				USceneComponent* SceneComponentTemplate = Cast<USceneComponent>(DroppedNodePtr->GetComponentTemplate());
-
-				// Cache current default values for propagation
-				FVector OldRelativeLocation, OldRelativeScale3D;
-				FRotator OldRelativeRotation;
-				if(SceneComponentTemplate)
-				{
-					OldRelativeLocation = SceneComponentTemplate->GetRelativeLocation();
-					OldRelativeRotation = SceneComponentTemplate->GetRelativeRotation();
-					OldRelativeScale3D = SceneComponentTemplate->GetRelativeScale3D();
-				}
-
-				// Check for a valid parent node
-				FSCSEditorTreeNodePtrType ParentNodePtr = DroppedNodePtr->GetParent();
-				if(ParentNodePtr.IsValid())
-				{
-					// Detach the dropped node from its parent
-					ParentNodePtr->RemoveChild(DroppedNodePtr);
-
-					// If the associated component template is a scene component, maintain its preview world position
-					if(SceneComponentTemplate)
-					{
-						// Save current state
-						SceneComponentTemplate->Modify();
-
-						// Reset the attach socket name
-						SceneComponentTemplate->SetupAttachment(SceneComponentTemplate->GetAttachParent(), NAME_None);
-						USCS_Node* SCS_Node = DroppedNodePtr->GetSCSNode();
-						if(SCS_Node)
-						{
-							SCS_Node->Modify();
-							SCS_Node->AttachToName = NAME_None;
-						}
-
-						// Attempt to locate a matching registered instance of the component template in the Actor context that's being edited
-						USceneComponent* InstancedSceneComponent = Cast<USceneComponent>(DroppedNodePtr->FindComponentInstanceInActor(PreviewActor));
-						if(InstancedSceneComponent && InstancedSceneComponent->IsRegistered())
-						{
-							// If we find a match, save off the world position
-							const FTransform& ComponentToWorld = InstancedSceneComponent->GetComponentToWorld();
-							SceneComponentTemplate->SetRelativeTransform_Direct(ComponentToWorld);
-						}
-					}
-				}
-
-				// Attach the dropped node to the given node
-				NodePtr->AddChild(DroppedNodePtr);
-
-				// Attempt to locate a matching instance of the parent component template in the Actor context that's being edited
-				USceneComponent* ParentSceneComponent = Cast<USceneComponent>(NodePtr->FindComponentInstanceInActor(PreviewActor));
-				if(SceneComponentTemplate && ParentSceneComponent && ParentSceneComponent->IsRegistered())
-				{
-					ConformTransformRelativeToParent(SceneComponentTemplate, ParentSceneComponent);
-				}
-
-				// Propagate any default value changes out to all instances of the template. If we didn't do this, then instances could incorrectly override the new default value with the old default value when construction scripts are re-run.
-				if(SceneComponentTemplate)
-				{
-					TArray<UObject*> InstancedSceneComponents;
-					SceneComponentTemplate->GetArchetypeInstances(InstancedSceneComponents);
-					for(int32 InstanceIndex = 0; InstanceIndex < InstancedSceneComponents.Num(); ++InstanceIndex)
-					{
-						USceneComponent* InstancedSceneComponent = Cast<USceneComponent>(InstancedSceneComponents[InstanceIndex]);
-						if(InstancedSceneComponent != nullptr)
-						{
-							FComponentEditorUtils::ApplyDefaultValueChange(InstancedSceneComponent, InstancedSceneComponent->GetRelativeLocation_DirectMutable(), OldRelativeLocation, SceneComponentTemplate->GetRelativeLocation());
-							FComponentEditorUtils::ApplyDefaultValueChange(InstancedSceneComponent, InstancedSceneComponent->GetRelativeRotation_DirectMutable(), OldRelativeRotation, SceneComponentTemplate->GetRelativeRotation());
-							FComponentEditorUtils::ApplyDefaultValueChange(InstancedSceneComponent, InstancedSceneComponent->GetRelativeScale3D_DirectMutable(),  OldRelativeScale3D,  SceneComponentTemplate->GetRelativeScale3D());
-						}
-					}
-				}
-			}
-		}
-	}
-	else    // EComponentEditorMode::ActorInstance
-	{
-		for(const FSCSEditorTreeNodePtrType& DroppedNodePtr : DroppedNodePtrs)
-		{
-			// Check for a valid parent node
-			FSCSEditorTreeNodePtrType ParentNodePtr = DroppedNodePtr->GetParent();
-			if(ParentNodePtr.IsValid())
-			{
-				// Detach the dropped node from its parent
-				ParentNodePtr->RemoveChild(DroppedNodePtr);
-			}
-
-			// Attach the dropped node to the given node
-			NodePtr->AddChild(DroppedNodePtr);
-		}
-	}
-
-	check(SCSEditorPtr->SCSTreeWidget.IsValid());
-	SCSEditorPtr->SCSTreeWidget->SetItemExpansion(NodePtr, true);
-
-	PostDragDropAction(bRegenerateTreeNodes);
-}
-
-void SSCS_RowWidget::OnDetachFromDropAction(const TArray<FSCSEditorTreeNodePtrType>& DroppedNodePtrs)
-{
-	check(DroppedNodePtrs.Num() > 0);
-
-	TSharedPtr<SSCSEditor> SCSEditorPtr = SCSEditor.Pin();
-	check(SCSEditorPtr.IsValid());
-
-	const FScopedTransaction TransactionContext(DroppedNodePtrs.Num() > 1 ? LOCTEXT("DetachComponents", "Detach Components") : LOCTEXT("DetachComponent", "Detach Component"));
-
-	if (SCSEditorPtr->GetEditorMode() == EComponentEditorMode::BlueprintSCS)
-	{
-		// Get the current "preview" Actor instance
-		AActor* PreviewActor = SCSEditorPtr->PreviewActor.Get();
-		check(PreviewActor);
-
-		for (const FSCSEditorTreeNodePtrType& DroppedNodePtr : DroppedNodePtrs)
-		{
-			FVector OldRelativeLocation, OldRelativeScale3D;
-			FRotator OldRelativeRotation;
-
-			check(DroppedNodePtr.IsValid());
-
-			// Detach the node from its parent
-			FSCSEditorTreeNodePtrType ParentNodePtr = DroppedNodePtr->GetParent();
-			check(ParentNodePtr.IsValid());
-			ParentNodePtr->RemoveChild(DroppedNodePtr);
-
-			// If the associated component template is a scene component, maintain its current world position
-			USceneComponent* SceneComponentTemplate = Cast<USceneComponent>(DroppedNodePtr->GetComponentTemplate());
-			if(SceneComponentTemplate)
-			{
-				// Cache current default values for propagation
-				OldRelativeLocation = SceneComponentTemplate->GetRelativeLocation();
-				OldRelativeRotation = SceneComponentTemplate->GetRelativeRotation();
-				OldRelativeScale3D = SceneComponentTemplate->GetRelativeScale3D();
-
-				// Save current state
-				SceneComponentTemplate->Modify();
-
-				// Reset the attach socket name
-				SceneComponentTemplate->SetupAttachment(SceneComponentTemplate->GetAttachParent(), NAME_None);
-				USCS_Node* SCS_Node = DroppedNodePtr->GetSCSNode();
-				if(SCS_Node)
-				{
-					SCS_Node->Modify();
-					SCS_Node->AttachToName = NAME_None;
-				}
-
-				// Attempt to locate a matching instance of the component template in the Actor context that's being edited
-				USceneComponent* InstancedSceneComponent = Cast<USceneComponent>(DroppedNodePtr->FindComponentInstanceInActor(PreviewActor));
-				if(InstancedSceneComponent && InstancedSceneComponent->IsRegistered())
-				{
-					// If we find a match, save off the world position
-					const FTransform& ComponentToWorld = InstancedSceneComponent->GetComponentToWorld();
-					SceneComponentTemplate->SetRelativeTransform_Direct(ComponentToWorld);
-				}
-			}
-
-			// Attach the dropped node to the current scene root node
-			FSCSEditorTreeNodePtrType SceneRootNodePtr = SCSEditorPtr->GetSceneRootNode();
-			check(SceneRootNodePtr.IsValid());
-			SceneRootNodePtr->AddChild(DroppedNodePtr);
-
-			// Attempt to locate a matching instance of the scene root component template in the Actor context that's being edited
-			USceneComponent* InstancedSceneRootComponent = Cast<USceneComponent>(SceneRootNodePtr->FindComponentInstanceInActor(PreviewActor));
-			if(SceneComponentTemplate && InstancedSceneRootComponent && InstancedSceneRootComponent->IsRegistered())
-			{
-				ConformTransformRelativeToParent(SceneComponentTemplate, InstancedSceneRootComponent);
-			}
-
-			// Propagate any default value changes out to all instances of the template. If we didn't do this, then instances could incorrectly override the new default value with the old default value when construction scripts are re-run.
-			if(SceneComponentTemplate)
-			{
-				TArray<UObject*> InstancedSceneComponents;
-				SceneComponentTemplate->GetArchetypeInstances(InstancedSceneComponents);
-				for(int32 InstanceIndex = 0; InstanceIndex < InstancedSceneComponents.Num(); ++InstanceIndex)
-				{
-					USceneComponent* InstancedSceneComponent = Cast<USceneComponent>(InstancedSceneComponents[InstanceIndex]);
-					if(InstancedSceneComponent != nullptr)
-					{
-						FComponentEditorUtils::ApplyDefaultValueChange(InstancedSceneComponent, InstancedSceneComponent->GetRelativeLocation_DirectMutable(), OldRelativeLocation, SceneComponentTemplate->GetRelativeLocation());
-						FComponentEditorUtils::ApplyDefaultValueChange(InstancedSceneComponent, InstancedSceneComponent->GetRelativeRotation_DirectMutable(), OldRelativeRotation, SceneComponentTemplate->GetRelativeRotation());
-						FComponentEditorUtils::ApplyDefaultValueChange(InstancedSceneComponent, InstancedSceneComponent->GetRelativeScale3D_DirectMutable(),  OldRelativeScale3D,  SceneComponentTemplate->GetRelativeScale3D());
-					}
-				}
-			}
-		}
-	}
-	else    // EComponentEditorMode::ActorInstance
-	{
-		for (const FSCSEditorTreeNodePtrType& DroppedNodePtr : DroppedNodePtrs)
-		{
-			check(DroppedNodePtr.IsValid());
-
-			// Detach the node from its parent
-			FSCSEditorTreeNodePtrType ParentNodePtr = DroppedNodePtr->GetParent();
-			check(ParentNodePtr.IsValid());
-			ParentNodePtr->RemoveChild(DroppedNodePtr);
-
-			// Attach the dropped node to the current scene root node
-			FSCSEditorTreeNodePtrType SceneRootNodePtr = SCSEditorPtr->GetSceneRootNode();
-			check(SceneRootNodePtr.IsValid());
-			SceneRootNodePtr->AddChild(DroppedNodePtr);
-		}
-	}
-	
-	PostDragDropAction(false);
-}
-
-void SSCS_RowWidget::OnMakeNewRootDropAction(FSCSEditorTreeNodePtrType DroppedNodePtr)
-{
-	TSharedPtr<SSCSEditor> SCSEditorPtr = SCSEditor.Pin();
-	check(SCSEditorPtr.IsValid());
-
-	// Get the current scene root node
-	FSCSEditorTreeNodePtrType SceneRootNodePtr = SCSEditorPtr->GetSceneRootNode();
-
-	FSCSEditorTreeNodePtrType NodePtr = GetNode();
-
-	// We cannot handle the drop action if any of these conditions fail on entry.
-	if (!ensure(NodePtr.IsValid()) || !ensure(DroppedNodePtr.IsValid()) || !ensure(NodePtr == SceneRootNodePtr))
-	{
-		return;
-	}
-
-	// Create a transaction record
-	const FScopedTransaction TransactionContext(LOCTEXT("MakeNewSceneRoot", "Make New Scene Root"));
-
-	FSCSEditorTreeNodePtrType OldSceneRootNodePtr;
-
-	// Remember whether or not we're replacing the default scene root
-	bool bWasDefaultSceneRoot = SceneRootNodePtr.IsValid() && SceneRootNodePtr->IsDefaultSceneRoot();
-
-	if (SCSEditorPtr->GetEditorMode() == EComponentEditorMode::BlueprintSCS)
-	{
-		// Get the current Blueprint context
-		UBlueprint* Blueprint = GetBlueprint();
-		check(Blueprint && Blueprint->SimpleConstructionScript);
-
-		// Clone the component if it's being dropped into a different SCS
-		if(DroppedNodePtr->GetBlueprint() != Blueprint)
-		{
-			UActorComponent* ComponentTemplate = DroppedNodePtr->GetComponentTemplate();
-			check(ComponentTemplate);
-
-			// Note: This will mark the Blueprint as structurally modified
-			UActorComponent* ClonedComponent = SCSEditorPtr->AddNewComponent(ComponentTemplate->GetClass(), nullptr);
-			check(ClonedComponent);
-
-			UEngine::CopyPropertiesForUnrelatedObjects(ComponentTemplate, ClonedComponent);
-
-			DroppedNodePtr = SCSEditorPtr->GetNodeFromActorComponent(ClonedComponent);
-			check(DroppedNodePtr.IsValid());
-		}
-
-		if(DroppedNodePtr->GetParent().IsValid()
-			&& DroppedNodePtr->GetBlueprint() == Blueprint)
-		{
-			// If the associated component template is a scene component, reset its transform since it will now become the root
-			USceneComponent* SceneComponentTemplate = Cast<USceneComponent>(DroppedNodePtr->GetComponentTemplate());
-			if(SceneComponentTemplate)
-			{
-				// Save current state
-				SceneComponentTemplate->Modify();
-
-				// Reset the attach socket name
-				SceneComponentTemplate->SetupAttachment(SceneComponentTemplate->GetAttachParent(), NAME_None);
-				USCS_Node* SCS_Node = DroppedNodePtr->GetSCSNode();
-				if(SCS_Node)
-				{
-					SCS_Node->Modify();
-					SCS_Node->AttachToName = NAME_None;
-				}
-
-				// Cache the current relative location and rotation values (for propagation)
-				const FVector OldRelativeLocation = SceneComponentTemplate->GetRelativeLocation();
-				const FRotator OldRelativeRotation = SceneComponentTemplate->GetRelativeRotation();
-
-				// Reset the relative transform (location and rotation only; scale is preserved)
-				SceneComponentTemplate->SetRelativeLocation(FVector::ZeroVector);
-				SceneComponentTemplate->SetRelativeRotation(FRotator::ZeroRotator);
-
-				// Propagate the root change & detachment to any instances of the template (done within the context of the current transaction)
-				TArray<UObject*> ArchetypeInstances;
-				SceneComponentTemplate->GetArchetypeInstances(ArchetypeInstances);
-				FDetachmentTransformRules DetachmentTransformRules(EDetachmentRule::KeepWorld, EDetachmentRule::KeepWorld, EDetachmentRule::KeepRelative, true);
-				for (int32 InstanceIndex = 0; InstanceIndex < ArchetypeInstances.Num(); ++InstanceIndex)
-				{
-					USceneComponent* SceneComponentInstance = Cast<USceneComponent>(ArchetypeInstances[InstanceIndex]);
-					if (SceneComponentInstance != nullptr)
-					{
-						// Detach from root (keeping world transform, except for scale)
-						SceneComponentInstance->DetachFromComponent(DetachmentTransformRules);
-
-						// Propagate the default relative location & rotation reset from the template to the instance
-						FComponentEditorUtils::ApplyDefaultValueChange(SceneComponentInstance, SceneComponentInstance->GetRelativeLocation_DirectMutable(), OldRelativeLocation, SceneComponentTemplate->GetRelativeLocation());
-						FComponentEditorUtils::ApplyDefaultValueChange(SceneComponentInstance, SceneComponentInstance->GetRelativeRotation_DirectMutable(), OldRelativeRotation, SceneComponentTemplate->GetRelativeRotation());
-
-						// Must also reset the root component here, so that RerunConstructionScripts() will cache the correct root component instance data
-						AActor* Owner = SceneComponentInstance->GetOwner();
-						if (Owner)
-						{
-							Owner->Modify();
-							Owner->SetRootComponent(SceneComponentInstance);
-						}
-					}
-				}
-			}
-
-			// Remove the dropped node from its existing parent
-			DroppedNodePtr->GetParent()->RemoveChild(DroppedNodePtr);
-		}
-
-		check(bWasDefaultSceneRoot || SceneRootNodePtr->CanReparent());
-
-		// Remove the current scene root node from the SCS context
-		Blueprint->SimpleConstructionScript->RemoveNode(SceneRootNodePtr->GetSCSNode(), /*bValidateSceneRootNodes=*/false);
-
-		// Save old root node
-		OldSceneRootNodePtr = SceneRootNodePtr;
-
-		// Set node we are dropping as new root
-		SceneRootNodePtr = DroppedNodePtr;
-		SCSEditorPtr->SetSceneRootNode(SceneRootNodePtr);
-
-		// Add dropped node to the SCS context
-		Blueprint->SimpleConstructionScript->AddNode(SceneRootNodePtr->GetSCSNode());
-
-		// Remove or re-parent the old root
-		if (OldSceneRootNodePtr.IsValid())
-		{
-			check(SceneRootNodePtr->CanReparent());
-
-			// Set old root as child of new root
-			SceneRootNodePtr->AddChild(OldSceneRootNodePtr);
-
-			// Expand the new scene root as we've just added a child to it
-			SCSEditorPtr->SetNodeExpansionState(SceneRootNodePtr, true);
-
-			if (bWasDefaultSceneRoot)
-			{
-				SCSEditorPtr->RemoveComponentNode(OldSceneRootNodePtr);
-			}
-		}
-	}
-	else    // EComponentEditorMode::ActorInstance
-	{
-		if(DroppedNodePtr->GetParent().IsValid())
-		{
-			// Remove the dropped node from its existing parent
-			DroppedNodePtr->GetParent()->RemoveChild(DroppedNodePtr);
-		}
-
-		// Save old root node
-		OldSceneRootNodePtr = SceneRootNodePtr;
-
-		// Set node we are dropping as new root
-		SceneRootNodePtr = DroppedNodePtr;
-		SCSEditorPtr->SetSceneRootNode(SceneRootNodePtr);
-
-		// Remove or re-parent the old root
-		if (OldSceneRootNodePtr.IsValid())
-		{
-			if (bWasDefaultSceneRoot)
-			{
-				SCSEditorPtr->RemoveComponentNode(OldSceneRootNodePtr);
-				SCSEditorPtr->GetActorContext()->SetRootComponent(CastChecked<USceneComponent>(DroppedNodePtr->GetComponentTemplate()));
-			}
-			else
-			{
-				check(SceneRootNodePtr->CanReparent());
-
-				// Set old root as child of new root
-				SceneRootNodePtr->AddChild(OldSceneRootNodePtr);
-
-				// Expand the new scene root as we've just added a child to it
-				SCSEditorPtr->SetNodeExpansionState(SceneRootNodePtr, true);
-			}
-		}
-	}
-
-	PostDragDropAction(true);
-}
-
-void SSCS_RowWidget::PostDragDropAction(bool bRegenerateTreeNodes)
-{
-	GUnrealEd->ComponentVisManager.ClearActiveComponentVis();
-
-	FSCSEditorTreeNodePtrType NodePtr = GetNode();
-
-	TSharedPtr<SSCSEditor> PinnedEditor = SCSEditor.Pin();
-	if(PinnedEditor.IsValid())
-	{
-		PinnedEditor->UpdateTree(bRegenerateTreeNodes);
-
-		PinnedEditor->RefreshSelectionDetails();
-
-		if (PinnedEditor->GetEditorMode() == EComponentEditorMode::BlueprintSCS)
-		{
-			if(NodePtr.IsValid())
-			{
-				UBlueprint* Blueprint = GetBlueprint();
-				if(Blueprint != nullptr)
-				{
-					FBlueprintEditorUtils::PostEditChangeBlueprintActors(Blueprint, true);
-				}
-			}
-		}
-		else
-		{
-			AActor* ActorInstance = PinnedEditor->GetActorContext();
-			if(ActorInstance)
-			{
-				ActorInstance->RerunConstructionScripts();
-			}
-		}
-	}
-}
-
-FText SSCS_RowWidget::GetNameLabel() const
-{
-	if( InlineWidget.IsValid() && !InlineWidget->IsInEditMode() )
-	{
-		FSCSEditorTreeNodePtrType NodePtr = GetNode();
-		if(NodePtr->IsInheritedComponent())
-		{
-			return FText::Format(LOCTEXT("NativeComponentFormatString","{0} (Inherited)"), FText::FromString(GetNode()->GetDisplayString()));
-		}
-	}
-
-	// NOTE: Whatever this returns also becomes the variable name
-	return FText::FromString(GetNode()->GetDisplayString());
-}
-
-FText SSCS_RowWidget::GetTooltipText() const
-{
-	FSCSEditorTreeNodePtrType NodePtr = GetNode();
-
-	if (NodePtr->IsDefaultSceneRoot())
-	{
-		if (NodePtr->IsInheritedComponent())
-		{
-			return LOCTEXT("InheritedDefaultSceneRootToolTip", "This is the default scene root component. It cannot be copied, renamed or deleted.\nIt has been inherited from the parent class, so its properties cannot be edited here.\nNew scene components will automatically be attached to it.");
-		}
-		else
-		{
-			return LOCTEXT("DefaultSceneRootToolTip", "This is the default scene root component. It cannot be copied, renamed or deleted.\nIt can be replaced by drag/dropping another scene component over it.");
-		}
-	}
-	else
-	{
-		UClass* Class = ( NodePtr->GetComponentTemplate() != nullptr ) ? NodePtr->GetComponentTemplate()->GetClass() : nullptr;
-		const FText ClassDisplayName = FBlueprintEditorUtils::GetFriendlyClassDisplayName(Class);
-		const FText ComponentDisplayName = NodePtr->GetDisplayName();
-
-
-		FFormatNamedArguments Args;
-		Args.Add(TEXT("ClassName"), ClassDisplayName);
-		Args.Add(TEXT("NodeName"), FText::FromString(NodePtr->GetDisplayString()));
-
-		return FText::Format(LOCTEXT("ComponentTooltip", "{NodeName} ({ClassName})"), Args);
-	}
-}
-
-FString SSCS_RowWidget::GetDocumentationLink() const
-{
-	check(SCSEditor.IsValid());
-
-	FSCSEditorTreeNodePtrType NodePtr = GetNode();
-	if ((NodePtr == SCSEditor.Pin()->GetSceneRootNode()) || NodePtr->IsInheritedComponent())
-	{
-		return TEXT("Shared/Editors/BlueprintEditor/ComponentsMode");
-	}
-
-	return TEXT("");
-}
-
-FString SSCS_RowWidget::GetDocumentationExcerptName() const
-{
-	check(SCSEditor.IsValid());
-
-	FSCSEditorTreeNodePtrType NodePtr = GetNode();
-	if (NodePtr == SCSEditor.Pin()->GetSceneRootNode())
-	{
-		return TEXT("RootComponent");
-	}
-	else if (NodePtr->IsNativeComponent())
-	{
-		return TEXT("NativeComponents");
-	}
-	else if (NodePtr->IsInheritedComponent())
-	{
-		return TEXT("InheritedComponents");
-	}
-
-	return TEXT("");
-}
-
-UBlueprint* SSCS_RowWidget::GetBlueprint() const
-{
-	check(SCSEditor.IsValid());
-	return SCSEditor.Pin()->GetBlueprint();
-}
-
-ESelectionMode::Type SSCS_RowWidget::GetSelectionMode() const
-{
-	FSCSEditorTreeNodePtrType NodePtr = GetNode();
-	if (NodePtr->GetNodeType() == FSCSEditorTreeNode::SeparatorNode)
-	{
-		return ESelectionMode::None;
-	}
-	
-	return SMultiColumnTableRow<FSCSEditorTreeNodePtrType>::GetSelectionMode();
-}
-
-bool SSCS_RowWidget::OnNameTextVerifyChanged(const FText& InNewText, FText& OutErrorMessage)
-{
-	FSCSEditorTreeNodePtrType NodePtr = GetNode();
-	UBlueprint* Blueprint = GetBlueprint();
-
-	const FString& NewTextStr = InNewText.ToString();
-
-	if (!NewTextStr.IsEmpty())
-	{
-		if (NodePtr->GetVariableName().ToString() == NewTextStr)
-		{
-			return true;
-		}
-
-		const UActorComponent* ComponentInstance = NodePtr->GetComponentTemplate();
-		if (ensure(ComponentInstance))
-		{
-			AActor* ExistingNameSearchScope = ComponentInstance->GetOwner();
-			if ((ExistingNameSearchScope == nullptr) && (Blueprint != nullptr))
-			{
-				ExistingNameSearchScope = Cast<AActor>(Blueprint->GeneratedClass->GetDefaultObject());
-			}
-
-			if (!FComponentEditorUtils::IsValidVariableNameString(ComponentInstance, NewTextStr))
-			{
-				OutErrorMessage = LOCTEXT("RenameFailed_EngineReservedName", "This name is reserved for engine use.");
-				return false;
-			}
-			else if (NewTextStr.Len() >= NAME_SIZE)
-			{
-				FFormatNamedArguments Arguments;
-				Arguments.Add(TEXT("CharCount"), NAME_SIZE);
-				OutErrorMessage = FText::Format(LOCTEXT("ComponentRenameFailed_TooLong", "Component name must be less than {CharCount} characters long."), Arguments);
-				return false;
-			}
-			else if (!FComponentEditorUtils::IsComponentNameAvailable(NewTextStr, ExistingNameSearchScope, ComponentInstance) 
-					|| !FComponentEditorUtils::IsComponentNameAvailable(NewTextStr, ComponentInstance->GetOuter(), ComponentInstance ))
-			{
-				OutErrorMessage = LOCTEXT("RenameFailed_ExistingName", "Another component already has the same name.");
-				return false;
-			}
-		}
-		else
-		{
-			OutErrorMessage = LOCTEXT("RenameFailed_InvalidComponentInstance", "This node is referencing an invalid component instance and cannot be renamed. Perhaps it was destroyed?");
-			return false;
-		}
-	}
-
-	TSharedPtr<INameValidatorInterface> NameValidator;
-	if (Blueprint != nullptr)
-	{
-		NameValidator = MakeShareable(new FKismetNameValidator(GetBlueprint(), NodePtr->GetVariableName()));
-	}
-	else
-	{
-		NameValidator = MakeShareable(new FStringSetNameValidator(NodePtr->GetComponentTemplate()->GetName()));
-	}
-
-	EValidatorResult ValidatorResult = NameValidator->IsValid(NewTextStr);
-	if (ValidatorResult == EValidatorResult::AlreadyInUse)
-	{
-		OutErrorMessage = FText::Format(LOCTEXT("RenameFailed_InUse", "{0} is in use by another variable or function!"), InNewText);
-	}
-	else if (ValidatorResult == EValidatorResult::EmptyName)
-	{
-		OutErrorMessage = LOCTEXT("RenameFailed_LeftBlank", "Names cannot be left blank!");
-	}
-	else if (ValidatorResult == EValidatorResult::TooLong)
-	{
-		OutErrorMessage = LOCTEXT("RenameFailed_NameTooLong", "Names must have fewer than 100 characters!");
-	}
-
-	if (OutErrorMessage.IsEmpty())
-	{
-		return true;
-	}
-
-	return false;
-}
-
-void SSCS_RowWidget::OnNameTextCommit(const FText& InNewName, ETextCommit::Type InTextCommit)
-{
-	GetNode()->OnCompleteRename(InNewName);
-
-	// No need to call UpdateTree() in SCS editor mode; it will already be called by MBASM internally
-	check(SCSEditor.IsValid());
-	TSharedPtr<SSCSEditor> PinnedEditor = SCSEditor.Pin();
-	if (PinnedEditor.IsValid() && PinnedEditor->GetEditorMode() == EComponentEditorMode::ActorInstance)
-	{
-		PinnedEditor->UpdateTree(false);
-	}
-}
-
-//////////////////////////////////////////////////////////////////////////
-// SSCS_RowWidget_ActorRoot
-
-FSCSEditorActorNodePtrType SSCS_RowWidget_ActorRoot::GetActorNode() const
-{
-	return StaticCastSharedPtr<FSCSEditorTreeNodeActorBase>(GetNode());
-}
-
-TSharedRef<SWidget> SSCS_RowWidget_ActorRoot::GenerateWidgetForColumn(const FName& ColumnName)
-{
-	FSCSEditorTreeNodePtrType NodePtr = GetNode();
-
-	// We've removed the other columns for now,  implement them for the root actor if necessary
-	ensure(ColumnName == SCS_ColumnName_ComponentClass);
-
-	// Create the name field
-	TSharedPtr<SInlineEditableTextBlock> InlineEditableWidget =
-		SNew(SInlineEditableTextBlock)
-		.Text(this, &SSCS_RowWidget_ActorRoot::GetActorDisplayText)
-		.OnVerifyTextChanged(this, &SSCS_RowWidget_ActorRoot::OnVerifyActorLabelChanged)
-		.OnTextCommitted(this, &SSCS_RowWidget_ActorRoot::OnNameTextCommit)
-		.IsSelected(this, &SSCS_RowWidget_ActorRoot::IsSelectedExclusively)
-		.IsReadOnly(!NodePtr->CanRename() || (SCSEditor.IsValid() && !SCSEditor.Pin()->IsEditingAllowed()) || FChildActorComponentEditorUtils::IsChildActorNode(NodePtr));
-
-	NodePtr->SetRenameRequestedDelegate(FSCSEditorTreeNode::FOnRenameRequested::CreateSP(InlineEditableWidget.Get(), &SInlineEditableTextBlock::EnterEditingMode));
-
-	const bool IsRootActorNode = NodePtr->GetNodeType() == FSCSEditorTreeNode::ENodeType::RootActorNode;
-
-	return SNew(SHorizontalBox)
-		.ToolTip(CreateToolTipWidget())
-
-	+ SHorizontalBox::Slot()
-		.Padding(FMargin(IsRootActorNode ? 0.f : 4.f, 0.f, 0.f, 0.f))
-		.AutoWidth()
-		.VAlign(VAlign_Center)
-		[
-			SNew(SExpanderArrow, SharedThis(this))
-			.Visibility(IsRootActorNode ? EVisibility::Collapsed : EVisibility::Visible)
-		]
-
-	+ SHorizontalBox::Slot()
-		.Padding(FMargin(IsRootActorNode ? 4.f : 0.f, 0.f, 0.f, 0.f))
-		.AutoWidth()
-		.VAlign(VAlign_Center)
-		[
-			SNew(SImage)
-			.Image(GetIconBrush())
-			.ColorAndOpacity(FSlateColor::UseForeground())
-		]
-
-	+ SHorizontalBox::Slot()
-		.AutoWidth()
-		.HAlign(HAlign_Left)
-		.VAlign(VAlign_Center)
-		.Padding(6.f, 0.f, 0.f, 0.f)
-		[
-			InlineEditableWidget.ToSharedRef()
-		]
-
-	+ SHorizontalBox::Slot()
-		.HAlign(HAlign_Left)
-		.VAlign(VAlign_Center)
-		.Padding(4.f, 0.f, 0.f, 0.f)
-		[
-			SNew(STextBlock)
-			.Text(this, &SSCS_RowWidget_ActorRoot::GetActorContextText)
-			.ColorAndOpacity(FSlateColor::UseForeground())
-		];
-}
-
-TSharedRef<SToolTip> SSCS_RowWidget_ActorRoot::CreateToolTipWidget() const
-{
-	// Create a box to hold every line of info in the body of the tooltip
-	TSharedRef<SVerticalBox> InfoBox = SNew(SVerticalBox);
-
-	// Add class
-	AddToToolTipInfoBox(InfoBox, LOCTEXT("TooltipClass", "Class"), SNullWidget::NullWidget, TAttribute<FText>::Create(TAttribute<FText>::FGetter::CreateSP(this, &SSCS_RowWidget_ActorRoot::GetActorClassNameText)), false);
-
-	// Add super class
-	AddToToolTipInfoBox(InfoBox, LOCTEXT("TooltipSuperClass", "Parent Class"), SNullWidget::NullWidget, TAttribute<FText>::Create(TAttribute<FText>::FGetter::CreateSP(this, &SSCS_RowWidget_ActorRoot::GetActorSuperClassNameText)), false);
-
-	// Add mobility
-	AddToToolTipInfoBox(InfoBox, LOCTEXT("TooltipMobility", "Mobility"), SNullWidget::NullWidget, TAttribute<FText>::Create(TAttribute<FText>::FGetter::CreateSP(this, &SSCS_RowWidget_ActorRoot::GetActorMobilityText)), false);
-
-	TSharedRef<SBorder> TooltipContent = SNew(SBorder)
-		.BorderImage(FAppStyle::GetBrush("NoBorder"))
-		.Padding(0.0f)
-		[
-			SNew(SVerticalBox)
-
-			+ SVerticalBox::Slot()
-			.AutoHeight()
-			.Padding(0.0f, 0.0f, 0.0f, 4.0f)
-			[
-				SNew(SVerticalBox)
-
-				+ SVerticalBox::Slot()
-				.AutoHeight()
-				[
-					SNew(SHorizontalBox)
-
-					+ SHorizontalBox::Slot()
-					.AutoWidth()
-					.VAlign(VAlign_Center)
-					.Padding(4.0f)
-					[
-						SNew(STextBlock)
-						.TextStyle(FAppStyle::Get(), "SCSEditor.ComponentTooltip.Title")
-						.Text(this, &SSCS_RowWidget_ActorRoot::GetActorDisplayText)
-					]
-				]
-			]
-
-			+ SVerticalBox::Slot()
-			.AutoHeight()
-			[
-				SNew(SBorder)
-				.BorderImage(FAppStyle::GetBrush("NoBorder"))
-				.Padding(4.0f)
-				[
-					InfoBox
-				]
-			]
-		];
-
-	return IDocumentation::Get()->CreateToolTip(TAttribute<FText>(this, &SSCS_RowWidget_ActorRoot::GetActorDisplayText), TooltipContent, InfoBox, TEXT(""), TEXT(""));
-}
-
-bool SSCS_RowWidget_ActorRoot::OnVerifyActorLabelChanged(const FText& InLabel, FText& OutErrorMessage)
-{
-	return FActorEditorUtils::ValidateActorName(InLabel, OutErrorMessage);
-}
-
-const FSlateBrush* SSCS_RowWidget_ActorRoot::GetIconBrush() const
-{
-	FSCSEditorActorNodePtrType NodePtr = GetActorNode();
-	if (NodePtr.IsValid())
-	{
-		if (const AActor* Actor = NodePtr->GetObject<AActor>())
-		{
-			return FClassIconFinder::FindIconForActor(Actor);
-		}
-	}
-
-	return nullptr;
-}
-
-FText SSCS_RowWidget_ActorRoot::GetActorDisplayText() const
-{
-	FSCSEditorActorNodePtrType NodePtr = GetActorNode();
-	if (NodePtr.IsValid())
-	{
-		if (FChildActorComponentEditorUtils::IsChildActorNode(NodePtr))
-		{
-			if (const AActor* ChildActor = NodePtr->GetObject<AActor>())
-			{
-				return ChildActor->GetClass()->GetDisplayNameText();
-			}
-		}
-		else
-		{
-			if (const AActor* DefaultActor = NodePtr->GetObject<AActor>())
-			{
-				FString Name;
-				UBlueprint* Blueprint = UBlueprint::GetBlueprintFromClass(DefaultActor->GetClass());
-				if (Blueprint != nullptr && !NodePtr->IsInstanced())
-				{
-					Blueprint->GetName(Name);
-				}
-				else
-				{
-					Name = DefaultActor->GetActorLabel();
-				}
-				return FText::FromString(Name);
-			}
-		}
-	}
-
-	return FText::GetEmpty();
-}
-
-FText SSCS_RowWidget_ActorRoot::GetActorContextText() const
-{
-	FSCSEditorActorNodePtrType NodePtr = GetActorNode();
-	if (NodePtr.IsValid())
-	{
-		if (FChildActorComponentEditorUtils::IsChildActorNode(NodePtr))
-		{
-			return LOCTEXT("ActorContext_ChildActor", " (Child Actor)");
-		}
-		else
-		{
-			if (const AActor* DefaultActor = NodePtr->GetObject<AActor>())
-			{
-				if (UBlueprint* Blueprint = UBlueprint::GetBlueprintFromClass(DefaultActor->GetClass()))
-				{
-					return LOCTEXT("ActorContext_self", " (self)");
-				}
-				else
-				{
-					return LOCTEXT("ActorContext_Instance", " (Instance)");
-				}
-			}
-		}
-	}
-	return FText::GetEmpty();
-}
-
-FText SSCS_RowWidget_ActorRoot::GetActorClassNameText() const
-{
-	FSCSEditorActorNodePtrType NodePtr = GetActorNode();
-	if (NodePtr.IsValid())
-	{
-		if (const AActor* DefaultActor = NodePtr->GetObject<AActor>())
-		{
-			return FText::FromString(DefaultActor->GetClass()->GetName());
-		}
-	}
-
-	return FText::GetEmpty();
-}
-
-FText SSCS_RowWidget_ActorRoot::GetActorSuperClassNameText() const
-{
-	FSCSEditorActorNodePtrType NodePtr = GetActorNode();
-	if (NodePtr.IsValid())
-	{
-		if (const AActor* DefaultActor = NodePtr->GetObject<AActor>())
-		{
-			return FText::FromString(DefaultActor->GetClass()->GetSuperClass()->GetName());
-		}
-	}
-
-	return FText::GetEmpty();
-}
-
-FText SSCS_RowWidget_ActorRoot::GetActorMobilityText() const
-{
-	FSCSEditorActorNodePtrType NodePtr = GetActorNode();
-	if (NodePtr.IsValid())
-	{
-		if (const AActor* DefaultActor = NodePtr->GetObject<AActor>())
-		{
-			USceneComponent* RootComponent = DefaultActor->GetRootComponent();
-
-			FSCSEditorTreeNodePtrType SceneRootNodePtr = NodePtr->GetSceneRootNode();
-			if ((RootComponent == nullptr) && SceneRootNodePtr.IsValid())
-			{
-				RootComponent = Cast<USceneComponent>(SceneRootNodePtr->GetComponentTemplate());
-			}
-
-			if (RootComponent != nullptr)
-			{
-				if (RootComponent->Mobility == EComponentMobility::Static)
-				{
-					return LOCTEXT("ComponentMobility_Static", "Static");
-				}
-				else if (RootComponent->Mobility == EComponentMobility::Stationary)
-				{
-					return LOCTEXT("ComponentMobility_Stationary", "Stationary");
-				}
-				else if (RootComponent->Mobility == EComponentMobility::Movable)
-				{
-					return LOCTEXT("ComponentMobility_Movable", "Movable");
-				}
-			}
-			else
-			{
-				return LOCTEXT("ComponentMobility_NoRoot", "No root component, unknown mobility");
-			}
-		}
-	}
-
-	return FText::GetEmpty();
-}
-
-//////////////////////////////////////////////////////////////////////////
-// SSCS_RowWidget_Separator
-
-
-TSharedRef<SWidget> SSCS_RowWidget_Separator::GenerateWidgetForColumn(const FName& ColumnName)
-{
-	return SNullWidget::NullWidget;
-
-	/*return SNew(SBox)
-		.Padding(1.f)
-		[
-			SNew(SBorder)
-			.Padding(FAppStyle::GetMargin(TEXT("Menu.Separator.Padding")))
-			.BorderImage(FAppStyle::GetBrush(TEXT("Menu.Separator")))
-		];*/
-}
-
-//////////////////////////////////////////////////////////////////////////
-// SSCSEditor
-
-BEGIN_SLATE_FUNCTION_BUILD_OPTIMIZATION
-void SSCSEditor::Construct( const FArguments& InArgs )
-{
-	EditorMode = InArgs._EditorMode;
-	ActorContext = InArgs._ActorContext;
-	AllowEditing = InArgs._AllowEditing;
-	PreviewActor = InArgs._PreviewActor;
-	OnSelectionUpdated = InArgs._OnSelectionUpdated;
-	OnItemDoubleClicked = InArgs._OnItemDoubleClicked;
-	OnHighlightPropertyInDetailsView = InArgs._OnHighlightPropertyInDetailsView;
-	OnObjectReplaced = InArgs._OnObjectReplaced;
-	bUpdatingSelection = false;
-	bAllowTreeUpdates = true;
-	bIsDiffing = InArgs._IsDiffing;
-
-	CommandList = MakeShareable( new FUICommandList );
-	CommandList->MapAction( FGenericCommands::Get().Cut,
-		FUIAction( FExecuteAction::CreateSP( this, &SSCSEditor::CutSelectedNodes ), 
-		FCanExecuteAction::CreateSP( this, &SSCSEditor::CanCutNodes ) ) 
-		);
-	CommandList->MapAction( FGenericCommands::Get().Copy,
-		FUIAction( FExecuteAction::CreateSP( this, &SSCSEditor::CopySelectedNodes ), 
-		FCanExecuteAction::CreateSP( this, &SSCSEditor::CanCopyNodes ) ) 
-		);
-	CommandList->MapAction( FGenericCommands::Get().Paste,
-		FUIAction( FExecuteAction::CreateSP( this, &SSCSEditor::PasteNodes ), 
-		FCanExecuteAction::CreateSP( this, &SSCSEditor::CanPasteNodes ) ) 
-		);
-	CommandList->MapAction( FGenericCommands::Get().Duplicate,
-		FUIAction( FExecuteAction::CreateSP( this, &SSCSEditor::OnDuplicateComponent ), 
-		FCanExecuteAction::CreateSP( this, &SSCSEditor::CanDuplicateComponent ) ) 
-		);
-
-	CommandList->MapAction( FGenericCommands::Get().Delete,
-		FUIAction( FExecuteAction::CreateSP( this, &SSCSEditor::OnDeleteNodes ), 
-		FCanExecuteAction::CreateSP( this, &SSCSEditor::CanDeleteNodes ) ) 
-		);
-
-	CommandList->MapAction( FGenericCommands::Get().Rename,
-			FUIAction( FExecuteAction::CreateSP( this, &SSCSEditor::OnRenameComponent),
-			FCanExecuteAction::CreateSP( this, &SSCSEditor::CanRenameComponent ) ) 
-		);
-
-	CommandList->MapAction( FGraphEditorCommands::Get().FindReferences,
-		FUIAction( FExecuteAction::CreateSP( this, &SSCSEditor::OnFindReferences, false, EGetFindReferenceSearchStringFlags::Legacy) )
-	);
-
-	CommandList->MapAction( FGraphEditorCommands::Get().FindReferencesByNameLocal,
-		FUIAction( FExecuteAction::CreateSP( this, &SSCSEditor::OnFindReferences, false, EGetFindReferenceSearchStringFlags::None ) )
-	);
-	
-	CommandList->MapAction( FGraphEditorCommands::Get().FindReferencesByNameGlobal,
-		FUIAction( FExecuteAction::CreateSP( this, &SSCSEditor::OnFindReferences, true, EGetFindReferenceSearchStringFlags::None ) )
-	);
-	
-	CommandList->MapAction( FGraphEditorCommands::Get().FindReferencesByClassMemberLocal,
-		FUIAction( FExecuteAction::CreateSP( this, &SSCSEditor::OnFindReferences, false, EGetFindReferenceSearchStringFlags::UseSearchSyntax ) )
-	);
-
-	CommandList->MapAction( FGraphEditorCommands::Get().FindReferencesByClassMemberGlobal,
-		FUIAction( FExecuteAction::CreateSP( this, &SSCSEditor::OnFindReferences, true, EGetFindReferenceSearchStringFlags::UseSearchSyntax ) )
-	);
-
-	FSlateBrush const* MobilityHeaderBrush = FAppStyle::GetBrush(TEXT("ClassIcon.ComponentMobilityHeaderIcon"));
-	
-	TSharedPtr<SHeaderRow> HeaderRow = SNew(SHeaderRow)
-		+ SHeaderRow::Column(SCS_ColumnName_ComponentClass)
-		.DefaultLabel(LOCTEXT("Class", "Class"))
-		.FillWidth(4);
-	
-	SCSTreeWidget = SNew(SSCSTreeType)
-		.ToolTipText(LOCTEXT("DropAssetToAddComponent", "Drop asset here to add a component."))
-		.SCSEditor(this)
-		.TreeItemsSource(&RootNodes)
-		.SelectionMode(ESelectionMode::Multi)
-		.OnGenerateRow(this, &SSCSEditor::MakeTableRowWidget)
-		.OnGetChildren(this, &SSCSEditor::OnGetChildrenForTree)
-		.OnSetExpansionRecursive(this, &SSCSEditor::SetItemExpansionRecursive)
-		.OnSelectionChanged(this, &SSCSEditor::OnTreeSelectionChanged)
-		.OnContextMenuOpening(this, &SSCSEditor::CreateContextMenu)
-		.OnItemScrolledIntoView(this, &SSCSEditor::OnItemScrolledIntoView)
-		.OnMouseButtonDoubleClick(this, &SSCSEditor::HandleItemDoubleClicked)
-		.ClearSelectionOnClick(InArgs._EditorMode == EComponentEditorMode::BlueprintSCS ? true : false)
-		.OnTableViewBadState(this, &SSCSEditor::DumpTree)
-		.ItemHeight(24)
-		.HeaderRow
-		(
-			HeaderRow
-		);
-
-	SCSTreeWidget->GetHeaderRow()->SetVisibility(EVisibility::Collapsed);
-
-	TSharedPtr<SWidget> Contents;
-
-	FMenuBuilder EditBlueprintMenuBuilder( true, NULL );
-
-	EditBlueprintMenuBuilder.BeginSection( NAME_None, LOCTEXT("EditBlueprintMenu_ExistingBlueprintHeader", "Existing Blueprint" ) );
-
-	EditBlueprintMenuBuilder.AddMenuEntry
-	(
-		LOCTEXT("OpenBlueprintEditor", "Open Blueprint Editor"),
-		LOCTEXT("OpenBlueprintEditor_ToolTip", "Opens the blueprint editor for this asset"),
-		FSlateIcon(),
-		FUIAction(FExecuteAction::CreateSP(this, &SSCSEditor::OnOpenBlueprintEditor, /*bForceCodeEditing=*/ false))
-	);
-
-	EditBlueprintMenuBuilder.AddMenuEntry
-	(
-		LOCTEXT("OpenBlueprintEditorScriptMode", "Add or Edit Script"),
-		LOCTEXT("OpenBlueprintEditorScriptMode_ToolTip", "Opens the blueprint editor for this asset, showing the event graph"),
-		FSlateIcon(),
-		FUIAction(FExecuteAction::CreateSP(this, &SSCSEditor::OnOpenBlueprintEditor, /*bForceCodeEditing=*/ true))
-	);
-
-	EditBlueprintMenuBuilder.BeginSection(NAME_None, LOCTEXT("EditBlueprintMenu_InstanceHeader", "Instance modifications"));
-
-	EditBlueprintMenuBuilder.AddMenuEntry
-	(
-		LOCTEXT("PushChangesToBlueprint", "Apply Instance Changes to Blueprint"),
-		TAttribute<FText>(this, &SSCSEditor::OnGetApplyChangesToBlueprintTooltip),
-		FSlateIcon(),
-		FUIAction(FExecuteAction::CreateSP(this, &SSCSEditor::OnApplyChangesToBlueprint))
-	);
-
-	EditBlueprintMenuBuilder.AddMenuEntry
-	(
-		LOCTEXT("ResetToDefault", "Reset Instance Changes to Blueprint Default"),
-		TAttribute<FText>(this, &SSCSEditor::OnGetResetToBlueprintDefaultsTooltip),
-		FSlateIcon(),
-		FUIAction(FExecuteAction::CreateSP(this, &SSCSEditor::OnResetToBlueprintDefaults))
-	);
-
-	EditBlueprintMenuBuilder.BeginSection( NAME_None, LOCTEXT("EditBlueprintMenu_NewHeader", "Create New" ) );
-	//EditBlueprintMenuBuilder.AddMenuSeparator();
-
-	EditBlueprintMenuBuilder.AddMenuEntry
-	(
-		LOCTEXT("CreateChildBlueprint", "Create Child Blueprint Class"),
-		LOCTEXT("CreateChildBlueprintTooltip", "Creates a Child Blueprint Class based on the current Blueprint, allowing you to create variants easily.  This replaces the current actor instance with a new one based on the new Child Blueprint Class." ),
-		FSlateIcon(),
-		FUIAction(FExecuteAction::CreateSP(this, &SSCSEditor::PromoteToBlueprint))
-	);
-
-	TSharedPtr<SVerticalBox>   HeaderBox;
-	TSharedPtr<SWidget> SearchBar = SAssignNew(FilterBox, SSearchBox)
-		.HintText(EditorMode == EComponentEditorMode::ActorInstance ? LOCTEXT("SearchComponentsHint", "Search Components") : LOCTEXT("SearchHint", "Search"))
-		.OnTextChanged(this, &SSCSEditor::OnFilterTextChanged)
-		.Visibility(this, &SSCSEditor::GetComponentsFilterBoxVisibility);
-
-	const bool  bInlineSearchBarWithButtons = (EditorMode == EComponentEditorMode::BlueprintSCS);
-
-	HideComponentClassCombo = InArgs._HideComponentClassCombo;
-	ComponentTypeFilter = InArgs._ComponentTypeFilter;
-
-	USCSEditorExtensionContext* ExtensionContext = NewObject<USCSEditorExtensionContext>();
-	ExtensionContext->SCSEditor = SharedThis(this);
-	ExtensionContext->AddToRoot();
-
-	ButtonBox = SNew(SHorizontalBox)
-	+ SHorizontalBox::Slot()
-	.VAlign(VAlign_Center)
-	.Padding(0.0f, 0.0f, 4.0f, 0.0f)
-	.AutoWidth()
-	[
-		SNew(SComponentClassCombo)
-		.AddMetaData<FTagMetaData>(FTagMetaData(TEXT("Actor.AddComponent")))
-		.Visibility(this, &SSCSEditor::GetComponentClassComboButtonVisibility)
-		.OnComponentClassSelected(this, &SSCSEditor::PerformComboAddClass)
-		.ToolTipText(LOCTEXT("AddComponent_Tooltip", "Adds a new component to this actor"))
-		.IsEnabled(AllowEditing)
-	]
-	+ SHorizontalBox::Slot()
-	.VAlign(VAlign_Center)
-	.Padding(0.0f, 0.0f, 4.0f, 0.0f)
-	.AutoWidth()
-	[
-		SAssignNew(ExtensionPanel, SExtensionPanel)
-		.ExtensionPanelID("SCSEditor.NextToAddComponentButton")
-		.ExtensionContext(ExtensionContext)
-	]
-	//
-	// horizontal slot index #2 => reserved for BP-editor search bar (see 'ButtonBox' and 'SearchBarHorizontalSlotIndex' usage below)
-
-	+ SHorizontalBox::Slot()
-	.VAlign(VAlign_Center)
-	.Padding(0.0f, 0.0f, 4.0f, 0.0f)
-	.AutoWidth()
-	[
-		SNew(SPositiveActionButton)
-		.AddMetaData<FTagMetaData>( FTagMetaData(TEXT("Actor.ConvertToBlueprint")) )
-		.Visibility( this, &SSCSEditor::GetPromoteToBlueprintButtonVisibility )
-		.OnClicked( this, &SSCSEditor::OnPromoteToBlueprintClicked )
-		.Icon(FAppStyle::Get().GetBrush("Icons.Blueprints"))
-		.ToolTip(IDocumentation::Get()->CreateToolTip(
-			LOCTEXT("PromoteToBluerprintTooltip", "Converts this actor into a reusable Blueprint Class that can have script behavior" ),
-			NULL,
-			TEXT("Shared/LevelEditor"),
-			TEXT("ConvertToBlueprint")))
-	]
-	+ SHorizontalBox::Slot()
-	.VAlign(VAlign_Center)
-	.AutoWidth()
-	[
-		SNew(SPositiveActionButton)
-		.AddMetaData<FTagMetaData>(FTagMetaData(TEXT("Actor.EditBlueprint")))
-		.Visibility(this, &SSCSEditor::GetEditBlueprintButtonVisibility)
-		.ToolTipText(LOCTEXT("EditActorBlueprint_Tooltip", "Edit the Blueprint for this Actor"))
-		.Icon(FAppStyle::Get().GetBrush("Icons.Blueprints"))
-		.MenuContent()
-		[
-			EditBlueprintMenuBuilder.MakeWidget()
-		]
-	];
-
-	Contents = SNew(SVerticalBox)
-	+ SVerticalBox::Slot()
-	.AutoHeight()
-	.VAlign(VAlign_Top)
-	.Padding(4.f, 0, 4.f, 4.f)
-	[
-		SAssignNew(HeaderBox, SVerticalBox)
-	]
-
-	+ SVerticalBox::Slot()
-	[
-		SNew(SBorder)
-		.BorderImage(FAppStyle::Get().GetBrush("SCSEditor.Background"))
-		.Padding(4.f)
-		.AddMetaData<FTagMetaData>(FTagMetaData(TEXT("ComponentsPanel")))
-		.Visibility(this, &SSCSEditor::GetComponentsTreeVisibility)
-		[
-			SCSTreeWidget.ToSharedRef()
-		]
-	];
-
-	// Only insert the buttons and search bar in the Blueprints version
-	if (bInlineSearchBarWithButtons) // Blueprints
-	{
-		ButtonBox->AddSlot()
-			.FillWidth(1.0f)
-			.VAlign(VAlign_Center)
-			.Padding(3.0f, 3.0f)
-		[
-			SearchBar.ToSharedRef()
-		];
-
-		HeaderBox->AddSlot()
-			.VAlign(VAlign_Center)
-			.AutoHeight()
-		[
-			ButtonBox.ToSharedRef()
-		];
-	}
-
-	this->ChildSlot
-	[
-		Contents.ToSharedRef()
-	];
-
-	// Refresh the tree widget
-	UpdateTree();
-
-	if (EditorMode == EComponentEditorMode::ActorInstance)
-	{
-		GEngine->OnLevelComponentRequestRename().AddSP(this, &SSCSEditor::OnLevelComponentRequestRename);
-		FCoreUObjectDelegates::OnObjectsReplaced.AddSP(this, &SSCSEditor::OnObjectsReplaced);
-	}
-}
-END_SLATE_FUNCTION_BUILD_OPTIMIZATION
-
-SSCSEditor::~SSCSEditor()
-{
-	if (UObject* ExtensionContext = ExtensionPanel->GetExtensionContext())
-	{
-		ExtensionContext->RemoveFromRoot();
-	}
-}
-
-TSharedPtr<SWidget> SSCSEditor::GetToolButtonsBox()
-{
-	return ButtonBox;	
-}
-
-void SSCSEditor::OnLevelComponentRequestRename(const UActorComponent* InComponent)
-{
-	TArray< FSCSEditorTreeNodePtrType > SelectedItems = SCSTreeWidget->GetSelectedItems();
-	
-	FSCSEditorTreeNodePtrType Node = GetNodeFromActorComponent(InComponent);
-	if (SelectedItems.Contains(Node) && CanRenameComponent())
-	{
-		OnRenameComponent();
-	}
-}
-
-void SSCSEditor::OnObjectsReplaced(const TMap<UObject*, UObject*>& OldToNewInstanceMap)
-{
-	bool bHasChanges = false;
-	ReplaceComponentReferencesInTree(GetActorNode(), OldToNewInstanceMap, bHasChanges);
-	if (bHasChanges)
-	{
-		OnObjectReplaced.ExecuteIfBound();
-	}
-}
-
-void SSCSEditor::ReplaceComponentReferencesInTree(FSCSEditorActorNodePtrType InActorNode, const TMap<UObject*, UObject*>& OldToNewInstanceMap, bool& OutHasChanges)
-{
-	if (InActorNode.IsValid())
-	{
-		ReplaceComponentReferencesInTree(InActorNode->GetComponentNodes(), OldToNewInstanceMap, OutHasChanges);
-	}
-}
-
-void SSCSEditor::ReplaceComponentReferencesInTree(const TArray<FSCSEditorTreeNodePtrType>& Nodes, const TMap<UObject*, UObject*>& OldToNewInstanceMap, bool& OutHasChanges)
-{
-	for (const FSCSEditorTreeNodePtrType& Node : Nodes)
-	{
-		if (Node.IsValid())
-		{
-			// We need to get the actual pointer to the old object which will be marked for pending kill, as these are the references which need updating
-			const bool bEvenIfPendingKill = true;
-			const UObject* OldObject = Node->GetObject<UObject>(bEvenIfPendingKill);
-			if (OldObject)
-			{
-				UObject* const* NewObject = OldToNewInstanceMap.Find(OldObject);
-				if (NewObject)
-				{
-					Node->SetObject(*NewObject);
-					OutHasChanges = true;
-				}
-			}
-
-			ReplaceComponentReferencesInTree(Node->GetChildren(), OldToNewInstanceMap, OutHasChanges);
-		}
-	}
-}
-
-UBlueprint* SSCSEditor::GetBlueprint() const
-{
-	if (AActor* Actor = GetActorContext())
-	{
-		const UClass* ActorClass = Actor->GetClass();
-		check(ActorClass != nullptr);
-
-		return UBlueprint::GetBlueprintFromClass(ActorClass);
-	}
-
-	return nullptr;
-}
-
-FReply SSCSEditor::OnKeyDown( const FGeometry& MyGeometry, const FKeyEvent& InKeyEvent )
-{
-	if (CommandList->ProcessCommandBindings(InKeyEvent))
-	{
-		return FReply::Handled();
-	}
-	return FReply::Unhandled();
-}
-
-TSharedRef<ITableRow> SSCSEditor::MakeTableRowWidget( FSCSEditorTreeNodePtrType InNodePtr, const TSharedRef<STableViewBase>& OwnerTable )
-{
-	// Setup a meta tag for this node
-	FGraphNodeMetaData TagMeta(TEXT("TableRow"));
-	if (InNodePtr.IsValid() && InNodePtr->GetComponentTemplate() != NULL )
-	{
-		TagMeta.FriendlyName = FString::Printf(TEXT("TableRow,%s,0"), *InNodePtr->GetComponentTemplate()->GetReadableName());
-	}
-
-	// Create the node of the appropriate type
-	if (InNodePtr->IsActorNode())
-	{
-		return SNew(SSCS_RowWidget_ActorRoot, SharedThis(this), InNodePtr, OwnerTable);
-	}
-	else if (InNodePtr->GetNodeType() == FSCSEditorTreeNode::SeparatorNode)
-	{
-		return SNew(SSCS_RowWidget_Separator, SharedThis(this), InNodePtr, OwnerTable);
-	}
-
-	return SNew(SSCS_RowWidget, SharedThis(this), InNodePtr, OwnerTable)
-		.AddMetaData<FTutorialMetaData>(TagMeta);
-}
-
-void SSCSEditor::GetSelectedItemsForContextMenu(TArray<FComponentEventConstructionData>& OutSelectedItems) const
-{
-	TArray<FSCSEditorTreeNodePtrType> SelectedTreeItems = SCSTreeWidget->GetSelectedItems();
-	for ( auto NodeIter = SelectedTreeItems.CreateConstIterator(); NodeIter; ++NodeIter )
-	{
-		FComponentEventConstructionData NewItem;
-		const FSCSEditorTreeNodePtrType& TreeNode = *NodeIter;
-		NewItem.VariableName = TreeNode->GetVariableName();
-		NewItem.Component = TreeNode->GetComponentTemplate();
-		OutSelectedItems.Add(NewItem);
-	}
-}
-
-TArray<UObject*> SSCSEditor::GetSelectedEditableObjects() const
-{
-	TArray<UObject*> SelectedObjects;
-	if (UBlueprint* BP = GetBlueprint())
-	{
-		TArray<FSCSEditorTreeNodePtrType> SelectedTreeItems = SCSTreeWidget->GetSelectedItems();
-		SelectedObjects.Reserve(SelectedTreeItems.Num());
-		for (const FSCSEditorTreeNodePtrType& TreeNode : SelectedTreeItems)
-		{
-			UObject* Obj = TreeNode->GetEditableObjectForBlueprint<UObject>(BP);
-			if (Obj)
-			{
-				SelectedObjects.Add(Obj);
-			}
-		}
-	}
-	return SelectedObjects;
-}
-
-void SSCSEditor::PopulateContextMenu(UToolMenu* Menu)
-{
-	TArray<FSCSEditorTreeNodePtrType> SelectedItems = SCSTreeWidget->GetSelectedItems();
-
-	if (SelectedItems.Num() > 0 || CanPasteNodes())
-	{
-		bool bOnlyShowPasteOption = false;
-
-		if (SelectedItems.Num() > 0)
-		{
-			if (SelectedItems.Num() == 1 && SelectedItems[0]->IsActorNode())
-			{
-				if (FChildActorComponentEditorUtils::IsChildActorNode(SelectedItems[0]))
-				{
-					// Include specific context menu options for a single child actor node selection
-					FChildActorComponentEditorUtils::FillChildActorContextMenuOptions(Menu, SelectedItems[0]);
-				}
-				else
-				{
-					bOnlyShowPasteOption = true;
-				}
-			}
-			else
-			{
-				for (FSCSEditorTreeNodePtrType SelectedNode : SelectedItems)
-				{
-					if (!SelectedNode->IsComponentNode())
-					{
-						bOnlyShowPasteOption = true;
-						break;
-					}
-				}
-				if (!bOnlyShowPasteOption)
-				{
-					bool bIsChildActorSubtreeNodeSelected = false;
-
-					TArray<UActorComponent*> SelectedComponents;
-					TArray<FSCSEditorTreeNodePtrType> SelectedNodes = GetSelectedNodes();
-					for (int32 i = 0; i < SelectedNodes.Num(); ++i)
-					{
-						// Get the current selected node reference
-						FSCSEditorTreeNodePtrType SelectedNodePtr = SelectedNodes[i];
-						check(SelectedNodePtr.IsValid());
-
-						// Get the component template associated with the selected node
-						UActorComponent* ComponentTemplate = SelectedNodePtr->GetComponentTemplate();
-						if (ComponentTemplate)
-						{
-							SelectedComponents.Add(ComponentTemplate);
-						}
-
-						// Determine if any selected node belongs to a child actor template
-						if (!bIsChildActorSubtreeNodeSelected)
-						{
-							bIsChildActorSubtreeNodeSelected = FChildActorComponentEditorUtils::IsChildActorSubtreeNode(SelectedNodePtr);
-						}
-					}
-
-					// Don't include these commands if any component was found above to belong to a child actor template (not supported at this time)
-					if (EditorMode == EComponentEditorMode::BlueprintSCS && !bIsChildActorSubtreeNodeSelected)
-					{
-						FToolMenuSection& BlueprintSCSSection = Menu->AddSection("BlueprintSCS");
-						if (SelectedItems.Num() == 1)
-						{
-							// Expandable menu: insert sub-menu here
-                            BlueprintSCSSection.AddSubMenu(
-                            	FName("FindReferenceSubMenu"),
-                            	LOCTEXT("FindReferences_Label", "Find References"),
-                            	LOCTEXT("FindReferences_Tooltip", "Options for finding references to class members"),
-                            	FNewToolMenuChoice(FNewMenuDelegate::CreateStatic(&FGraphEditorCommands::BuildFindReferencesMenu))
-                            );
-						}
-
-						// Create an "Add Event" option in the context menu only if we can edit
-						// the currently selected objects
-						if (IsEditingAllowed())
-						{
-							// Collect the classes of all selected objects
-							TArray<UClass*> SelectionClasses;
-							for (auto NodeIter = SelectedNodes.CreateConstIterator(); NodeIter; ++NodeIter)
-							{
-								FSCSEditorTreeNodePtrType TreeNode = *NodeIter;
-								if (UActorComponent* ComponentTemplate = TreeNode->GetComponentTemplate())
-								{
-									// If the component is native then we need to ensure it can actually be edited before we display it
-									if (!TreeNode->IsNativeComponent() || FComponentEditorUtils::GetPropertyForEditableNativeComponent(ComponentTemplate))
-									{
-										SelectionClasses.Add(ComponentTemplate->GetClass());
-									}
-								}
-							}
-
-							if (SelectionClasses.Num())
-							{
-								// Find the common base class of all selected classes
-								UClass* SelectedClass = UClass::FindCommonBase(SelectionClasses);
-								// Build an event submenu if we can generate events
-								if (FBlueprintEditorUtils::CanClassGenerateEvents(SelectedClass))
-								{
-									BlueprintSCSSection.AddSubMenu(
-										"AddEventSubMenu",
-										LOCTEXT("AddEventSubMenu", "Add Event"),
-										LOCTEXT("ActtionsSubMenu_ToolTip", "Add Event"),
-										FNewMenuDelegate::CreateStatic(&SSCSEditor::BuildMenuEventsSection,
-											GetBlueprint(), SelectedClass, FCanExecuteAction::CreateSP(this, &SSCSEditor::IsEditingAllowed),
-											FGetSelectedObjectsDelegate::CreateSP(this, &SSCSEditor::GetSelectedItemsForContextMenu)));
-								}
-							}
-						}						
-					}					
-
-					// Common menu options added for all component types
-					FComponentEditorUtils::FillComponentContextMenuOptions(Menu, SelectedComponents);
-
-					// For a selection outside of a child actor subtree, we may choose to include additional options
-					if (SelectedComponents.Num() == 1 && !bIsChildActorSubtreeNodeSelected)
-					{
-						// Extra options for a child actor component
-						if (UChildActorComponent* SelectedChildActorComponent = Cast<UChildActorComponent>(SelectedComponents[0]))
-						{
-							// These options will get added only in SCS mode
-							if (EditorMode == EComponentEditorMode::BlueprintSCS)
-							{
-								FChildActorComponentEditorUtils::FillComponentContextMenuOptions(Menu, SelectedChildActorComponent);
-							}
-						}
-					}
-				}
-			}
-		}
-		else
-		{
-			bOnlyShowPasteOption = true;
-		}
-
-		if (bOnlyShowPasteOption)
-		{
-			FToolMenuSection& Section = Menu->AddSection("PasteComponent", LOCTEXT("EditComponentHeading", "Edit"));
-			{
-				Section.AddMenuEntry(FGenericCommands::Get().Paste);
-			}
-		}
-	}
-}
-
-void SSCSEditor::RegisterContextMenu()
-{
-	UToolMenus* ToolMenus = UToolMenus::Get();
-	if (!ToolMenus->IsMenuRegistered(SCS_ContextMenuName))
-	{
-		UToolMenu* Menu = ToolMenus->RegisterMenu(SCS_ContextMenuName);
-		Menu->AddDynamicSection("SCSEditorDynamic", FNewToolMenuDelegate::CreateLambda([](UToolMenu* InMenu)
-		{
-			USSCSEditorMenuContext* ContextObject = InMenu->FindContext<USSCSEditorMenuContext>();
-			if (ContextObject && ContextObject->SCSEditor.IsValid())
-			{
-				ContextObject->SCSEditor.Pin()->PopulateContextMenu(InMenu);
-			}
-		}));
-	}
-}
-
-TSharedPtr< SWidget > SSCSEditor::CreateContextMenu()
-{
-	TArray<FSCSEditorTreeNodePtrType> SelectedItems = SCSTreeWidget->GetSelectedItems();
-
-	if (SelectedItems.Num() > 0 || CanPasteNodes())
-	{
-		RegisterContextMenu();
-		USSCSEditorMenuContext* ContextObject = NewObject<USSCSEditorMenuContext>();
-		ContextObject->SCSEditor = SharedThis(this);
-		FToolMenuContext ToolMenuContext(CommandList, TSharedPtr<FExtender>(), ContextObject);
-		return UToolMenus::Get()->GenerateWidget(SCS_ContextMenuName, ToolMenuContext);
-	}
-	return TSharedPtr<SWidget>();
-}
-
-void SSCSEditor::BuildMenuEventsSection(FMenuBuilder& Menu, UBlueprint* Blueprint, UClass* SelectedClass, FCanExecuteAction CanExecuteActionDelegate, FGetSelectedObjectsDelegate GetSelectedObjectsDelegate)
-{
-	// Get Selected Nodes
-	TArray<FComponentEventConstructionData> SelectedNodes;
-	GetSelectedObjectsDelegate.ExecuteIfBound( SelectedNodes );
-
-	struct FMenuEntry
-	{
-		FText		Label;
-		FText		ToolTip;
-		FUIAction	UIAction;
-	};
-
-	TArray< FMenuEntry > Actions;
-	TArray< FMenuEntry > NodeActions;
-	// Build Events entries
-	for (TFieldIterator<FMulticastDelegateProperty> PropertyIt(SelectedClass, EFieldIteratorFlags::IncludeSuper); PropertyIt; ++PropertyIt)
-	{
-		FMulticastDelegateProperty* Property = *PropertyIt;
-
-		// Check for multicast delegates that we can safely assign
-		if (!Property->HasAnyPropertyFlags(CPF_Parm) && Property->HasAllPropertyFlags(CPF_BlueprintAssignable))
-		{
-			FName EventName = Property->GetFName();
-			int32 ComponentEventViewEntries = 0;
-			// Add View Event Per Component
-			for (auto NodeIter = SelectedNodes.CreateConstIterator(); NodeIter; ++NodeIter )
-			{
-				if( NodeIter->Component.IsValid() )
-				{
-					FName VariableName = NodeIter->VariableName;
-					FObjectProperty* VariableProperty = FindFProperty<FObjectProperty>( Blueprint->SkeletonGeneratedClass, VariableName );
-
-					if( VariableProperty && FKismetEditorUtilities::FindBoundEventForComponent( Blueprint, EventName, VariableProperty->GetFName() ))
-					{
-						FMenuEntry NewEntry;
-						NewEntry.Label = ( SelectedNodes.Num() > 1 ) ?	FText::Format( LOCTEXT("ViewEvent_ToolTipFor", "{0} for {1}"), FText::FromName( EventName ), FText::FromName( VariableName )) : 
-																		FText::Format( LOCTEXT("ViewEvent_ToolTip", "{0}"), FText::FromName( EventName ));
-						NewEntry.UIAction =	FUIAction(FExecuteAction::CreateStatic( &SSCSEditor::ViewEvent, Blueprint, EventName, *NodeIter ), CanExecuteActionDelegate);
-						NodeActions.Add( NewEntry );
-						ComponentEventViewEntries++;
-					}
-				}
-			}
-			if( ComponentEventViewEntries < SelectedNodes.Num() )
-			{
-				// Create menu Add entry
-				FMenuEntry NewEntry;
-				NewEntry.Label = FText::Format( LOCTEXT("AddEvent_ToolTip", "Add {0}" ), FText::FromName( EventName ));
-				NewEntry.UIAction =	FUIAction(FExecuteAction::CreateStatic( &SSCSEditor::CreateEventsForSelection, Blueprint, EventName, GetSelectedObjectsDelegate), CanExecuteActionDelegate);
-				Actions.Add( NewEntry );
-			}
-		}
-	}
-	// Build Menu Sections
-	Menu.BeginSection("AddComponentActions", LOCTEXT("AddEventHeader", "Add Event"));
-	for (auto ItemIter = Actions.CreateConstIterator(); ItemIter; ++ItemIter )
-	{
-		Menu.AddMenuEntry( ItemIter->Label, ItemIter->ToolTip, FSlateIcon(), ItemIter->UIAction );
-	}
-	Menu.EndSection();
-	Menu.BeginSection("ViewComponentActions", LOCTEXT("ViewEventHeader", "View Existing Events"));
-	for (auto ItemIter = NodeActions.CreateConstIterator(); ItemIter; ++ItemIter )
-	{
-		Menu.AddMenuEntry( ItemIter->Label, ItemIter->ToolTip, FSlateIcon(), ItemIter->UIAction );
-	}
-	Menu.EndSection();
-}
-
-void SSCSEditor::CreateEventsForSelection(UBlueprint* Blueprint, FName EventName, FGetSelectedObjectsDelegate GetSelectedObjectsDelegate)
-{	
-	if (EventName != NAME_None)
-	{
-		TArray<FComponentEventConstructionData> SelectedNodes;
-		GetSelectedObjectsDelegate.ExecuteIfBound(SelectedNodes);
-
-		for (auto SelectionIter = SelectedNodes.CreateConstIterator(); SelectionIter; ++SelectionIter)
-		{
-			ConstructEvent( Blueprint, EventName, *SelectionIter );
-		}
-	}
-}
-
-void SSCSEditor::ConstructEvent(UBlueprint* Blueprint, const FName EventName, const FComponentEventConstructionData EventData)
-{
-	// Find the corresponding variable property in the Blueprint
-	FObjectProperty* VariableProperty = FindFProperty<FObjectProperty>(Blueprint->SkeletonGeneratedClass, EventData.VariableName );
-
-	if( VariableProperty )
-	{
-		if (!FKismetEditorUtilities::FindBoundEventForComponent(Blueprint, EventName, VariableProperty->GetFName()))
-		{
-			FKismetEditorUtilities::CreateNewBoundEventForComponent(EventData.Component.Get(), EventName, Blueprint, VariableProperty);
-		}
-	}
-}
-
-void SSCSEditor::ViewEvent(UBlueprint* Blueprint, const FName EventName, const FComponentEventConstructionData EventData)
-{
-	// Find the corresponding variable property in the Blueprint
-	FObjectProperty* VariableProperty = FindFProperty<FObjectProperty>(Blueprint->SkeletonGeneratedClass, EventData.VariableName );
-
-	if( VariableProperty )
-	{
-		const UK2Node_ComponentBoundEvent* ExistingNode = FKismetEditorUtilities::FindBoundEventForComponent(Blueprint, EventName, VariableProperty->GetFName());
-		if (ExistingNode)
-		{
-			FKismetEditorUtilities::BringKismetToFocusAttentionOnObject(ExistingNode);
-		}
-	}
-}
-
-void SSCSEditor::OnFindReferences(bool bSearchAllBlueprints, const EGetFindReferenceSearchStringFlags Flags)
-{
-	TArray<FSCSEditorTreeNodePtrType> SelectedNodes = SCSTreeWidget->GetSelectedItems();
-	if (SelectedNodes.Num() == 1)
-	{
-		TSharedPtr<IToolkit> FoundAssetEditor = FToolkitManager::Get().FindEditorForAsset(GetBlueprint());
-		if (FoundAssetEditor.IsValid())
-		{
-			const FString VariableName = SelectedNodes[0]->GetVariableName().ToString();
-
-			FMemberReference MemberReference;
-			MemberReference.SetSelfMember(*VariableName);
-			const FString SearchTerm = EnumHasAnyFlags(Flags, EGetFindReferenceSearchStringFlags::UseSearchSyntax) ? MemberReference.GetReferenceSearchString(GetBlueprint()->SkeletonGeneratedClass) : FString::Printf(TEXT("\"%s\""), *VariableName);;
-
-			TSharedRef<IBlueprintEditor> BlueprintEditor = StaticCastSharedRef<IBlueprintEditor>(FoundAssetEditor.ToSharedRef());
-			const bool bSetFindWithinBlueprint = !bSearchAllBlueprints;
-			BlueprintEditor->SummonSearchUI(true, SearchTerm);
-		}
-	}
-}
-
-bool SSCSEditor::CanDuplicateComponent() const
-{
-	if(!IsEditingAllowed())
-	{
-		return false;
-	}
-
-	// @todo - Allow duplication of components that belong to a child actor template? For now, we don't support this.
-	return CanCopyNodes() && !FChildActorComponentEditorUtils::ContainsChildActorSubtreeNode(SCSTreeWidget->GetSelectedItems());
-}
-
-void SSCSEditor::OnDuplicateComponent()
-{
-	TArray<FSCSEditorTreeNodePtrType> SelectedNodes = SCSTreeWidget->GetSelectedItems();
-	if(SelectedNodes.Num() > 0)
-	{
-		// Force the text box being edited (if any) to commit its text. The duplicate operation may trigger a regeneration of the tree view,
-		// releasing all row widgets. If one row was in edit mode (rename/rename on create), it was released before losing the focus and
-		// this would prevent the completion of the 'rename' or 'create + give initial name' transaction (occurring on focus lost).
-		FSlateApplication::Get().ClearKeyboardFocus();
-
-		const FScopedTransaction Transaction(SelectedNodes.Num() > 1 ? LOCTEXT("DuplicateComponents", "Duplicate Components") : LOCTEXT("DuplicateComponent", "Duplicate Component"));
-
-		FAddNewComponentParams NewComponentParams;
-		NewComponentParams.bConformTransformToParent = false;
-
-		TMap<USceneComponent*, USceneComponent*> DuplicateSceneComponentMap;
-		for (int32 i = 0; i < SelectedNodes.Num(); ++i)
-		{
-			if (UActorComponent* ComponentTemplate = SelectedNodes[i]->GetComponentTemplate())
-			{
-				USCS_Node* SCSNode = SelectedNodes[i]->GetSCSNode();
-				check(SCSNode == nullptr || SCSNode->ComponentTemplate == ComponentTemplate);
-				UActorComponent* CloneComponent = AddNewComponent(ComponentTemplate->GetClass(), (SCSNode ? (UObject*)SCSNode : ComponentTemplate), NewComponentParams);
-				if (USceneComponent* SceneClone = Cast<USceneComponent>(CloneComponent))
-				{
-					DuplicateSceneComponentMap.Add(CastChecked<USceneComponent>(ComponentTemplate), SceneClone);
-				}
-			}
-		}
-
-		for (const TPair<USceneComponent*,USceneComponent*>& DuplicatedPair : DuplicateSceneComponentMap)
-		{
-			USceneComponent* OriginalComponent = DuplicatedPair.Key;
-			USceneComponent* NewSceneComponent = DuplicatedPair.Value;
-
-			if (EditorMode == EComponentEditorMode::BlueprintSCS)
-			{
-				// Ensure that any native attachment relationship inherited from the original copy is removed (to prevent a GLEO assertion)
-				NewSceneComponent->DetachFromComponent(FDetachmentTransformRules::KeepWorldTransform);
-			}
-					
-			// Attempt to locate the original node in the SCS tree
-			FSCSEditorTreeNodePtrType OriginalNodePtr = FindTreeNode(OriginalComponent);
-			if(OriginalNodePtr.IsValid())
-			{
-				// If we're duplicating the root then we're already a child of it so need to reparent, but we do need to reset the scale
-				// otherwise we'll end up with the square of the root's scale instead of being the same size.
-				if (OriginalNodePtr == GetSceneRootNode())
-				{
-					NewSceneComponent->SetRelativeScale3D_Direct(FVector(1.f));
-				}
-				else
-				{
-					// If the original node was parented, attempt to add the duplicate as a child of the same parent node if the parent is not
-					// part of the duplicate set, otherwise parent to the parent's duplicate
-					FSCSEditorTreeNodePtrType ParentNodePtr = OriginalNodePtr->GetParent();
-					if (ParentNodePtr.IsValid())
-					{
-						if (USceneComponent** ParentDuplicateComponent = DuplicateSceneComponentMap.Find(Cast<USceneComponent>(ParentNodePtr->GetComponentTemplate())))
-						{
-							FSCSEditorTreeNodePtrType DuplicateParentNodePtr = FindTreeNode(*ParentDuplicateComponent);
-							if (DuplicateParentNodePtr.IsValid())
-							{
-								ParentNodePtr = DuplicateParentNodePtr;
-							}
-						}
-
-						// Locate the duplicate node (as a child of the current scene root node), and switch it to be a child of the original node's parent
-						FSCSEditorTreeNodePtrType NewChildNodePtr = GetSceneRootNode()->FindChild(NewSceneComponent, true);
-						if (NewChildNodePtr.IsValid())
-						{
-							// Note: This method will handle removal from the scene root node as well
-							ParentNodePtr->AddChild(NewChildNodePtr);
-						}
-					}
-				}
-			}
-		}
-	}
-}
-
-void SSCSEditor::OnGetChildrenForTree( FSCSEditorTreeNodePtrType InNodePtr, TArray<FSCSEditorTreeNodePtrType>& OutChildren )
-{
-	if (InNodePtr.IsValid())
-	{
-		const TArray<FSCSEditorTreeNodePtrType>& Children = InNodePtr->GetChildren();
-		OutChildren.Reserve(Children.Num());
-
-		if (GetComponentTypeFilterToApply() || !GetFilterText().IsEmpty())
-		{
-			for (FSCSEditorTreeNodePtrType Child : Children)
-			{
-				if (!Child->IsFlaggedForFiltration())
-				{
-					OutChildren.Add(Child);
-				}
-			}
-		}
-		else
-		{
-			OutChildren = Children;
-		}
-	}
-	else
-	{
-		OutChildren.Empty();
-	}
-}
-
-
-UActorComponent* SSCSEditor::PerformComboAddClass(TSubclassOf<UActorComponent> ComponentClass, EComponentCreateAction::Type ComponentCreateAction, UObject* AssetOverride)
-{
-	UClass* NewClass = ComponentClass;
-
-	UActorComponent* NewComponent = nullptr;
-
-	if( ComponentCreateAction == EComponentCreateAction::CreateNewCPPClass )
-	{
-		NewClass = CreateNewCPPComponent( ComponentClass );
-	}
-	else if( ComponentCreateAction == EComponentCreateAction::CreateNewBlueprintClass )
-	{
-		NewClass = CreateNewBPComponent( ComponentClass );
-	}
-
-	if( NewClass != nullptr )
-	{
-		FEditorDelegates::LoadSelectedAssetsIfNeeded.Broadcast();
-		USelection* Selection =  GEditor->GetSelectedObjects();
-
-		bool bAddedComponent = false;
-
-		// This adds components according to the type selected in the drop down. If the user
-		// has the appropriate objects selected in the content browser then those are added,
-		// else we go down the previous route of adding components by type.
-		//
-		// Furthermore don't try to match up assets for USceneComponent it will match lots of things and doesn't have any nice behavior for asset adds 
-		if (Selection->Num() > 0 && !AssetOverride && NewClass != USceneComponent::StaticClass())
-		{
-			for(FSelectionIterator ObjectIter(*Selection); ObjectIter; ++ObjectIter)
-			{
-				UObject* Object = *ObjectIter;
-				UClass*  Class	= Object->GetClass();
-
-				TArray< TSubclassOf<UActorComponent> > ComponentClasses = FComponentAssetBrokerage::GetComponentsForAsset(Object);
-
-				// if the selected asset supports the selected component type then go ahead and add it
-				for(int32 ComponentIndex = 0; ComponentIndex < ComponentClasses.Num(); ComponentIndex++)
-				{
-					if(ComponentClasses[ComponentIndex]->IsChildOf(NewClass))
-					{
-						NewComponent = AddNewComponent(NewClass, Object);
-						bAddedComponent = true;
-						break;
-					}
-				}
-			}
-		}
-
-		if(!bAddedComponent)
-		{
-			// As the SCS splits up the scene and actor components, can now add directly
-			NewComponent = AddNewComponent(NewClass, AssetOverride);
-		}
-
-		UpdateTree();
-	}
-
-	return NewComponent;
-}
-
-TArray<FSCSEditorTreeNodePtrType>  SSCSEditor::GetSelectedNodes() const
-{
-	TArray<FSCSEditorTreeNodePtrType> SelectedTreeNodes = SCSTreeWidget->GetSelectedItems();
-
-	struct FCompareSelectedSCSEditorTreeNodes
-	{
-		FORCEINLINE bool operator()(const FSCSEditorTreeNodePtrType& A, const FSCSEditorTreeNodePtrType& B) const
-		{
-			return B.IsValid() && B->IsAttachedTo(A);
-		}
-	};
-
-	// Ensure that nodes are ordered from parent to child (otherwise they are sorted in the order that they were selected)
-	SelectedTreeNodes.Sort(FCompareSelectedSCSEditorTreeNodes());
-
-	return SelectedTreeNodes;
-}
-
-FSCSEditorTreeNodePtrType SSCSEditor::GetNodeFromActorComponent(const UActorComponent* ActorComponent, bool bIncludeAttachedComponents) const
-{
-	FSCSEditorTreeNodePtrType NodePtr;
-
-	if(ActorComponent)
-	{
-		if (EditorMode == EComponentEditorMode::BlueprintSCS)
-		{
-			// If the given component instance is not already an archetype object
-			if (!ActorComponent->IsTemplate())
-			{
-				// Get the component owner's class object
-				check(ActorComponent->GetOwner() != NULL);
-				UClass* OwnerClass = ActorComponent->GetOwner()->GetClass();
-
-				// If the given component is one that's created during Blueprint construction
-				if (ActorComponent->IsCreatedByConstructionScript())
-				{
-					// Check the entire Class hierarchy for the node
-					TArray<UBlueprintGeneratedClass*> ParentBPStack;
-					UBlueprint::GetBlueprintHierarchyFromClass(OwnerClass, ParentBPStack);
-
-					for(int32 StackIndex = ParentBPStack.Num() - 1; StackIndex >= 0; --StackIndex)
-					{
-						USimpleConstructionScript* ParentSCS = ParentBPStack[StackIndex] ? ParentBPStack[StackIndex]->SimpleConstructionScript.Get() : nullptr;
-						if (ParentSCS)
-						{
-							// Attempt to locate an SCS node with a variable name that matches the name of the given component
-							for (USCS_Node* SCS_Node : ParentSCS->GetAllNodes())
-							{
-								check(SCS_Node != NULL);
-								if (SCS_Node->GetVariableName() == ActorComponent->GetFName())
-								{
-									// We found a match; redirect to the component archetype instance that may be associated with a tree node
-									ActorComponent = SCS_Node->ComponentTemplate;
-									break;
-								}
-							}
-
-						}
-
-					}
-				}
-				else
-				{
-					// Get the class default object
-					const AActor* CDO = Cast<AActor>(OwnerClass->GetDefaultObject());
-					if (CDO)
-					{
-						// Iterate over the Components array and attempt to find a component with a matching name
-						for (UActorComponent* ComponentTemplate : CDO->GetComponents())
-						{
-							if (ComponentTemplate && ComponentTemplate->GetFName() == ActorComponent->GetFName())
-							{
-								// We found a match; redirect to the component archetype instance that may be associated with a tree node
-								ActorComponent = ComponentTemplate;
-								break;
-							}
-						}
-					}
-				}
-			}
-		}
-
-		// If we have a valid component archetype instance, attempt to find a tree node that corresponds to it
-		const TArray<FSCSEditorTreeNodePtrType>& Nodes = GetRootNodes();
-		for (int32 i = 0; i < Nodes.Num() && !NodePtr.IsValid(); i++)
-		{
-			NodePtr = FindTreeNode(ActorComponent, Nodes[i]);
-		}
-
-		// If we didn't find it in the tree, step up the chain to the parent of the given component and recursively see if that is in the tree (unless the flag is false)
-		if(!NodePtr.IsValid() && bIncludeAttachedComponents)
-		{
-			const USceneComponent* SceneComponent = Cast<const USceneComponent>(ActorComponent);
-			if(SceneComponent && SceneComponent->GetAttachParent())
-			{
-				return GetNodeFromActorComponent(SceneComponent->GetAttachParent(), bIncludeAttachedComponents);
-			}
-		}
-	}
-
-	return NodePtr;
-}
-
-void SSCSEditor::SelectRoot()
-{
-	const TArray<FSCSEditorTreeNodePtrType>& Nodes = GetRootNodes();
-	if (Nodes.Num() > 0)
-	{
-		SCSTreeWidget->SetSelection(Nodes[0]);
-	}
-}
-
-void SSCSEditor::SelectNode(FSCSEditorTreeNodePtrType InNodeToSelect, bool IsCntrlDown) 
-{
-	if(SCSTreeWidget.IsValid() && InNodeToSelect.IsValid())
-	{
-		if(!IsCntrlDown)
-		{
-			SCSTreeWidget->SetSelection(InNodeToSelect);
-		}
-		else
-		{
-			SCSTreeWidget->SetItemSelection(InNodeToSelect, !SCSTreeWidget->IsItemSelected(InNodeToSelect));
-		}
-	}
-}
-
-void SSCSEditor::SetNodeExpansionState(FSCSEditorTreeNodePtrType InNodeToChange, const bool bIsExpanded)
-{
-	if(SCSTreeWidget.IsValid() && InNodeToChange.IsValid())
-	{
-		SCSTreeWidget->SetItemExpansion(InNodeToChange, bIsExpanded);
-	}
-}
-
-static FSCSEditorTreeNode* FindRecursive( FSCSEditorTreeNode* Node, FName Name )
-{
-	if (Node->GetVariableName() == Name)
-	{
-		return Node;
-	}
-	else
-	{
-		for (const FSCSEditorTreeNodePtrType& Child : Node->GetChildren())
-		{
-			if (FSCSEditorTreeNode* Result = FindRecursive(Child.Get(), Name))
-			{
-				return Result;
-			}
-		}
-	}
-
-	return nullptr;
-}
-
-void SSCSEditor::HighlightTreeNode(FName TreeNodeName, const class FPropertyPath& Property)
-{
-	for( const FSCSEditorTreeNodePtrType& Node : GetRootNodes() )
-	{
-		if( FSCSEditorTreeNode* FoundNode = FindRecursive( Node.Get(), TreeNodeName ) )
-		{
-			SelectNode(FoundNode->AsShared(), false);
-
-			if (Property != FPropertyPath())
-			{
-				// Invoke the delegate to highlight the property
-				OnHighlightPropertyInDetailsView.ExecuteIfBound(Property);
-			}
-
-			return;
-		}
-	}
-	
-	ClearSelection();
-}
-
-void SSCSEditor::HighlightTreeNode(const USCS_Node* Node, FName Property)
-{
-	check(Node);
-	FSCSEditorTreeNodePtrType TreeNode = FindTreeNode( Node );
-	check( TreeNode.IsValid() );
-	SelectNode( TreeNode, false );
-	if( Property != FName() )
-	{
-		UActorComponent* Component = TreeNode->GetComponentTemplate();
-		FProperty* CurrentProp = FindFProperty<FProperty>(Component->GetClass(), Property);
-		FPropertyPath Path;
-		if( CurrentProp )
-		{
-			FPropertyInfo NewInfo(CurrentProp, -1);
-			Path.ExtendPath(NewInfo);
-		}
-
-		// Invoke the delegate to highlight the property
-		OnHighlightPropertyInDetailsView.ExecuteIfBound( Path );
-	}
-}
-
-void SSCSEditor::BuildSubTreeForActorNode(FSCSEditorActorNodePtrType InActorNode)
-{
-	if (!InActorNode.IsValid())
-	{
-		return;
-	}
-
-	// Get the actor instance that we're editing
-	const AActor* Actor = InActorNode->GetObject<AActor>();
-	if (!Actor)
-	{
-		return;
-	}
-
-	// Build the tree data source according to what mode we're in
-	if (!InActorNode->IsInstanced())
-	{
-		TInlineComponentArray<UActorComponent*> Components;
-		Actor->GetComponents(Components);
-
-		// Add the native root component
-		USceneComponent* RootComponent = Actor->GetRootComponent();
-		if (RootComponent != nullptr)
-		{
-			Components.Remove(RootComponent);
-			AddTreeNodeFromComponent(RootComponent, FindOrCreateParentForExistingComponent(RootComponent, InActorNode));
-		}
-
-		// Add the rest of the native base class SceneComponent hierarchy
-		for (UActorComponent* Component : Components)
-		{
-			AddTreeNodeFromComponent(Component, FindOrCreateParentForExistingComponent(Component, InActorNode));
-		}
-
-		// If it's a Blueprint-generated class, also get the inheritance stack
-		TArray<UBlueprintGeneratedClass*> ParentBPStack;
-		UBlueprint::GetBlueprintHierarchyFromClass(Actor->GetClass(), ParentBPStack);
-
-		UBlueprint* ActorBP = (ParentBPStack.Num() > 0 && ParentBPStack[0]) ? Cast<UBlueprint>(ParentBPStack[0]->ClassGeneratedBy) : nullptr;
-		ensure(ActorBP);
-
-		// Add the full SCS tree node hierarchy (including SCS nodes inherited from parent blueprints)
-		for (int32 StackIndex = ParentBPStack.Num() - 1; StackIndex >= 0; --StackIndex)
-		{
-			USimpleConstructionScript* ParentSCS = ParentBPStack[StackIndex] ? ParentBPStack[StackIndex]->SimpleConstructionScript.Get() : nullptr;
-			if (ParentSCS)
-			{
-				for (USCS_Node* SCS_Node : ParentSCS->GetRootNodes())
-				{
-					check(SCS_Node);
-
-					FSCSEditorTreeNodePtrType NewNodePtr;
-					if (SCS_Node->ParentComponentOrVariableName != NAME_None)
-					{
-						USceneComponent* ParentComponent = SCS_Node->GetParentComponentTemplate(ActorBP);
-						if (ParentComponent)
-						{
-							FSCSEditorTreeNodePtrType ParentNodePtr = FindTreeNode(ParentComponent, InActorNode);
-							if (ParentNodePtr.IsValid())
-							{
-								NewNodePtr = AddTreeNode(SCS_Node, ParentNodePtr, StackIndex > 0);
-							}
-						}
-					}
-					else
-					{
-						NewNodePtr = AddTreeNode(SCS_Node, InActorNode, StackIndex > 0);
-					}
-
-					// Only necessary to do the following for inherited nodes (StackIndex > 0).
-					if (NewNodePtr.IsValid() && StackIndex > 0)
-					{
-						// This call creates ICH override templates for the current Blueprint. Without this, the parent node
-						// search above can fail when attempting to match an inherited node in the tree via component template.
-						NewNodePtr->GetOrCreateEditableComponentTemplate(ActorBP);
-						for (FSCSEditorTreeNodePtrType ChildNodePtr : NewNodePtr->GetChildren())
-						{
-							if (ensure(ChildNodePtr.IsValid()))
-							{
-								ChildNodePtr->GetOrCreateEditableComponentTemplate(ActorBP);
-							}
-						}
-					}
-				}
-			}
-		}
-	}
-	else    // InActorNode->IsInstanced()
-	{
-		// Get the full set of instanced components
-		TSet<UActorComponent*> ComponentsToAdd(Actor->GetComponents());
-
-		const bool bHideConstructionScriptComponentsInDetailsView = GetDefault<UBlueprintEditorSettings>()->bHideConstructionScriptComponentsInDetailsView;
-		auto ShouldAddInstancedActorComponent = [bHideConstructionScriptComponentsInDetailsView](UActorComponent* ActorComp, USceneComponent* ParentSceneComp)
-		{
-			// Exclude nested DSOs attached to BP-constructed instances, which are not mutable.
-			return (ActorComp != nullptr
-				&& (!ActorComp->IsVisualizationComponent())
-				&& (ActorComp->CreationMethod != EComponentCreationMethod::UserConstructionScript || !bHideConstructionScriptComponentsInDetailsView)
-				&& (ParentSceneComp == nullptr || !ParentSceneComp->IsCreatedByConstructionScript() || !ActorComp->HasAnyFlags(RF_DefaultSubObject)))
-				&& (ActorComp->CreationMethod != EComponentCreationMethod::Native || FComponentEditorUtils::GetPropertyForEditableNativeComponent(ActorComp));
-		};
-
-		for (TSet<UActorComponent*>::TIterator It(ComponentsToAdd.CreateIterator()); It; ++It)
-		{
-			UActorComponent* ActorComp = *It;
-			USceneComponent* SceneComp = Cast<USceneComponent>(ActorComp);
-			USceneComponent* ParentSceneComp = SceneComp != nullptr ? SceneComp->GetAttachParent() : nullptr;
-			if (!ShouldAddInstancedActorComponent(ActorComp, ParentSceneComp))
-			{
-				It.RemoveCurrent();
-			}
-		}
-
-		TFunction<void(USceneComponent*, FSCSEditorTreeNodePtrType)> AddInstancedTreeNodesRecursive = [&](USceneComponent* Component, FSCSEditorTreeNodePtrType TreeNode)
-		{
-			if (Component != nullptr)
-			{
-				TArray<USceneComponent*> Components = Component->GetAttachChildren();
-				for (USceneComponent* ChildComponent : Components)
-				{
-					if (ComponentsToAdd.Contains(ChildComponent)
-						&& ChildComponent->GetOwner() == Component->GetOwner())
-					{
-						ComponentsToAdd.Remove(ChildComponent);
-
-						FSCSEditorTreeNodePtrType NewParentNode = AddTreeNodeFromComponent(ChildComponent, TreeNode);
-						AddInstancedTreeNodesRecursive(ChildComponent, NewParentNode);
-					}
-				}
-			}
-		};
-
-		// Add the root component first (it may not be the first one)
-		USceneComponent* RootComponent = Actor->GetRootComponent();
-		if (RootComponent != nullptr)
-		{
-			ComponentsToAdd.Remove(RootComponent);
-
-			// Recursively add any instanced children that are already attached through the root, and keep track of added
-			// instances. This will be a faster path than the loop below, because we create new parent tree nodes as we go.
-			FSCSEditorTreeNodePtrType NewParentNode = AddTreeNodeFromComponent(RootComponent, FindOrCreateParentForExistingComponent(RootComponent, InActorNode));
-			AddInstancedTreeNodesRecursive(RootComponent, NewParentNode);
-		}
-
-		// Sort components by type (always put scene components first in the tree)
-		ComponentsToAdd.Sort([](const UActorComponent& A, const UActorComponent& /* B */)
-		{
-			return A.IsA<USceneComponent>();
-		});
-
-		// Now add any remaining instanced owned components not already added above. This will first add any
-		// unattached scene components followed by any instanced non-scene components owned by the Actor instance.
-		for (UActorComponent* ActorComp : ComponentsToAdd)
-		{
-			AddTreeNodeFromComponent(ActorComp, FindOrCreateParentForExistingComponent(ActorComp, InActorNode));
-		}
-	}
-
-	// Always expand actor nodes to reveal children
-	SCSTreeWidget->SetItemExpansion(InActorNode, true);
-}
-
-void SSCSEditor::UpdateTree(bool bRegenerateTreeNodes)
-{
-	check(SCSTreeWidget.IsValid());
-
-	// Early exit if we're deferring tree updates
-	if(!bAllowTreeUpdates)
-	{
-		return;
-	}
-
-	if(bRegenerateTreeNodes)
-	{
-		// Obtain the set of expandable tree nodes that are currently collapsed
-		TSet<FSCSEditorTreeNodePtrType> CollapsedTreeNodes;
-		GetCollapsedNodes(GetSceneRootNode(), CollapsedTreeNodes);
-
-		// Obtain the list of selected items
-		TArray<FSCSEditorTreeNodePtrType> SelectedTreeNodes = SCSTreeWidget->GetSelectedItems();
-
-		// Clear the current tree
-		if (SelectedTreeNodes.Num() != 0)
-		{
-			SCSTreeWidget->ClearSelection();
-		}
-		RootNodes.Empty();
-
-		TSharedPtr<FSCSEditorTreeNodeRootActor> RootActorNode = MakeShareable(new FSCSEditorTreeNodeRootActor(GetActorContext(), EditorMode == EComponentEditorMode::ActorInstance));
-		RefreshFilteredState(RootActorNode, false);
-		SCSTreeWidget->SetItemExpansion(RootActorNode, true);
-		RootNodes.Add(RootActorNode);
-
-		BuildSubTreeForActorNode(RootActorNode);
-
-		AActor* PreviewActorInstance = PreviewActor.Get();
-		if (PreviewActorInstance != nullptr && !GetDefault<UBlueprintEditorSettings>()->bHideConstructionScriptComponentsInDetailsView)
-		{
-			TInlineComponentArray<UActorComponent*> Components;
-			PreviewActorInstance->GetComponents(Components);
-
-			for (UActorComponent* Component : Components)
-			{
-				if (Component->CreationMethod == EComponentCreationMethod::UserConstructionScript)
-				{
-					AddTreeNodeFromComponent(Component, FindOrCreateParentForExistingComponent(Component, RootActorNode));
-				}
-			}
-		}
-
-		// Restore the previous expansion state on the new tree nodes
-		TArray<FSCSEditorTreeNodePtrType> CollapsedTreeNodeArray = CollapsedTreeNodes.Array();
-		for(int32 i = 0; i < CollapsedTreeNodeArray.Num(); ++i)
-		{
-			// Look for a component match in the new hierarchy; if found, mark it as collapsed to match the previous setting
-			FSCSEditorTreeNodePtrType NodeToExpandPtr = FindTreeNode(CollapsedTreeNodeArray[i]->GetComponentTemplate());
-			if(NodeToExpandPtr.IsValid())
-			{
-				SCSTreeWidget->SetItemExpansion(NodeToExpandPtr, false);
-			}
-		}
-
-		if(SelectedTreeNodes.Num() > 0)
-		{
-			// If there is only one item selected, imitate user selection to preserve navigation
-			ESelectInfo::Type SelectInfo = SelectedTreeNodes.Num() == 1 ? ESelectInfo::OnMouseClick : ESelectInfo::Direct;
-
-			// Restore the previous selection state on the new tree nodes
-			for (int i = 0; i < SelectedTreeNodes.Num(); ++i)
-			{
-				if (SelectedTreeNodes[i]->GetNodeType() == FSCSEditorTreeNode::RootActorNode)
-				{
-					SCSTreeWidget->SetItemSelection(RootActorNode, true, SelectInfo);
-				}
-				else
-				{
-					FSCSEditorTreeNodePtrType NodeToSelectPtr = FindTreeNode(SelectedTreeNodes[i]->GetComponentTemplate());
-					if (NodeToSelectPtr.IsValid())
-					{
-						SCSTreeWidget->SetItemSelection(NodeToSelectPtr, true, SelectInfo);
-					}
-				}
-			}
-
-			if (GetEditorMode() != EComponentEditorMode::BlueprintSCS)
-			{
-				TArray<FSCSEditorTreeNodePtrType> NewSelectedTreeNodes = SCSTreeWidget->GetSelectedItems();
-				if (NewSelectedTreeNodes.Num() == 0)
-				{
-					SCSTreeWidget->SetItemSelection(GetRootNodes()[0], true);
-				}
-			}
-		}
-
-		// If we have a pending deferred rename request, redirect it to the new tree node
-		if(DeferredRenameRequest != NAME_None)
-		{
-			FSCSEditorTreeNodePtrType NodeToRenamePtr = FindTreeNode(DeferredRenameRequest);
-			if(NodeToRenamePtr.IsValid())
-			{
-				SCSTreeWidget->RequestScrollIntoView(NodeToRenamePtr);
-			}
-		}
-	}
-
-	// refresh widget
-	SCSTreeWidget->RequestTreeRefresh();
-}
-
-void SSCSEditor::DumpTree()
-{
-	/* Example:
-
-		[ACTOR] MyBlueprint (self)
-		|
-		[SEPARATOR]
-		|
-		DefaultSceneRoot (Inherited)
-		|
-		+- StaticMesh (Inherited)
-		|  |
-		|  +- Scene4 (Inherited)
-		|  |
-		|  +- Scene (Inherited)
-		|     |
-		|     +- Scene1 (Inherited)
-		|  
-		+- Scene2 (Inherited)
-		|  |
-		|  +- Scene3 (Inherited)
-		|
-		[SEPARATOR]
-		|
-		ProjectileMovement (Inherited)
-	*/
-
-	UE_LOG(LogSCSEditor, Log, TEXT("---------------------"));
-	UE_LOG(LogSCSEditor, Log, TEXT(" STreeView NODE DUMP"));
-	UE_LOG(LogSCSEditor, Log, TEXT("---------------------"));
-
-	const UBlueprint* BlueprintContext = nullptr;
-	const AActor* ActorInstance = GetActorContext();
-	if (ActorInstance)
-	{
-		BlueprintContext = UBlueprint::GetBlueprintFromClass(ActorInstance->GetClass());
-	}
-
-	TArray<TArray<FSCSEditorTreeNodePtrType>> NodeListStack;
-	NodeListStack.Push(RootNodes);
-
-	auto LineSpacingLambda = [&NodeListStack](const TArray<FSCSEditorTreeNodePtrType>& NodeList, int32 CurrentDepth, const FString& Prefix)
-	{
-		bool bAddLineSpacing = false;
-		for (int Depth = 0; Depth <= CurrentDepth && !bAddLineSpacing; ++Depth)
-		{
-			bAddLineSpacing = NodeListStack[Depth].Num() > 0;
-		}
-
-		if (bAddLineSpacing)
-		{
-			UE_LOG(LogSCSEditor, Log, TEXT(" %s%s"), *Prefix, NodeList.Num() > 0 ? TEXT("|") : TEXT(""));
-		}
-	};
-
-	while (NodeListStack.Num() > 0)
-	{
-		const int32 CurrentDepth = NodeListStack.Num() - 1;
-		TArray<FSCSEditorTreeNodePtrType>& NodeList = NodeListStack[CurrentDepth];
-		if (NodeList.Num() > 0)
-		{
-			FString Prefix;
-			for (int32 Depth = 1; Depth < CurrentDepth; ++Depth)
-			{
-				int32 NodeCount = NodeListStack[Depth].Num();
-				if (Depth == 1)
-				{
-					NodeCount += NodeListStack[0].Num();
-				}
-
-				Prefix += (NodeCount > 0) ? TEXT("|  ") : TEXT("   ");
-			}
-
-			FString NodePrefix;
-			if (CurrentDepth > 0)
-			{
-				NodePrefix = TEXT("+- ");
-			}
-
-			FSCSEditorTreeNodePtrType Node = NodeList[0];
-			NodeList.RemoveAt(0);
-
-			if (Node.IsValid())
-			{
-				FString NodeLabel = TEXT("[UNKNOWN]");
-				switch (Node->GetNodeType())
-				{
-				case FSCSEditorTreeNode::ENodeType::RootActorNode:
-					switch (EditorMode)
-					{
-					case EComponentEditorMode::ActorInstance:
-						NodeLabel = TEXT("[ACTOR]");
-						break;
-
-					case EComponentEditorMode::BlueprintSCS:
-						NodeLabel = TEXT("[BLUEPRINT]");
-						break;
-					}
-
-					if (BlueprintContext)
-					{
-						NodeLabel += FString::Printf(TEXT(" %s (self)"), *BlueprintContext->GetName());
-					}
-					else if (ActorInstance)
-					{
-						NodeLabel += FString::Printf(TEXT(" %s (Instance)"), *ActorInstance->GetActorLabel());
-					}
-					break;
-
-				case FSCSEditorTreeNode::ENodeType::SeparatorNode:
-					NodeLabel = TEXT("[SEPARATOR]");
-					break;
-
-				case FSCSEditorTreeNode::ENodeType::ComponentNode:
-					NodeLabel = Node->GetDisplayString();
-					if (Node->IsInheritedComponent())
-					{
-						NodeLabel += TEXT(" (Inherited)");
-					}
-					break;
-
-				case FSCSEditorTreeNode::ENodeType::ChildActorNode:
-					NodeLabel = Node->GetDisplayString();
-					NodeLabel += TEXT(" [CHILD ACTOR]");
-					break;
-				}
-
-				UE_LOG(LogSCSEditor, Log, TEXT(" %s%s%s"), *Prefix, *NodePrefix, *NodeLabel);
-
-				const TArray<FSCSEditorTreeNodePtrType>& Children = Node->GetChildren();
-				if (Children.Num() > 0)
-				{
-					if (CurrentDepth > 1)
-					{
-						UE_LOG(LogSCSEditor, Log, TEXT(" %s%s|"), *Prefix, NodeListStack[CurrentDepth].Num() > 0 ? TEXT("|  ") : TEXT("   "));
-					}
-					else if (CurrentDepth == 1)
-					{
-						UE_LOG(LogSCSEditor, Log, TEXT(" %s%s|"), *Prefix, NodeListStack[0].Num() > 0 ? TEXT("|  ") : TEXT("   "));
-					}
-					else
-					{
-						UE_LOG(LogSCSEditor, Log, TEXT(" %s|"), *Prefix);
-					}
-
-					NodeListStack.Push(Children);
-				}
-				else
-				{
-					LineSpacingLambda(NodeList, CurrentDepth, Prefix);
-				}
-			}
-			else
-			{
-				UE_LOG(LogSCSEditor, Log, TEXT(" %s%s[INVALID]"), *Prefix, *NodePrefix);
-				
-				LineSpacingLambda(NodeList, CurrentDepth, Prefix);
-			}
-		}
-		else
-		{
-			NodeListStack.Pop();
-		}
-	}
-
-	UE_LOG(LogSCSEditor, Log, TEXT("--------(end)--------"));
-}
-
-const TArray<FSCSEditorTreeNodePtrType>& SSCSEditor::GetRootNodes() const
-{
-	return RootNodes;
-}
-
-FSCSEditorActorNodePtrType SSCSEditor::GetActorNode() const
-{
-	if (RootNodes.Num() > 0)
-	{
-		return StaticCastSharedPtr<FSCSEditorTreeNodeRootActor>(RootNodes[0]);
-	}
-
-	return FSCSEditorActorNodePtrType();
-}
-
-FSCSEditorTreeNodePtrType SSCSEditor::GetSceneRootNode() const
-{
-	FSCSEditorActorNodePtrType ActorNode = GetActorNode();
-	if (ActorNode.IsValid())
-	{
-		return ActorNode->GetSceneRootNode();
-	}
-
-	return FSCSEditorTreeNodePtrType();
-}
-
-void SSCSEditor::SetSceneRootNode(FSCSEditorTreeNodePtrType NewSceneRootNode)
-{
-	GetActorNode()->SetSceneRootNode(NewSceneRootNode);
-}
-
-class FComponentClassParentFilter : public IClassViewerFilter
-{
-public:
-	FComponentClassParentFilter(const TSubclassOf<UActorComponent>& InComponentClass) : ComponentClass(InComponentClass) {}
-
-	virtual bool IsClassAllowed(const FClassViewerInitializationOptions& InInitOptions, const UClass* InClass, TSharedRef< FClassViewerFilterFuncs > InFilterFuncs ) override
-	{
-		return InClass->IsChildOf(ComponentClass);
-	}
-
-	virtual bool IsUnloadedClassAllowed(const FClassViewerInitializationOptions& InInitOptions, const TSharedRef< const IUnloadedBlueprintData > InUnloadedClassData, TSharedRef< FClassViewerFilterFuncs > InFilterFuncs) override
-	{
-		return InUnloadedClassData->IsChildOf(ComponentClass);
-	}
-
-	TSubclassOf<UActorComponent> ComponentClass;
-};
-
-typedef FComponentClassParentFilter FNativeComponentClassParentFilter;
-
-class FBlueprintComponentClassParentFilter : public FComponentClassParentFilter
-{
-public:
-	FBlueprintComponentClassParentFilter(const TSubclassOf<UActorComponent>& InComponentClass) : FComponentClassParentFilter(InComponentClass) {}
-
-	virtual bool IsClassAllowed(const FClassViewerInitializationOptions& InInitOptions, const UClass* InClass, TSharedRef< FClassViewerFilterFuncs > InFilterFuncs ) override
-	{
-		return FComponentClassParentFilter::IsClassAllowed(InInitOptions, InClass, InFilterFuncs) && FKismetEditorUtilities::CanCreateBlueprintOfClass(InClass);
-	}
-};
-
-UClass* SSCSEditor::CreateNewCPPComponent( TSubclassOf<UActorComponent> ComponentClass )
-{
-	TSharedPtr<SWindow> ParentWindow = FSlateApplication::Get().FindWidgetWindow(SharedThis(this));
-
-	FString AddedClassName;
-	auto OnCodeAddedToProject = [&AddedClassName](const FString& ClassName, const FString& ClassPath, const FString& ModuleName)
-	{
-		if(!ClassName.IsEmpty() && !ClassPath.IsEmpty())
-		{
-			AddedClassName = FString::Printf(TEXT("/Script/%s.%s"), *ModuleName, *ClassName);
-		}
-	};
-
-	FGameProjectGenerationModule::Get().OpenAddCodeToProjectDialog(
-		FAddToProjectConfig()
-		.WindowTitle(LOCTEXT("AddNewC++Component", "Add C++ Component"))
-		.ParentWindow(ParentWindow)
-		.Modal()
-		.OnAddedToProject(FOnAddedToProject::CreateLambda(OnCodeAddedToProject))
-		.FeatureComponentClasses()
-		.AllowableParents(MakeShareable( new FNativeComponentClassParentFilter(ComponentClass) ))
-		.DefaultClassPrefix(TEXT("New"))
-	);
-
-
-	return LoadClass<UActorComponent>(nullptr, *AddedClassName, nullptr, LOAD_None, nullptr);
-}
-
-UClass* SSCSEditor::CreateNewBPComponent(TSubclassOf<UActorComponent> ComponentClass)
-{
-	UClass* NewClass = nullptr;
-
-	auto OnAddedToProject = [&](const FString& ClassName, const FString& PackagePath, const FString& ModuleName)
-	{
-		if(!ClassName.IsEmpty() && !PackagePath.IsEmpty())
-		{
-			if (UPackage* Package = FindPackage(nullptr, *PackagePath))
-			{
-				if (UBlueprint* NewBP = FindObjectFast<UBlueprint>(Package, *ClassName))	
-				{
-					NewClass = NewBP->GeneratedClass;
-
-					TArray<UObject*> Objects;
-					Objects.Emplace(NewBP);
-					GEditor->SyncBrowserToObjects(Objects);
-
-					// Open the editor for the new blueprint
-					GEditor->GetEditorSubsystem<UAssetEditorSubsystem>()->OpenEditorForAsset(NewBP);
-				}
-			}
-		}
-	};
-
-	FGameProjectGenerationModule::Get().OpenAddBlueprintToProjectDialog(
-		FAddToProjectConfig()
-		.WindowTitle(LOCTEXT("AddNewBlueprintComponent", "Add Blueprint Component"))
-		.ParentWindow(FSlateApplication::Get().FindWidgetWindow(SharedThis(this)))
-		.Modal()
-		.AllowableParents(MakeShareable( new FBlueprintComponentClassParentFilter(ComponentClass) ))
-		.FeatureComponentClasses()
-		.OnAddedToProject(FOnAddedToProject::CreateLambda(OnAddedToProject))
-		.DefaultClassPrefix(TEXT("New"))
-	);
-
-	return NewClass;
-}
-
-void SSCSEditor::ClearSelection()
-{
-	if ( bUpdatingSelection == false )
-	{
-		check(SCSTreeWidget.IsValid());
-		SCSTreeWidget->ClearSelection();
-	}
-}
-
-void SSCSEditor::SaveSCSCurrentState( USimpleConstructionScript* SCSObj )
-{
-	if (SCSObj)
-	{
-		SCSObj->SaveToTransactionBuffer();
-	}
-}
-
-void SSCSEditor::SaveSCSNode( USCS_Node* Node )
-{
-	if (Node)
-	{
-		Node->SaveToTransactionBuffer();
-	}
-}
-
-bool SSCSEditor::IsEditingAllowed() const
-{
-	return AllowEditing.Get() && nullptr == GEditor->PlayWorld;
-}
-
-UActorComponent* SSCSEditor::AddNewComponent( UClass* NewComponentClass, UObject* Asset, const FAddNewComponentParams Params)
-{
-	if (NewComponentClass->ClassWithin && NewComponentClass->ClassWithin != UObject::StaticClass())
-	{
-		FNotificationInfo Info(LOCTEXT("AddComponentFailed", "Cannot add components that have \"Within\" markup"));
-		Info.Image = FAppStyle::GetBrush(TEXT("Icons.Error"));
-		Info.bFireAndForget = true;
-		Info.bUseSuccessFailIcons = false;
-		Info.ExpireDuration = 5.0f;
-
-		FSlateNotificationManager::Get().AddNotification(Info);
-		return nullptr;
-	}
-
-	// If an 'add' transaction is ongoing, it is most likely because AddNewComponent() is being called in a tight loop inside a larger transaction (e.g. 'duplicate')
-	// and bSetFocusToNewItem was true for each element.
-	if (DeferredOngoingCreateTransaction.IsValid() && Params.bSetFocusToNewItem)
-	{
-		// Close the ongoing 'add' sub-transaction before staring another one. The user will not be able to edit the name of that component because the
-		// new component is going to still focus.
-		DeferredOngoingCreateTransaction.Reset();
-	}
-
-	// Begin a transaction. The transaction will end when the component name will be provided/confirmed by the user.
-	TUniquePtr<FScopedTransaction> AddTransaction = MakeUnique<FScopedTransaction>( LOCTEXT("AddComponent", "Add Component") );
-
-	UActorComponent* NewComponent = nullptr;
-	FName TemplateVariableName;
-
-	USCS_Node* SCSNode = Cast<USCS_Node>(Asset);
-	UActorComponent* ComponentTemplate = (SCSNode ? ToRawPtr(SCSNode->ComponentTemplate) : Cast<UActorComponent>(Asset));
-
-	if (SCSNode)
-	{
-		TemplateVariableName = SCSNode->GetVariableName();
-		Asset = nullptr;
-	}
-	else if (ComponentTemplate)
-	{
-		Asset = nullptr;
-	}
-
-	if (EditorMode == EComponentEditorMode::BlueprintSCS)
-	{
-		UBlueprint* Blueprint = GetBlueprint();
-		check(Blueprint != nullptr && Blueprint->SimpleConstructionScript != nullptr);
-		
-		Blueprint->Modify();
-		SaveSCSCurrentState(Blueprint->SimpleConstructionScript);
-
-		// Defer Blueprint class regeneration and tree updates until after we copy any object properties from a source template.
-		const bool bMarkBlueprintModified = false;
-		bAllowTreeUpdates = false;
-		
-		FName NewVariableName;
-		if (ComponentTemplate)
-		{
-			if (!TemplateVariableName.IsNone())
-			{
-				NewVariableName = TemplateVariableName;
-			}
-			else
-			{
-				FString TemplateName = ComponentTemplate->GetName();
-				NewVariableName = (TemplateName.EndsWith(USimpleConstructionScript::ComponentTemplateNameSuffix) 
-									? FName(*TemplateName.LeftChop(USimpleConstructionScript::ComponentTemplateNameSuffix.Len()))
-									: ComponentTemplate->GetFName());
-			}
-		}
-		else if (Asset)
-		{
-			NewVariableName = *FComponentEditorUtils::GenerateValidVariableNameFromAsset(Asset, nullptr);
-		}
-
-		USCS_Node* NewSCSNode = Blueprint->SimpleConstructionScript->CreateNode(NewComponentClass, NewVariableName);
-		NewComponent = NewSCSNode->ComponentTemplate;
-
-		FAddedNodeDetails NewNodeDetails;
-		AddNewNode(NewNodeDetails, MoveTemp(AddTransaction), NewSCSNode, Asset, bMarkBlueprintModified, Params.bSetFocusToNewItem);
-
-		if (ComponentTemplate)
-		{
-			UEngine::CopyPropertiesForUnrelatedObjects(ComponentTemplate, NewComponent);
-			NewComponent->UpdateComponentToWorld();
-		}
-
-		if (Params.bConformTransformToParent)
-		{
-			if (USceneComponent* AsSceneComp = Cast<USceneComponent>(NewComponent))
-			{
-				if (USceneComponent* ParentSceneComp = CastChecked<USceneComponent>(NewNodeDetails.ParentNodePtr->GetComponentTemplate(), ECastCheckedType::NullAllowed))
-				{ 
-					ConformTransformRelativeToParent(AsSceneComp, ParentSceneComp);
-				}
-			}
-		}
-
-		// Wait until here to mark as structurally modified because we don't want any RerunConstructionScript() calls to happen until AFTER we've serialized properties from the source object.
-		if (!Params.bSkipMarkBlueprintModified)
-		{
-			bAllowTreeUpdates = true;
-			FBlueprintEditorUtils::MarkBlueprintAsStructurallyModified(Blueprint);
-		}
-	}
-	else    // EComponentEditorMode::ActorInstance
-	{
-		if (ComponentTemplate)
-		{
-			// Create a duplicate of the provided template
-			NewComponent = FComponentEditorUtils::DuplicateComponent(ComponentTemplate);
-			FSCSEditorTreeNodePtrType ParentNodePtr = FindParentForNewComponent(NewComponent);
-			AddNewNodeForInstancedComponent(MoveTemp(AddTransaction), NewComponent, ParentNodePtr, nullptr, Params.bSetFocusToNewItem);
-		}
-		else if (AActor* ActorInstance = GetActorContext())
-		{
-			// No template, so create a wholly new component
-			ActorInstance->Modify();
-
-			// Create an appropriate name for the new component
-			FName NewComponentName = NAME_None;
-			if (Asset)
-			{
-				NewComponentName = *FComponentEditorUtils::GenerateValidVariableNameFromAsset(Asset, ActorInstance);
-			}
-			else
-			{
-				NewComponentName = *FComponentEditorUtils::GenerateValidVariableName(NewComponentClass, ActorInstance);
-			}
-
-			// Get the set of owned components that exists prior to instancing the new component.
-			TInlineComponentArray<UActorComponent*> PreInstanceComponents;
-			ActorInstance->GetComponents(PreInstanceComponents);
-
-			// Construct the new component and attach as needed
-			UActorComponent* NewInstanceComponent = NewObject<UActorComponent>(ActorInstance, NewComponentClass, NewComponentName, RF_Transactional);
-			FSCSEditorTreeNodePtrType ParentNodePtr = FindParentForNewComponent(NewInstanceComponent);
-						
-			// Do Scene Attachment if this new Component is a USceneComponent
-			if (USceneComponent* NewSceneComponent = Cast<USceneComponent>(NewInstanceComponent))
-			{
-				if(ParentNodePtr->GetNodeType() == FSCSEditorTreeNode::RootActorNode)
-				{
-					ActorInstance->SetRootComponent(NewSceneComponent);
-				}
-				else
-				{
-					USceneComponent* AttachTo = Cast<USceneComponent>(ParentNodePtr->GetComponentTemplate());
-					if (AttachTo == nullptr)
-					{
-						AttachTo = ActorInstance->GetRootComponent();
-					}
-					check(AttachTo != nullptr);
-
-					// Make sure that the mobility of the new scene component is such that we can attach it
-					if (AttachTo->Mobility == EComponentMobility::Movable)
-					{
-						NewSceneComponent->Mobility = EComponentMobility::Movable;
-					}
-					else if (AttachTo->Mobility == EComponentMobility::Stationary && NewSceneComponent->Mobility == EComponentMobility::Static)
-					{
-						NewSceneComponent->Mobility = EComponentMobility::Stationary;
-					}
-
-					NewSceneComponent->AttachToComponent(AttachTo, FAttachmentTransformRules::SnapToTargetNotIncludingScale);
-				}
-			}
-
-			// If the component was created from/for a particular asset, assign it now
-			if (Asset)
-			{
-				FComponentAssetBrokerage::AssignAssetToComponent(NewInstanceComponent, Asset);
-			}
-
-			// Add to SerializedComponents array so it gets saved
-			ActorInstance->AddInstanceComponent(NewInstanceComponent);
-			NewInstanceComponent->OnComponentCreated();
-			NewInstanceComponent->RegisterComponent();
-
-			// Register any new components that may have been created during construction of the instanced component, but were not explicitly registered.
-			TInlineComponentArray<UActorComponent*> PostInstanceComponents;
-			ActorInstance->GetComponents(PostInstanceComponents);
-			for (UActorComponent* ActorComponent : PostInstanceComponents)
-			{
-				if (!ActorComponent->IsRegistered() && ActorComponent->bAutoRegister && IsValid(ActorComponent) && !PreInstanceComponents.Contains(ActorComponent))
-				{
-					ActorComponent->RegisterComponent();
-				}
-			}
-
-			// Rerun construction scripts
-			ActorInstance->RerunConstructionScripts();
-
-			// If the running the construction script destroyed the new node, don't create an entry for it
-			if (IsValid(NewInstanceComponent))
-			{
-				AddNewNodeForInstancedComponent(MoveTemp(AddTransaction), NewInstanceComponent, ParentNodePtr, Asset, Params.bSetFocusToNewItem);
-				NewComponent = NewInstanceComponent;
-			}
-		}
-	}
-
-	return NewComponent;
-}
-
-FSCSEditorTreeNodePtrType SSCSEditor::FindOrCreateParentForExistingComponent(UActorComponent* InActorComponent, FSCSEditorActorNodePtrType ActorRootNode)
-{
-	check(InActorComponent != nullptr);
-
-	USceneComponent* SceneComponent = Cast<USceneComponent>(InActorComponent);
-	if (SceneComponent == nullptr)
-	{
-		check(ActorRootNode.IsValid());
-		check(ActorRootNode->IsActorNode());
-		return ActorRootNode;
-	}
-
-	FSCSEditorTreeNodePtrType ParentNodePtr;
-	if (SceneComponent->GetAttachParent() != nullptr
-		&& (EditorMode != EComponentEditorMode::ActorInstance || SceneComponent->GetAttachParent()->GetOwner() == ActorRootNode->GetObject<AActor>()))
-	{
-		// Attempt to find the parent node in the current tree
-		ParentNodePtr = FindTreeNode(SceneComponent->GetAttachParent(), ActorRootNode);
-		if (!ParentNodePtr.IsValid())
-		{
-			// If the actual attach parent wasn't found, attempt to find its archetype.
-			// This handles the BP editor case where we might add UCS component nodes taken
-			// from the preview actor instance, which are not themselves template objects.
-			ParentNodePtr = FindTreeNode(Cast<USceneComponent>(SceneComponent->GetAttachParent()->GetArchetype()), ActorRootNode);
-			if (!ParentNodePtr.IsValid())
-			{
-				// Recursively add the parent node to the tree if it does not exist yet
-				ParentNodePtr = AddTreeNodeFromComponent(SceneComponent->GetAttachParent(), FindOrCreateParentForExistingComponent(SceneComponent->GetAttachParent(), ActorRootNode));
-			}
-		}
-	}
-
-	if (!ParentNodePtr.IsValid())
-	{
-		ParentNodePtr = ActorRootNode->GetSceneRootNode();
-	}
-
-	// Actor doesn't have a root component yet
-	if (!ParentNodePtr.IsValid())
-	{
-		ParentNodePtr = ActorRootNode; 
-	}
-
-	return ParentNodePtr;
-}
-
-FSCSEditorTreeNodePtrType SSCSEditor::FindParentForNewComponent(UActorComponent* NewComponent) const
-{
-	// Find Parent to attach to (depending on the new Node type).
-	FSCSEditorTreeNodePtrType TargetParentNode;
-	TArray<FSCSEditorTreeNodePtrType> SelectedTreeNodes;
-	if (SCSTreeWidget.IsValid() && SCSTreeWidget->GetSelectedItems(SelectedTreeNodes))
-	{
-		TargetParentNode = SelectedTreeNodes[0];
-	}
-
-	// If the current selection belongs to a child actor template, move the target to its outer component node.
-	while (TargetParentNode.IsValid() && FChildActorComponentEditorUtils::IsChildActorSubtreeNode(TargetParentNode))
-	{
-		TargetParentNode = FChildActorComponentEditorUtils::GetOuterChildActorComponentNode(TargetParentNode);
-	}
-
-	if (USceneComponent* NewSceneComponent = Cast<USceneComponent>(NewComponent))
-	{
-		if (TargetParentNode.IsValid())
-		{
-			if (TargetParentNode->IsActorNode())
-			{
-				FSCSEditorActorNodePtrType TargetActorNode = StaticCastSharedPtr<FSCSEditorTreeNodeActorBase>(TargetParentNode);
-				if (TargetActorNode.IsValid())
-				{
-					FSCSEditorTreeNodePtrType TargetSceneRootNode = TargetActorNode->GetSceneRootNode();
-					if (TargetSceneRootNode.IsValid())
-					{
-						TargetParentNode = TargetSceneRootNode;
-						USceneComponent* CastTargetToSceneComponent = Cast<USceneComponent>(TargetParentNode->GetComponentTemplate());
-						if (CastTargetToSceneComponent == nullptr || !NewSceneComponent->CanAttachAsChild(CastTargetToSceneComponent, NAME_None))
-						{
-							TargetParentNode = GetSceneRootNode(); // Default to SceneRoot
-						}
-					}
-				}
-			}
-			else if(TargetParentNode->IsComponentNode())
-			{
-				USceneComponent* CastTargetToSceneComponent = Cast<USceneComponent>(TargetParentNode->GetComponentTemplate());
-				if (CastTargetToSceneComponent == nullptr || !NewSceneComponent->CanAttachAsChild(CastTargetToSceneComponent, NAME_None))
-				{
-					TargetParentNode = GetSceneRootNode(); // Default to SceneRoot
-				}
-			}
-		}
-		else
-		{
-			TargetParentNode = GetSceneRootNode();
-		}
-	}
-	else
-	{
-		if (TargetParentNode.IsValid())
-		{
-			while (!TargetParentNode->IsActorNode())
-			{
-				TargetParentNode = TargetParentNode->GetParent();
-			}
-		}
-		else
-		{
-			TargetParentNode = GetActorNode();
-		}
-
-		check(TargetParentNode.IsValid() && TargetParentNode->IsActorNode());
-	}
-
-	return TargetParentNode;
-}
-
-FSCSEditorTreeNodePtrType SSCSEditor::FindParentForNewNode(USCS_Node* NewNode) const
-{
-	return FindParentForNewComponent(NewNode->ComponentTemplate);
-}
-
-UActorComponent* SSCSEditor::AddNewNode(TUniquePtr<FScopedTransaction> OngoingCreateTransaction, USCS_Node* NewNode, UObject* Asset, bool bMarkBlueprintModified, bool bSetFocusToNewItem)
-{
-	FAddedNodeDetails AddedNodeDetails;
-	AddNewNode(AddedNodeDetails, MoveTemp(OngoingCreateTransaction), NewNode, Asset, bMarkBlueprintModified, bSetFocusToNewItem);
-	return NewNode->ComponentTemplate;
-}
-
-void SSCSEditor::AddNewNode(SSCSEditor::FAddedNodeDetails& OutAddedNodeDetails, TUniquePtr<FScopedTransaction> InOngoingCreateTransaction, USCS_Node* NewNode, UObject* Asset, bool bMarkBlueprintModified, bool bSetFocusToNewItem)
-{
-	check(NewNode != nullptr);
-
-	if(Asset)
-	{
-		FComponentAssetBrokerage::AssignAssetToComponent(NewNode->ComponentTemplate, Asset);
-	}
-
-	FSCSEditorTreeNodePtrType& NewNodePtr = OutAddedNodeDetails.NewNodePtr;
-	FSCSEditorTreeNodePtrType& ParentNodePtr = OutAddedNodeDetails.ParentNodePtr;
-	ParentNodePtr = FindParentForNewNode(NewNode);
-	
-	UBlueprint* Blueprint = GetBlueprint();
-	check(Blueprint != nullptr && Blueprint->SimpleConstructionScript != nullptr);
-
-	// Add the new node to the editor tree
-	NewNodePtr = AddTreeNode(NewNode, ParentNodePtr, /*bIsInheritedSCS=*/ false);
-
-	// Potentially adjust variable names for any child blueprints
-	const FName VariableName = NewNode->GetVariableName();
-	if(VariableName != NAME_None)
-	{
-		FBlueprintEditorUtils::ValidateBlueprintChildVariables(Blueprint, VariableName);
-	}
-	
-	if(bSetFocusToNewItem)
-	{
-		// Select and request a rename on the new component
-		SCSTreeWidget->SetSelection(NewNodePtr);
-		OnRenameComponent(MoveTemp(InOngoingCreateTransaction));
-	}
-
-	// Will call UpdateTree as part of OnBlueprintChanged handling
-	if(bMarkBlueprintModified)
-	{
-		FBlueprintEditorUtils::MarkBlueprintAsStructurallyModified(Blueprint);
-	}
-	else
-	{
-		UpdateTree();
-	}
-}
-
-void SSCSEditor::AddNewNodeForInstancedComponent(TUniquePtr<FScopedTransaction> InOngoingCreateTransaction, UActorComponent* NewInstanceComponent, FSCSEditorTreeNodePtrType InParentNodePtr, UObject* Asset, bool bSetFocusToNewItem)
-{
-	check(NewInstanceComponent != nullptr);
-
-	FSCSEditorTreeNodePtrType NewNodePtr;
-
-	// Add the new node to the editor tree
-	NewNodePtr = AddTreeNodeFromComponent(NewInstanceComponent, InParentNodePtr);
-
-	if(bSetFocusToNewItem)
-	{
-		// Select and request a rename on the new component
-		SCSTreeWidget->SetSelection(NewNodePtr);
-		OnRenameComponent(MoveTemp(InOngoingCreateTransaction));
-	}
-
-	UpdateTree(false);
-}
-
-bool SSCSEditor::IsComponentSelected(const UPrimitiveComponent* PrimComponent) const
-{
-	check(PrimComponent);
-
-	if (SCSTreeWidget.IsValid())
-	{
-		FSCSEditorTreeNodePtrType NodePtr = GetNodeFromActorComponent(PrimComponent, false);
-		if (NodePtr.IsValid())
-		{
-			return SCSTreeWidget->IsItemSelected(NodePtr);
-		}
-		else
-		{
-			UChildActorComponent* PossiblySelectedComponent = nullptr;
-			AActor* ComponentOwner = PrimComponent->GetOwner();
-			while (ComponentOwner->IsChildActor())
-			{
-				PossiblySelectedComponent = ComponentOwner->GetParentComponent();
-				ComponentOwner = ComponentOwner->GetParentActor();
-			}
-
-			if (PossiblySelectedComponent)
-			{
-				NodePtr = GetNodeFromActorComponent(PossiblySelectedComponent, false);
-				if (NodePtr.IsValid())
-				{
-					return SCSTreeWidget->IsItemSelected(NodePtr);
-				}
-			}
-		}
-	}
-
-	return false;
-}
-
-void SSCSEditor::SetSelectionOverride(UPrimitiveComponent* PrimComponent) const
-{
-	PrimComponent->SelectionOverrideDelegate = UPrimitiveComponent::FSelectionOverride::CreateSP(this, &SSCSEditor::IsComponentSelected);
-	PrimComponent->PushSelectionToProxy();
-}
-
-bool SSCSEditor::CanCutNodes() const
-{
-	return CanCopyNodes() && CanDeleteNodes();
-}
-
-void SSCSEditor::CutSelectedNodes()
-{
-	TArray<FSCSEditorTreeNodePtrType> SelectedNodes = GetSelectedNodes();
-	const FScopedTransaction Transaction( SelectedNodes.Num() > 1 ? LOCTEXT("CutComponents", "Cut Components") : LOCTEXT("CutComponent", "Cut Component") );
-
-	CopySelectedNodes();
-	OnDeleteNodes();
-}
-
-bool SSCSEditor::CanCopyNodes() const
-{
-	TArray<UActorComponent*> ComponentsToCopy;
-	TArray<FSCSEditorTreeNodePtrType> SelectedNodes = GetSelectedNodes();
-	for (int32 i = 0; i < SelectedNodes.Num(); ++i)
-	{
-		// Get the current selected node reference
-		FSCSEditorTreeNodePtrType SelectedNodePtr = SelectedNodes[i];
-		check(SelectedNodePtr.IsValid());
-
-		// Get the component template associated with the selected node
-		UActorComponent* ComponentTemplate = SelectedNodePtr->GetComponentTemplate();
-		if (ComponentTemplate)
-		{
-			ComponentsToCopy.Add(ComponentTemplate);
-		}
-	}
-
-	// Verify that the components can be copied
-	return FComponentEditorUtils::CanCopyComponents(ComponentsToCopy);
-}
-
-void SSCSEditor::CopySelectedNodes()
-{
-	// Distill the selected nodes into a list of components to copy
-	TArray<UActorComponent*> ComponentsToCopy;
-	TArray<FSCSEditorTreeNodePtrType> SelectedNodes = GetSelectedNodes();
-	for (int32 i = 0; i < SelectedNodes.Num(); ++i)
-	{
-		// Get the current selected node reference
-		FSCSEditorTreeNodePtrType SelectedNodePtr = SelectedNodes[i];
-		check(SelectedNodePtr.IsValid());
-
-		// Get the component template associated with the selected node
-		UActorComponent* ComponentTemplate = SelectedNodePtr->GetComponentTemplate();
-		if (ComponentTemplate)
-		{
-			ComponentsToCopy.Add(ComponentTemplate);
-
-			if (EditorMode == EComponentEditorMode::BlueprintSCS && ComponentTemplate->CreationMethod != EComponentCreationMethod::UserConstructionScript)
-			{
-				// CopyComponents uses component attachment to maintain hierarchy, but the SCS templates are not
-				// setup with a relationship to each other. Briefly setup the attachment between the templates being
-				// copied so that the hierarchy is retained upon pasting
-				if (USceneComponent* SceneTemplate = Cast<USceneComponent>(ComponentTemplate))
-				{
-					FSCSEditorTreeNodePtrType SelectedParentNodePtr = SelectedNodePtr->GetParent();
-					if (SelectedParentNodePtr.IsValid())
-					{
-						if (USceneComponent* ParentSceneTemplate = Cast<USceneComponent>(SelectedParentNodePtr->GetComponentTemplate()))
-						{
-							SceneTemplate->SetupAttachment(ParentSceneTemplate);
-						}
-					}
-				}
-			}
-		}
-	}
-
-	// Copy the components to the clipboard
-	FComponentEditorUtils::CopyComponents(ComponentsToCopy);
-
-	if (EditorMode == EComponentEditorMode::BlueprintSCS)
-	{
-		for (UActorComponent* ComponentTemplate : ComponentsToCopy)
-		{
-			if (ComponentTemplate->CreationMethod != EComponentCreationMethod::UserConstructionScript)
-			{
-				if (USceneComponent* SceneTemplate = Cast<USceneComponent>(ComponentTemplate))
-				{
-					// clear back out any temporary attachments we set up for the copy
-					SceneTemplate->SetupAttachment(nullptr);
-				}
-			}
-		}
-	}
-}
-
-bool SSCSEditor::CanPasteNodes() const
-{
-	if(!IsEditingAllowed())
-	{
-		return false;
-	}
-
-	FSCSEditorTreeNodePtrType SceneRootNodePtr = GetSceneRootNode();
-	return SceneRootNodePtr.IsValid() && FComponentEditorUtils::CanPasteComponents(Cast<USceneComponent>(SceneRootNodePtr->GetComponentTemplate()), SceneRootNodePtr->IsDefaultSceneRoot(), true);
-}
-
-void SSCSEditor::PasteNodes()
-{
-	const FScopedTransaction Transaction(LOCTEXT("PasteComponents", "Paste Component(s)"));
-
-	if (EditorMode == EComponentEditorMode::BlueprintSCS)
-	{
-		// Get the components to paste from the clipboard
-		TMap<FName, FName> ParentMap;
-		TMap<FName, UActorComponent*> NewObjectMap;
-		FComponentEditorUtils::GetComponentsFromClipboard(ParentMap, NewObjectMap, true);
-		
-		// Clear the current selection
-		SCSTreeWidget->ClearSelection();
-
-		// Get the blueprint that's being edited
-		UBlueprint* Blueprint = GetBlueprint();
-		check(Blueprint != nullptr && Blueprint->SimpleConstructionScript != nullptr);
-
-		Blueprint->Modify();
-		SaveSCSCurrentState(Blueprint->SimpleConstructionScript);
-
-		// stop allowing tree updates
-		bool bRestoreAllowTreeUpdates = bAllowTreeUpdates;
-		bAllowTreeUpdates = false;
-
-		// Create a new tree node for each new (pasted) component
-		FSCSEditorTreeNodePtrType FirstNode;
-		TMap<FName, FSCSEditorTreeNodePtrType> NewNodeMap;
-		for (const TPair<FName, UActorComponent*>& NewObjectPair : NewObjectMap)
-		{
-			// Get the component object instance
-			UActorComponent* NewActorComponent = NewObjectPair.Value;
-			check(NewActorComponent);
-
-			// Create a new SCS node to contain the new component and add it to the tree
-			NewActorComponent = AddNewNode(TUniquePtr<FScopedTransaction>(), Blueprint->SimpleConstructionScript->CreateNodeAndRenameComponent(NewActorComponent), nullptr, false, false);
-
-			if (NewActorComponent)
-			{
-				// Locate the node that corresponds to the new component template or instance
-				FSCSEditorTreeNodePtrType NewNodePtr = FindTreeNode(NewActorComponent);
-				if (NewNodePtr.IsValid())
-				{
-					// Add the new node to the node map
-					NewNodeMap.Add(NewObjectPair.Key, NewNodePtr);
-
-					// Update the selection to include the new node
-					SCSTreeWidget->SetItemSelection(NewNodePtr, true);
-
-					if (!FirstNode.IsValid())
-					{
-						FirstNode = NewNodePtr;
-					}
-				}
-			}
-		}
-
-		// Restore the node hierarchy from the original copy
-		for (const TPair<FName, FSCSEditorTreeNodePtrType>& NewNodePair : NewNodeMap)
-		{
-			// If an entry exists in the set of known parent nodes for the current node
-			if (ParentMap.Contains(NewNodePair.Key))
-			{
-				// Get the parent node name
-				FName ParentName = ParentMap[NewNodePair.Key];
-				if (NewNodeMap.Contains(ParentName))
-				{
-					// Reattach the current node to the parent node (this will also handle detachment from the scene root node)
-					NewNodeMap[ParentName]->AddChild(NewNodePair.Value);
-
-					// Ensure that the new node is expanded to show the child node(s)
-					SCSTreeWidget->SetItemExpansion(NewNodeMap[ParentName], true);
-				}
-			}
-		}
-
-		// allow tree updates again
-		bAllowTreeUpdates = bRestoreAllowTreeUpdates;
-
-		// scroll the first node into view
-		if (FirstNode.IsValid())
-		{
-			SCSTreeWidget->RequestScrollIntoView(FirstNode);
-		}
-
-		// Modify the Blueprint generated class structure (this will also call UpdateTree() as a result)
-		FBlueprintEditorUtils::MarkBlueprintAsStructurallyModified(Blueprint);
-	}
-	else    // EComponentEditorMode::ActorInstance
-	{
-		// Determine where in the hierarchy to paste (default to the root)
-		USceneComponent* TargetComponent = GetActorContext()->GetRootComponent();
-		for (FSCSEditorTreeNodePtrType SelectedNodePtr : GetSelectedNodes())
-		{
-			check(SelectedNodePtr.IsValid());
-
-			if (USceneComponent* SceneComponent = Cast<USceneComponent>(SelectedNodePtr->GetComponentTemplate()))
-			{
-				TargetComponent = SceneComponent;
-				break;
-			}
-		}
-
-		// Paste the components
-		TArray<UActorComponent*> PastedComponents;
-		FComponentEditorUtils::PasteComponents(PastedComponents, GetActorContext(), TargetComponent);
-
-		if (PastedComponents.Num() > 0)
-		{
-			// We only want the pasted node(s) to be selected
-			SCSTreeWidget->ClearSelection();
-			UpdateTree();
-
-			// Select the nodes that correspond to the pasted components
-			for (UActorComponent* PastedComponent : PastedComponents)
-			{
-				FSCSEditorTreeNodePtrType PastedNode = GetNodeFromActorComponent(PastedComponent);
-				if (PastedNode.IsValid())
-				{
-					SCSTreeWidget->SetItemSelection(PastedNode, true);
-				}
-			}
-		}
-	}
-}
-
-bool SSCSEditor::CanDeleteNodes() const
-{
-	if(!IsEditingAllowed())
-	{
-		return false;
-	}
-
-	TArray<FSCSEditorTreeNodePtrType> SelectedNodes = SCSTreeWidget->GetSelectedItems();
-	for (int32 i = 0; i < SelectedNodes.Num(); ++i)
-	{
-		if (!SelectedNodes[i]->CanDelete())
-		{
-			return false;
-		}
-		else
-		{
-			// Don't allow nodes that belong to a child actor template to be deleted
-			const bool bIsChildActorSubtreeNode = FChildActorComponentEditorUtils::IsChildActorSubtreeNode(SelectedNodes[i]);
-			if (bIsChildActorSubtreeNode)
-			{
-				return false;
-			}
-		}
-	}
-	return SelectedNodes.Num() > 0;
-}
-
-void SSCSEditor::OnDeleteNodes()
-{
-	const FScopedTransaction Transaction( LOCTEXT("RemoveComponents", "Remove Components") );
-
-	// Invalidate any active component in the visualizer
-	GUnrealEd->ComponentVisManager.ClearActiveComponentVis();
-
-	if (EditorMode == EComponentEditorMode::BlueprintSCS)
-	{
-		UBlueprint* Blueprint = GetBlueprint();
-		check(Blueprint != nullptr);
-
-		// Get the current render info for the blueprint. If this is NULL then the blueprint is not currently visualizable (no visible primitive components)
-		FThumbnailRenderingInfo* RenderInfo = GUnrealEd->GetThumbnailManager()->GetRenderingInfo( Blueprint );
-
-		// A lamda for displaying a confirm message to the user if there is a dynamic delegate bound to the 
-		// component they are trying to delete
-		auto ConfirmDeleteLambda = [](USCS_Node* ScsNode) -> FSuppressableWarningDialog::EResult
-		{
-			if (ensure(ScsNode))
-			{
-				FText VarNam = FText::FromName(ScsNode->GetVariableName());
-				FText ConfirmDelete = FText::Format(LOCTEXT("ConfirmDeleteDynamicDelegate", "Component \"{0}\" has bound events in use! If you delete it then those nodes will become invalid. Are you sure you want to delete it?"), VarNam);
-
-				// Warn the user that this may result in data loss
-				FSuppressableWarningDialog::FSetupInfo Info(ConfirmDelete, LOCTEXT("DeleteComponent", "Delete Component"), "DeleteComponentInUse_Warning");
-				Info.ConfirmText = LOCTEXT("ConfirmDeleteDynamicDelegate_Yes", "Yes");
-				Info.CancelText = LOCTEXT("ConfirmDeleteDynamicDelegate_No", "No");
-
-				FSuppressableWarningDialog DeleteVariableInUse(Info);
-
-				// If the user selects cancel then return false
-				return DeleteVariableInUse.ShowModal();
-			}
-
-			return FSuppressableWarningDialog::Cancel;
-		};
-
-		// Remove node(s) from SCS
-		TArray<FSCSEditorTreeNodePtrType> SelectedNodes = SCSTreeWidget->GetSelectedItems();
-
-		// Confirm that the user wants to delete this node
-		for (int32 i = 0; i < SelectedNodes.Num(); ++i)
-		{
-			FSCSEditorTreeNodePtrType Node = SelectedNodes[i];
-			USCS_Node* SCS_Node = Node->GetSCSNode();
-
-			if (SCS_Node != nullptr)
-			{
-				// If this node is in use by Dynamic delegates, then confirm before continuing
-				if (FKismetEditorUtilities::PropertyHasBoundEvents(Blueprint, SCS_Node->GetVariableName()))
-				{
-					// The user has decided not to delete the component, stop trying to delete this component
-					if (ConfirmDeleteLambda(SCS_Node) == FSuppressableWarningDialog::Cancel)
-					{
-						return;
-					}
-				}
-			}
-		}
-
-		//const FScopedTransaction Transaction(LOCTEXT("SetNodeEnabledState", "Set Node Enabled State"));
-
-		for (int32 i = 0; i < SelectedNodes.Num(); ++i)
-		{
-			FSCSEditorTreeNodePtrType Node = SelectedNodes[i];
-			USCS_Node* SCS_Node = Node->GetSCSNode();
-
-			if(SCS_Node != nullptr)
-			{
-				USimpleConstructionScript* SCS = SCS_Node->GetSCS();
-				check(SCS != nullptr && Blueprint == SCS->GetBlueprint());
-
-				// Saving objects for restoring purpose.
-				Blueprint->Modify();
-				SaveSCSCurrentState( SCS );
-			}
-
-			RemoveComponentNode(Node);
-		}
-
-		// Will call UpdateTree as part of OnBlueprintChanged handling
-		FBlueprintEditorUtils::MarkBlueprintAsStructurallyModified(Blueprint);
-
-		// If we had a thumbnail before we deleted any components, check to see if we should clear it
-		// If we deleted the final visualizable primitive from the blueprint, GetRenderingInfo should return NULL
-		FThumbnailRenderingInfo* NewRenderInfo = GUnrealEd->GetThumbnailManager()->GetRenderingInfo( Blueprint );
-		if ( RenderInfo && !NewRenderInfo )
-		{
-			// We removed the last visible primitive component, clear the thumbnail
-			const FString BPFullName = FString::Printf(TEXT("%s %s"), *Blueprint->GetClass()->GetName(), *Blueprint->GetPathName());
-			UPackage* BPPackage = Blueprint->GetOutermost();
-			ThumbnailTools::CacheEmptyThumbnail( BPFullName, BPPackage );
-		}
-
-		// Do this AFTER marking the Blueprint as modified
-		UpdateSelectionFromNodes(SCSTreeWidget->GetSelectedItems());
-	}
-	else    // EComponentEditorMode::ActorInstance
-	{
-	//	const FScopedTransaction Transaction(LOCTEXT("SetNodeEnabledState", "Set Node Enabled State"));
-
-		if (AActor* ActorInstance = GetActorContext())
-		{
-			ActorInstance->Modify();
-		}
-
-		TArray<UActorComponent*> ComponentsToDelete;
-		TArray<FSCSEditorTreeNodePtrType> SelectedNodes = GetSelectedNodes();
-		for (int32 i = 0; i < SelectedNodes.Num(); ++i)
-		{
-			// Get the current selected node reference
-			FSCSEditorTreeNodePtrType SelectedNodePtr = SelectedNodes[i];
-			check(SelectedNodePtr.IsValid());
-
-			// Get the component template associated with the selected node
-			UActorComponent* ComponentTemplate = SelectedNodePtr->GetComponentTemplate();
-			if (ComponentTemplate)
-			{
-				ComponentsToDelete.Add(ComponentTemplate);
-			}
-		}
-
-		UActorComponent* ComponentToSelect = nullptr;
-		int32 NumDeletedComponents = FComponentEditorUtils::DeleteComponents(ComponentsToDelete, ComponentToSelect);
-		if (NumDeletedComponents > 0)
-		{
-			if (ComponentToSelect)
-			{
-				FSCSEditorTreeNodePtrType NodeToSelect = GetNodeFromActorComponent(ComponentToSelect);
-				if (NodeToSelect.IsValid())
-				{
-					SCSTreeWidget->SetSelection(NodeToSelect);
-				}
-			}
-
-			// Rebuild the tree view to reflect the new component hierarchy
-			UpdateTree();
-		}
-
-		// Do this AFTER marking the Blueprint as modified
-		UpdateSelectionFromNodes(SCSTreeWidget->GetSelectedItems());
-	}
-}
-
-void SSCSEditor::RemoveComponentNode(FSCSEditorTreeNodePtrType InNodePtr)
-{
-	check(InNodePtr.IsValid());
-
-	if (EditorMode == EComponentEditorMode::BlueprintSCS)
-	{
-		USCS_Node* SCS_Node = InNodePtr->GetSCSNode();
-		if(SCS_Node != nullptr)
-		{
-			// Clear selection if current
-			if (SCSTreeWidget->GetSelectedItems().Contains(InNodePtr))
-			{
-				SCSTreeWidget->ClearSelection();
-			}
-
-			USimpleConstructionScript* SCS = SCS_Node->GetSCS();
-			check(SCS != nullptr);
-
-			// Remove any instances of variable accessors from the blueprint graphs
-			UBlueprint* Blueprint = SCS->GetBlueprint();
-			if(Blueprint != nullptr)
-			{
-				FBlueprintEditorUtils::RemoveVariableNodes(Blueprint, InNodePtr->GetVariableName());
-				
-				// If there are any Bound Component events for this property then give them compiler errors
-				TArray<UK2Node_ComponentBoundEvent*> EventNodes;
-				FKismetEditorUtilities::FindAllBoundEventsForComponent(Blueprint, SCS_Node->GetVariableName(), EventNodes);
-				if(EventNodes.Num() > 0)
-				{
-					// Find any dynamic delegate nodes and give a compiler error for each that is problematic
-					FCompilerResultsLog LogResults;
-					FMessageLog MessageLog("BlueprintLog");
-					
-					// Add a compiler error for each bound event node
-					for (UK2Node_ComponentBoundEvent* Node : EventNodes)
-					{
-						LogResults.Error(*LOCTEXT("RemoveBoundEvent_Error", "The component that @@ was bound to has been deleted! This node is no longer valid").ToString(), Node);
-					}
-
-					// Notify the user that these nodes are no longer valid
-					MessageLog.NewPage(LOCTEXT("RemoveBoundEvent_Error_Label", "Removed Owner of Component Bound Event"));
-					MessageLog.AddMessages(LogResults.Messages);
-					MessageLog.Notify(LOCTEXT("RemoveBoundEvent_Error_Msg", "Removed Owner of Component Bound Event"));
-						
-					// Focus on the first node that we found
-					FKismetEditorUtilities::BringKismetToFocusAttentionOnObject(EventNodes[0]);									
-				}
-			}
-
-			// Remove node from SCS tree
-			SCS->RemoveNodeAndPromoteChildren(SCS_Node);
-
-			// Clear the delegate
-			SCS_Node->SetOnNameChanged(FSCSNodeNameChanged());
-
-			// on removal, since we don't move the template from the GeneratedClass (which we shouldn't, as it would create a 
-			// discrepancy with existing instances), we rename it instead so that we can re-use the name without having to compile  
-			// (we still have a problem if they attempt to name it to what ever we choose here, but that is unlikely)
-			// note: skip this for the default scene root; we don't actually destroy that node when it's removed, so we don't need the template to be renamed.
-			if (!InNodePtr->IsDefaultSceneRoot() && SCS_Node->ComponentTemplate != nullptr)
-			{
-				const FName TemplateName = SCS_Node->ComponentTemplate->GetFName();
-				const FString RemovedName = SCS_Node->GetVariableName().ToString() + TEXT("_REMOVED_") + FGuid::NewGuid().ToString();
-
-				SCS_Node->ComponentTemplate->Modify();
-				SCS_Node->ComponentTemplate->Rename(*RemovedName, /*NewOuter =*/nullptr, REN_DontCreateRedirectors);
-
-				TArray<UObject*> ArchetypeInstances;
-				auto DestroyArchetypeInstances = [&ArchetypeInstances, &RemovedName](UActorComponent* ComponentTemplate)
-				{
-					ComponentTemplate->GetArchetypeInstances(ArchetypeInstances);
-					for (UObject* ArchetypeInstance : ArchetypeInstances)
-					{
-						if (!ArchetypeInstance->HasAllFlags(RF_ArchetypeObject | RF_InheritableComponentTemplate))
-						{
-							CastChecked<UActorComponent>(ArchetypeInstance)->DestroyComponent();
-							ArchetypeInstance->Rename(*RemovedName, nullptr, REN_DontCreateRedirectors);
-						}
-					}
-				};
-
-				DestroyArchetypeInstances(SCS_Node->ComponentTemplate);
-				
-				if (Blueprint)
-				{
-					// Children need to have their inherited component template instance of the component renamed out of the way as well
-					TArray<UClass*> ChildrenOfClass;
-					GetDerivedClasses(Blueprint->GeneratedClass, ChildrenOfClass);
-
-					for (UClass* ChildClass : ChildrenOfClass)
-					{
-						UBlueprintGeneratedClass* BPChildClass = CastChecked<UBlueprintGeneratedClass>(ChildClass);
-
-						if (UActorComponent* Component = (UActorComponent*)FindObjectWithOuter(BPChildClass, UActorComponent::StaticClass(), TemplateName))
-						{
-							Component->Modify();
-							Component->Rename(*RemovedName, /*NewOuter =*/nullptr, REN_DontCreateRedirectors);
-
-							DestroyArchetypeInstances(Component);
-						}
-					}
-				}
-			}
-		}
-	}
-	else    // EComponentEditorMode::ActorInstance
-	{
-		AActor* ActorInstance = GetActorContext();
-
-		UActorComponent* ComponentInstance = InNodePtr->GetComponentTemplate();
-		if ((ActorInstance != nullptr) && (ComponentInstance != nullptr))
-		{
-			// Clear selection if current
-			if (SCSTreeWidget->GetSelectedItems().Contains(InNodePtr))
-			{
-				SCSTreeWidget->ClearSelection();
-			}
-
-			const bool bWasDefaultSceneRoot = InNodePtr.IsValid() && InNodePtr->IsDefaultSceneRoot();
-
-			// Destroy the component instance
-			ComponentInstance->Modify();
-			ComponentInstance->DestroyComponent(!bWasDefaultSceneRoot);
-		}
-	}
-}
-
-void SSCSEditor::UpdateSelectionFromNodes(const TArray<FSCSEditorTreeNodePtrType> &SelectedNodes)
-{
-	bUpdatingSelection = true;
-
-	// Notify that the selection has updated
-	OnSelectionUpdated.ExecuteIfBound(SelectedNodes);
-
-	bUpdatingSelection = false;
-}
-
-void SSCSEditor::RefreshSelectionDetails()
-{
-	UpdateSelectionFromNodes(SCSTreeWidget->GetSelectedItems());
-}
-
-void SSCSEditor::OnTreeSelectionChanged(FSCSEditorTreeNodePtrType, ESelectInfo::Type /*SelectInfo*/)
-{
-	UpdateSelectionFromNodes(SCSTreeWidget->GetSelectedItems());
-}
-
-bool SSCSEditor::IsNodeInSimpleConstructionScript( USCS_Node* Node ) const
-{
-	check(Node);
-
-	USimpleConstructionScript* NodeSCS = Node->GetSCS();
-	if(NodeSCS != NULL)
-	{
-		return NodeSCS->GetAllNodes().Contains(Node);
-	}
-	
-	return false;
-}
-
-FSCSEditorTreeNodePtrType SSCSEditor::AddTreeNode(USCS_Node* InSCSNode, FSCSEditorTreeNodePtrType InParentNodePtr, const bool bIsInheritedSCS)
-{
-	FSCSEditorTreeNodePtrType NewNodePtr;
-
-	check(InSCSNode != nullptr && InParentNodePtr.IsValid());
-
-	// During diffs, ComponentTemplates can easily be null, so prevent these checks.
-	if (!bIsDiffing && InSCSNode->ComponentTemplate)
-	{
-		checkf(InSCSNode->ParentComponentOrVariableName == NAME_None
-			|| (!InSCSNode->bIsParentComponentNative && InParentNodePtr->GetSCSNode() != nullptr && InParentNodePtr->GetSCSNode()->GetVariableName() == InSCSNode->ParentComponentOrVariableName)
-			|| (InSCSNode->bIsParentComponentNative && InParentNodePtr->GetComponentTemplate() != nullptr && InParentNodePtr->GetComponentTemplate()->GetFName() == InSCSNode->ParentComponentOrVariableName),
-			TEXT("Failed to add SCS node %s to tree:\n- bIsParentComponentNative=%d\n- Stored ParentComponentOrVariableName=%s\n- Actual ParentComponentOrVariableName=%s"),
-			*InSCSNode->GetVariableName().ToString(),
-			!!InSCSNode->bIsParentComponentNative,
-			*InSCSNode->ParentComponentOrVariableName.ToString(),
-			!InSCSNode->bIsParentComponentNative
-			? (InParentNodePtr->GetSCSNode() != nullptr ? *InParentNodePtr->GetSCSNode()->GetVariableName().ToString() : TEXT("NULL"))
-			: (InParentNodePtr->GetComponentTemplate() != nullptr ? *InParentNodePtr->GetComponentTemplate()->GetFName().ToString() : TEXT("NULL")));
-	}
-	
-	// Determine whether or not the given node is inherited from a parent Blueprint
-	USimpleConstructionScript* NodeSCS = InSCSNode->GetSCS();
-
-	// do this first, because we need a FSCSEditorTreeNodePtrType for the new node
-	NewNodePtr = InParentNodePtr->AddChild(InSCSNode, bIsInheritedSCS);
-	RefreshFilteredState(NewNodePtr, /*bRecursive =*/false);
-	
-	if( InSCSNode->ComponentTemplate && 
-		InSCSNode->ComponentTemplate->IsA(USceneComponent::StaticClass()) && 
-		InParentNodePtr->GetNodeType() == FSCSEditorTreeNode::ComponentNode)
-	{
-		bool bParentIsEditorOnly = InParentNodePtr->GetComponentTemplate()->IsEditorOnly();
-		// if you can't nest this new node under the proposed parent (then swap the two)
-		if (bParentIsEditorOnly && !InSCSNode->ComponentTemplate->IsEditorOnly() && InParentNodePtr->CanReparent())
-		{
-			FSCSEditorTreeNodePtrType OldParentPtr = InParentNodePtr;
-			InParentNodePtr = OldParentPtr->GetParent();
-
-			OldParentPtr->RemoveChild(NewNodePtr);
-			NodeSCS->RemoveNode(OldParentPtr->GetSCSNode());
-
-			// if the grandparent node is invalid (assuming this means that the parent node was the scene-root)
-			if (!InParentNodePtr.IsValid())
-			{
-				check(OldParentPtr == GetSceneRootNode());
-				SetSceneRootNode(NewNodePtr);
-				NodeSCS->AddNode(NewNodePtr->GetSCSNode());
-			}
-			else 
-			{
-				InParentNodePtr->AddChild(NewNodePtr);
-			}
-
-			// move the proposed parent in as a child to the new node
-			NewNodePtr->AddChild(OldParentPtr);
-		} // if bParentIsEditorOnly...
-	}
-	else 
-	{
-		// If the SCS root node array does not already contain the given node, this will add it (this should only occur after node creation)
-		if(NodeSCS != nullptr)
-		{
-			NodeSCS->AddNode(InSCSNode);
-		}
-	}
-
-	// Expand parent nodes by default
-	SCSTreeWidget->SetItemExpansion(InParentNodePtr, true);
-
-	// Add this node's child actor node, if present
-	AddTreeNodeFromChildActor(NewNodePtr);
-
-	// Recursively add the given SCS node's child nodes
-	for (USCS_Node* ChildNode : InSCSNode->GetChildNodes())
-	{
-		AddTreeNode(ChildNode, NewNodePtr, bIsInheritedSCS);
-	}
-
-	return NewNodePtr;
-}
-
-FSCSEditorTreeNodePtrType SSCSEditor::AddTreeNodeFromComponent(UActorComponent* InActorComponent, FSCSEditorTreeNodePtrType InParentTreeNode)
-{
-	check(InActorComponent != NULL);
-	ensure(IsValid(InActorComponent));
-
-	FSCSEditorTreeNodePtrType NewNodePtr = InParentTreeNode->FindChild(InActorComponent);
-	if (!NewNodePtr.IsValid())
-	{
-		NewNodePtr = FSCSEditorTreeNode::FactoryNodeFromComponent(InActorComponent);
-		InParentTreeNode->AddChild(NewNodePtr);
-		RefreshFilteredState(NewNodePtr, false);
-	}
-
-	AddTreeNodeFromChildActor(NewNodePtr);
-
-	SCSTreeWidget->SetItemExpansion(NewNodePtr, true);
-
-	return NewNodePtr;
-}
-
-FSCSEditorTreeNodePtrType SSCSEditor::AddTreeNodeFromChildActor(FSCSEditorTreeNodePtrType InNodePtr)
-{
-	if (!InNodePtr.IsValid())
-	{
-		return nullptr;
-	}
-
-	const EChildActorComponentTreeViewVisualizationMode DefaultVisOverride = UICustomization.IsValid() ? UICustomization->GetChildActorVisualizationMode() : EChildActorComponentTreeViewVisualizationMode::UseDefault;
-
-	// Skip any expansion logic if the option is disabled
-	if (DefaultVisOverride == EChildActorComponentTreeViewVisualizationMode::UseDefault && !FChildActorComponentEditorUtils::IsChildActorTreeViewExpansionEnabled())
-	{
-		return nullptr;
-	}
-
-	FSCSEditorChildActorNodePtrType ChildActorNodePtr = InNodePtr->GetChildActorNode();
-	if (ChildActorNodePtr.IsValid())
-	{
-		// Get the child actor component that's associated with the child actor node
-		UChildActorComponent* ChildActorComponent = ChildActorNodePtr->GetChildActorComponent();
-		check(ChildActorComponent != nullptr);
-
-		// Check to see if we should expand the child actor node within the tree view
-		const bool bExpandChildActorInTreeView = FChildActorComponentEditorUtils::ShouldExpandChildActorInTreeView(ChildActorComponent, DefaultVisOverride);
-		if (bExpandChildActorInTreeView)
-		{
-			// Do the expansion as a normal actor subtree
-			BuildSubTreeForActorNode(ChildActorNodePtr);
-
-			// Check to see if we should include the child actor node within the tree view
-			const bool bShowChildActorNodeInTreeView = FChildActorComponentEditorUtils::ShouldShowChildActorNodeInTreeView(ChildActorComponent, DefaultVisOverride);
-			if (bShowChildActorNodeInTreeView)
-			{
-				// Add the child actor node into the tree view
-				InNodePtr->AddChild(ChildActorNodePtr);
-				RefreshFilteredState(ChildActorNodePtr, false);
-			}
-			else
-			{
-				// Add the child actor's subtree into the tree view
-				TArray<FSCSEditorTreeNodePtrType> Children = ChildActorNodePtr->GetChildren();
-				for (const FSCSEditorTreeNodePtrType& ChildNode : Children)
-				{
-					// Remove the child actor node as the visible root in the tree view, and replace it with the given node
-					ChildActorNodePtr->RemoveChild(ChildNode);
-					InNodePtr->AddChild(ChildNode);
-
-					// Restore the subtree's actor root to indicate that it's still a child actor expansion in the tree view
-					ChildNode->SetActorRootNode(ChildActorNodePtr);
-				}
-			}
-
-			// Expand the given parent to reveal the child actor node and/or its subtree
-			SCSTreeWidget->SetItemExpansion(InNodePtr, true);
-		}
-	}
-
-	return ChildActorNodePtr;
-}
-
-FSCSEditorTreeNodePtrType SSCSEditor::FindTreeNode(const USCS_Node* InSCSNode, FSCSEditorTreeNodePtrType InStartNodePtr) const
-{
-	FSCSEditorTreeNodePtrType NodePtr;
-	if(InSCSNode != NULL)
-	{
-		// Start at the scene root node if none was given
-		if(!InStartNodePtr.IsValid())
-		{
-			InStartNodePtr = GetSceneRootNode();
-		}
-
-		if(InStartNodePtr.IsValid())
-		{
-			// Check to see if the given SCS node matches the given tree node
-			if(InStartNodePtr->GetSCSNode() == InSCSNode)
-			{
-				NodePtr = InStartNodePtr;
-			}
-			else
-			{
-				// Recursively search for the node in our child set
-				NodePtr = InStartNodePtr->FindChild(InSCSNode);
-				if(!NodePtr.IsValid())
-				{
-					for(int32 i = 0; i < InStartNodePtr->GetChildren().Num() && !NodePtr.IsValid(); ++i)
-					{
-						NodePtr = FindTreeNode(InSCSNode, InStartNodePtr->GetChildren()[i]);
-					}
-				}
-			}
-		}
-	}
-
-	return NodePtr;
-}
-
-FSCSEditorTreeNodePtrType SSCSEditor::FindTreeNode(const UActorComponent* InComponent, FSCSEditorTreeNodePtrType InStartNodePtr) const
-{
-	FSCSEditorTreeNodePtrType NodePtr;
-	if(InComponent != NULL)
-	{
-		// Start at the scene root node if none was given
-		if(!InStartNodePtr.IsValid())
-		{
-			InStartNodePtr = GetActorNode();
-		}
-
-		if(InStartNodePtr.IsValid())
-		{
-			// Check to see if the given component template matches the given tree node
-			// 
-			// For certain node types, GetOrCreateEditableComponentTemplate() will handle retrieving 
-			// the "OverridenComponentTemplate" which may be what we're looking for in some 
-			// cases; if not, then we fall back to just checking GetComponentTemplate()
-			if (InStartNodePtr->GetOrCreateEditableComponentTemplate(GetBlueprint()) == InComponent)
-			{
-				NodePtr = InStartNodePtr;
-			}
-			else if (InStartNodePtr->GetComponentTemplate() == InComponent)
-			{
-				NodePtr = InStartNodePtr;
-			}
-			else
-			{
-				// Recursively search for the node in our child set
-				NodePtr = InStartNodePtr->FindChild(InComponent);
-				if(!NodePtr.IsValid())
-				{
-					for(int32 i = 0; i < InStartNodePtr->GetChildren().Num() && !NodePtr.IsValid(); ++i)
-					{
-						NodePtr = FindTreeNode(InComponent, InStartNodePtr->GetChildren()[i]);
-					}
-				}
-			}
-		}
-	}
-
-	return NodePtr;
-}
-
-FSCSEditorTreeNodePtrType SSCSEditor::FindTreeNode(const FName& InVariableOrInstanceName, FSCSEditorTreeNodePtrType InStartNodePtr) const
-{
-	FSCSEditorTreeNodePtrType NodePtr;
-	if(InVariableOrInstanceName != NAME_None)
-	{
-		// Start at the root node if none was given
-		if(!InStartNodePtr.IsValid())
-		{
-			InStartNodePtr = GetActorNode();
-		}
-
-		if(InStartNodePtr.IsValid())
-		{
-			FName ItemName = InStartNodePtr->GetNodeID();
-
-			// Check to see if the given name matches the item name
-			if(InVariableOrInstanceName == ItemName)
-			{
-				NodePtr = InStartNodePtr;
-			}
-			else
-			{
-				// Recursively search for the node in our child set
-				NodePtr = InStartNodePtr->FindChild(InVariableOrInstanceName);
-				if(!NodePtr.IsValid())
-				{
-					for(int32 i = 0; i < InStartNodePtr->GetChildren().Num() && !NodePtr.IsValid(); ++i)
-					{
-						NodePtr = FindTreeNode(InVariableOrInstanceName, InStartNodePtr->GetChildren()[i]);
-					}
-				}
-			}
-		}
-	}
-
-	return NodePtr;
-}
-
-void SSCSEditor::OnItemScrolledIntoView( FSCSEditorTreeNodePtrType InItem, const TSharedPtr<ITableRow>& InWidget)
-{
-	if(DeferredRenameRequest != NAME_None)
-	{
-		FName ItemName = InItem->GetNodeID();
-		if(DeferredRenameRequest == ItemName)
-		{
-			DeferredRenameRequest = NAME_None;
-			InItem->OnRequestRename(MoveTemp(DeferredOngoingCreateTransaction)); // Transfer responsibility to end the 'create + give initial name' transaction to the tree item if such transaction is ongoing.
-		}
-	}
-}
-
-void SSCSEditor::HandleItemDoubleClicked(FSCSEditorTreeNodePtrType InItem)
-{
-	// Notify that the selection has updated
-	OnItemDoubleClicked.ExecuteIfBound(InItem);
-}
-
-void SSCSEditor::OnRenameComponent()
-{
-	OnRenameComponent(nullptr); // null means that the rename is not part of the creation process (create + give initial name).
-}
-
-void SSCSEditor::OnRenameComponent(TUniquePtr<FScopedTransaction> InComponentCreateTransaction)
-{
-	TArray< FSCSEditorTreeNodePtrType > SelectedItems = SCSTreeWidget->GetSelectedItems();
-
-	// Should already be prevented from making it here.
-	check(SelectedItems.Num() == 1);
-
-	DeferredRenameRequest = SelectedItems[0]->GetNodeID();
-
-	check(!DeferredOngoingCreateTransaction.IsValid()); // If this fails, something in the chain of responsibility failed to end the previous transaction.
-	DeferredOngoingCreateTransaction = MoveTemp(InComponentCreateTransaction); // If a 'create + give initial name' transaction is ongoing, take responsibility of ending it until the selected item is scrolled into view.
-
-	SCSTreeWidget->RequestScrollIntoView(SelectedItems[0]);
-
-	if (DeferredOngoingCreateTransaction.IsValid() && !PostTickHandle.IsValid())
-	{
-		// Ensure the item will be scrolled into view during the frame (See explanation in OnPostTick()).
-		PostTickHandle = FSlateApplication::Get().OnPostTick().AddSP(this, &SSCSEditor::OnPostTick);
-	}
-}
-
-void SSCSEditor::OnPostTick(float)
-{
-	// If a 'create + give initial name' is ongoing and the transaction ownership was not transferred during the frame it was requested, it is most likely because the newly
-	// created item could not be scrolled into view (should say 'teleported', the scrolling is not animated). The tree view will not put the item in view if the there is
-	// no space left to display the item. (ex a splitter where all the display space is used by the other component). End the transaction before starting a new frame. The user
-	// will not be able to rename on creation, the widget is likely not in view and cannot be edited anyway.
-	DeferredOngoingCreateTransaction.Reset();
-
-	// The post tick event handler is not required anymore.
-	FSlateApplication::Get().OnPostTick().Remove(PostTickHandle);
-	PostTickHandle.Reset();
-}
-
-bool SSCSEditor::CanRenameComponent() const
-{
-	if (!IsEditingAllowed())
-	{
-		return false;
-	}
-
-	// In addition to certain node types, don't allow nodes within a child actor template's hierarchy to be renamed
-	TArray<FSCSEditorTreeNodePtrType> SelectedItems = SCSTreeWidget->GetSelectedItems();
-	return SelectedItems.Num() == 1 && SelectedItems[0]->CanRename() && !FChildActorComponentEditorUtils::IsChildActorSubtreeNode(SelectedItems[0]);
-}
-
-void SSCSEditor::DepthFirstTraversal(const FSCSEditorTreeNodePtrType& InNodePtr, TSet<FSCSEditorTreeNodePtrType>& OutVisitedNodes, const TFunctionRef<void(const FSCSEditorTreeNodePtrType&)> InFunction) const
-{
-	if (InNodePtr.IsValid()
-		&& ensureMsgf(!OutVisitedNodes.Contains(InNodePtr), TEXT("Already visited node: %s (Parent: %s)"), *InNodePtr->GetDisplayString(), InNodePtr->GetParent().IsValid() ? *InNodePtr->GetParent()->GetDisplayString() : TEXT("NULL")))
-	{
-		InFunction(InNodePtr);
-		OutVisitedNodes.Add(InNodePtr);
-
-		for (const FSCSEditorTreeNodePtrType& Child : InNodePtr->GetChildren())
-		{
-			DepthFirstTraversal(Child, OutVisitedNodes, InFunction);
-		}
-	}
-}
-
-void SSCSEditor::GetCollapsedNodes(const FSCSEditorTreeNodePtrType& InNodePtr, TSet<FSCSEditorTreeNodePtrType>& OutCollapsedNodes) const
-{
-	TSet<FSCSEditorTreeNodePtrType> VisitedNodes;
-	DepthFirstTraversal(InNodePtr, VisitedNodes, [SCSTreeWidget = this->SCSTreeWidget, &OutCollapsedNodes](const FSCSEditorTreeNodePtrType& InNodePtr)
-	{
-		if (InNodePtr->GetChildren().Num() > 0 && !SCSTreeWidget->IsItemExpanded(InNodePtr))
-		{
-			OutCollapsedNodes.Add(InNodePtr);
-		}
-	});
-}
-
-EVisibility SSCSEditor::GetPromoteToBlueprintButtonVisibility() const
-{
-	UObject* ActorContextPtr = GetActorContext();
-	return (UICustomization.IsValid() && UICustomization->HideBlueprintButtons(MakeArrayView(&ActorContextPtr, 1)))
-		|| (EditorMode != EComponentEditorMode::ActorInstance) 
-		|| (GetBlueprint() != nullptr)
-		? EVisibility::Collapsed : EVisibility::Visible;
-}
-
-EVisibility SSCSEditor::GetEditBlueprintButtonVisibility() const
-{
-	UObject* ActorContextPtr = GetActorContext();
-	return (UICustomization.IsValid() && UICustomization->HideBlueprintButtons(MakeArrayView(&ActorContextPtr, 1)))
-		|| (EditorMode != EComponentEditorMode::ActorInstance)
-		|| (GetBlueprint() == nullptr)
-		? EVisibility::Collapsed : EVisibility::Visible;
-}
-
-EVisibility SSCSEditor::GetComponentClassComboButtonVisibility() const
-{
-	UObject* ActorContextPtr = GetActorContext();
-	return (HideComponentClassCombo.Get() 
-		|| (UICustomization.IsValid() && UICustomization->HideAddComponentButton(MakeArrayView(&ActorContextPtr, 1))))
-		? EVisibility::Collapsed : EVisibility::Visible;
-}
-
-EVisibility SSCSEditor::GetComponentsTreeVisibility() const
-{
-	UObject* ActorContextPtr = GetActorContext();
-	return (UICustomization.IsValid() && UICustomization->HideComponentsTree(MakeArrayView(&ActorContextPtr, 1)))
-		? EVisibility::Collapsed : EVisibility::Visible;
-}
-
-EVisibility SSCSEditor::GetComponentsFilterBoxVisibility() const
-{
-	UObject* ActorContextPtr = GetActorContext();
-	return (UICustomization.IsValid() && UICustomization->HideComponentsFilterBox(MakeArrayView(&ActorContextPtr, 1)))
-		? EVisibility::Collapsed : EVisibility::Visible;
-}
-
-FText SSCSEditor::OnGetApplyChangesToBlueprintTooltip() const
-{
-	int32 NumChangedProperties = 0;
-
-	AActor* Actor = GetActorContext();
-	UBlueprint* Blueprint = (Actor != nullptr) ? Cast<UBlueprint>(Actor->GetClass()->ClassGeneratedBy) : nullptr;
-
-	if(Actor != NULL && Blueprint != NULL && Actor->GetClass()->ClassGeneratedBy == Blueprint)
-	{
-		AActor* BlueprintCDO = Actor->GetClass()->GetDefaultObject<AActor>();
-		if(BlueprintCDO != NULL)
-		{
-			const EditorUtilities::ECopyOptions::Type CopyOptions = (EditorUtilities::ECopyOptions::Type)(EditorUtilities::ECopyOptions::PreviewOnly|EditorUtilities::ECopyOptions::OnlyCopyEditOrInterpProperties|EditorUtilities::ECopyOptions::SkipInstanceOnlyProperties);
-			NumChangedProperties += EditorUtilities::CopyActorProperties(Actor, BlueprintCDO, CopyOptions);
-		}
-		NumChangedProperties += Actor->GetInstanceComponents().Num();
-	}
-
-
-	if(NumChangedProperties == 0)
-	{
-		return LOCTEXT("DisabledPushToBlueprintDefaults_ToolTip", "Replaces the Blueprint's defaults with any altered property values.");
-	}
-	else if(NumChangedProperties > 1)
-	{
-		return FText::Format(LOCTEXT("PushToBlueprintDefaults_ToolTip", "Click to apply {0} changed properties to the Blueprint."), FText::AsNumber(NumChangedProperties));
-	}
-	else
-	{
-		return LOCTEXT("PushOneToBlueprintDefaults_ToolTip", "Click to apply 1 changed property to the Blueprint.");
-	}
-}
-
-FText SSCSEditor::OnGetResetToBlueprintDefaultsTooltip() const
-{
-	int32 NumChangedProperties = 0;
-
-	AActor* Actor = GetActorContext();
-	UBlueprint* Blueprint = (Actor != nullptr) ? Cast<UBlueprint>(Actor->GetClass()->ClassGeneratedBy) : nullptr;
-	if(Actor != NULL && Blueprint != NULL && Actor->GetClass()->ClassGeneratedBy == Blueprint)
-	{
-		AActor* BlueprintCDO = Actor->GetClass()->GetDefaultObject<AActor>();
-		if(BlueprintCDO != NULL)
-		{
-			const EditorUtilities::ECopyOptions::Type CopyOptions = (EditorUtilities::ECopyOptions::Type)(EditorUtilities::ECopyOptions::PreviewOnly|EditorUtilities::ECopyOptions::OnlyCopyEditOrInterpProperties);
-			NumChangedProperties += EditorUtilities::CopyActorProperties(BlueprintCDO, Actor, CopyOptions);
-		}
-		NumChangedProperties += Actor->GetInstanceComponents().Num();
-	}
-
-	if(NumChangedProperties == 0)
-	{
-		return LOCTEXT("DisabledResetBlueprintDefaults_ToolTip", "Resets altered properties back to their Blueprint default values.");
-	}
-	else if(NumChangedProperties > 1)
-	{
-		return FText::Format(LOCTEXT("ResetToBlueprintDefaults_ToolTip", "Click to reset {0} changed properties to their Blueprint default values."), FText::AsNumber(NumChangedProperties));
-	}
-	else
-	{
-		return LOCTEXT("ResetOneToBlueprintDefaults_ToolTip", "Click to reset 1 changed property to its Blueprint default value.");
-	}
-}
-
-void SSCSEditor::OnOpenBlueprintEditor(bool bForceCodeEditing) const
-{
-	if (AActor* ActorInstance = GetActorContext())
-	{
-		if (UBlueprint* Blueprint = Cast<UBlueprint>(ActorInstance->GetClass()->ClassGeneratedBy))
-		{
-			if (bForceCodeEditing && (Blueprint->UbergraphPages.Num() > 0))
-			{
-				FKismetEditorUtilities::BringKismetToFocusAttentionOnObject(Blueprint->GetLastEditedUberGraph());
-			}
-			else
-			{
-				GEditor->GetEditorSubsystem<UAssetEditorSubsystem>()->OpenEditorForAsset(Blueprint);
-			}
-		}
-	}
-}
-
-/** 
-This struct saves and deselects all selected instanced components (from given actor), then finds them (in recreated actor instance, after compilation) and selects them again.
-*/
-struct FRestoreSelectedInstanceComponent
-{
-	TWeakObjectPtr<UClass> ActorClass;
-	FName ActorName;
-	TWeakObjectPtr<UObject> ActorOuter;
-
-	struct FComponentKey
-	{
-		FName Name;
-		TWeakObjectPtr<UClass> Class;
-
-		FComponentKey(FName InName, UClass* InClass) : Name(InName), Class(InClass) {}
-	};
-	TArray<FComponentKey> ComponentKeys;
-
-	FRestoreSelectedInstanceComponent()
-		: ActorClass(nullptr)
-		, ActorOuter(nullptr)
-	{ }
-
-	void Save(AActor* InActor)
-	{
-		check(InActor);
-		ActorClass = InActor->GetClass();
-		ActorName = InActor->GetFName();
-		ActorOuter = InActor->GetOuter();
-
-		check(GEditor);
-		TArray<UActorComponent*> ComponentsToSaveAndDelesect;
-		for (FSelectionIterator Iter = GEditor->GetSelectedComponentIterator(); Iter; ++Iter)
-		{
-			UActorComponent* Component = CastChecked<UActorComponent>(*Iter, ECastCheckedType::NullAllowed);
-			if (Component && InActor->GetInstanceComponents().Contains(Component))
-			{
-				ComponentsToSaveAndDelesect.Add(Component);
-			}
-		}
-
-		for (UActorComponent* Component : ComponentsToSaveAndDelesect)
-		{
-			USelection* SelectedComponents = GEditor->GetSelectedComponents();
-			if (ensure(SelectedComponents))
-			{
-				ComponentKeys.Add(FComponentKey(Component->GetFName(), Component->GetClass()));
-				SelectedComponents->Deselect(Component);
-			}
-		}
-	}
-
-	void Restore()
-	{
-		AActor* Actor = (ActorClass.IsValid() && ActorOuter.IsValid()) 
-			? Cast<AActor>((UObject*)FindObjectWithOuter(ActorOuter.Get(), ActorClass.Get(), ActorName)) 
-			: nullptr;
-		if (Actor)
-		{
-			for (const FComponentKey& IterKey : ComponentKeys)
-			{
-				UActorComponent* const* ComponentPtr = Algo::FindByPredicate(Actor->GetComponents(), [&](UActorComponent* InComp)
-				{
-					return InComp && (InComp->GetFName() == IterKey.Name) && (InComp->GetClass() == IterKey.Class.Get());
-				});
-				if (ComponentPtr && *ComponentPtr)
-				{
-					check(GEditor);
-					GEditor->SelectComponent(*ComponentPtr, true, false);
-				}
-			}
-		}
-	}
-};
-
-void SSCSEditor::OnApplyChangesToBlueprint() const
-{
-	AActor* Actor = GetActorContext();
-	const UBlueprint* const Blueprint = (Actor != nullptr) ? Cast<UBlueprint>(Actor->GetClass()->ClassGeneratedBy) : nullptr;
-
-	if (Actor != NULL && Blueprint != NULL && Actor->GetClass()->ClassGeneratedBy.Get() == Blueprint)
-	{
-		const FString ActorLabel = Actor->GetActorLabel();
-		int32 NumChangedProperties = FKismetEditorUtilities::ApplyInstanceChangesToBlueprint(Actor);
-
-		// Set up a notification record to indicate success/failure
-		FNotificationInfo NotificationInfo(FText::GetEmpty());
-		NotificationInfo.FadeInDuration = 1.0f;
-		NotificationInfo.FadeOutDuration = 2.0f;
-		NotificationInfo.bUseLargeFont = false;
-		SNotificationItem::ECompletionState CompletionState;
-		if (NumChangedProperties > 0)
-		{
-			if (NumChangedProperties > 1)
-			{
-				FFormatNamedArguments Args;
-				Args.Add(TEXT("BlueprintName"), FText::FromName(Blueprint->GetFName()));
-				Args.Add(TEXT("NumChangedProperties"), NumChangedProperties);
-				Args.Add(TEXT("ActorName"), FText::FromString(ActorLabel));
-				NotificationInfo.Text = FText::Format(LOCTEXT("PushToBlueprintDefaults_ApplySuccess", "Updated Blueprint {BlueprintName} ({NumChangedProperties} property changes applied from actor {ActorName})."), Args);
-			}
-			else
-			{
-				FFormatNamedArguments Args;
-				Args.Add(TEXT("BlueprintName"), FText::FromName(Blueprint->GetFName()));
-				Args.Add(TEXT("ActorName"), FText::FromString(ActorLabel));
-				NotificationInfo.Text = FText::Format(LOCTEXT("PushOneToBlueprintDefaults_ApplySuccess", "Updated Blueprint {BlueprintName} (1 property change applied from actor {ActorName})."), Args);
-			}
-			CompletionState = SNotificationItem::CS_Success;
-		}
-		else
-		{
-			NotificationInfo.Text = LOCTEXT("PushToBlueprintDefaults_ApplyFailed", "No properties were copied");
-			CompletionState = SNotificationItem::CS_Fail;
-		}
-
-		// Add the notification to the queue
-		const TSharedPtr<SNotificationItem> Notification = FSlateNotificationManager::Get().AddNotification(NotificationInfo);
-		Notification->SetCompletionState(CompletionState);
-	}
-}
-
-void SSCSEditor::OnResetToBlueprintDefaults()
-{
-	int32 NumChangedProperties = 0;
-
-	AActor* Actor = GetActorContext();
-	UBlueprint* Blueprint = (Actor != nullptr) ? Cast<UBlueprint>(Actor->GetClass()->ClassGeneratedBy) : nullptr;
-
-	if ((Actor != NULL) && (Blueprint != NULL) && (Actor->GetClass()->ClassGeneratedBy == Blueprint))
-	{
-		const FScopedTransaction Transaction(LOCTEXT("ResetToBlueprintDefaults_Transaction", "Reset to Class Defaults"));
-
-		{
-			AActor* BlueprintCDO = Actor->GetClass()->GetDefaultObject<AActor>();
-			if (BlueprintCDO != NULL)
-			{
-				const EditorUtilities::ECopyOptions::Type CopyOptions = (EditorUtilities::ECopyOptions::Type)(EditorUtilities::ECopyOptions::OnlyCopyEditOrInterpProperties);
-				NumChangedProperties = EditorUtilities::CopyActorProperties(BlueprintCDO, Actor, CopyOptions);
-			}
-			NumChangedProperties += Actor->GetInstanceComponents().Num();
-			Actor->ClearInstanceComponents(true);
-		}
-
-		// Set up a notification record to indicate success/failure
-		FNotificationInfo NotificationInfo(FText::GetEmpty());
-		NotificationInfo.FadeInDuration = 1.0f;
-		NotificationInfo.FadeOutDuration = 2.0f;
-		NotificationInfo.bUseLargeFont = false;
-		SNotificationItem::ECompletionState CompletionState;
-		if (NumChangedProperties > 0)
-		{
-			if (NumChangedProperties > 1)
-			{
-				FFormatNamedArguments Args;
-				Args.Add(TEXT("BlueprintName"), FText::FromName(Blueprint->GetFName()));
-				Args.Add(TEXT("NumChangedProperties"), NumChangedProperties);
-				Args.Add(TEXT("ActorName"), FText::FromString(Actor->GetActorLabel()));
-				NotificationInfo.Text = FText::Format(LOCTEXT("ResetToBlueprintDefaults_ApplySuccess", "Reset {ActorName} ({NumChangedProperties} property changes applied from Blueprint {BlueprintName})."), Args);
-			}
-			else
-			{
-				FFormatNamedArguments Args;
-				Args.Add(TEXT("BlueprintName"), FText::FromName(Blueprint->GetFName()));
-				Args.Add(TEXT("ActorName"), FText::FromString(Actor->GetActorLabel()));
-				NotificationInfo.Text = FText::Format(LOCTEXT("ResetOneToBlueprintDefaults_ApplySuccess", "Reset {ActorName} (1 property change applied from Blueprint {BlueprintName})."), Args);
-			}
-			CompletionState = SNotificationItem::CS_Success;
-		}
-		else
-		{
-			NotificationInfo.Text = LOCTEXT("ResetToBlueprintDefaults_Failed", "No properties were reset");
-			CompletionState = SNotificationItem::CS_Fail;
-		}
-
-		UpdateTree();
-
-		// Add the notification to the queue
-		const TSharedPtr<SNotificationItem> Notification = FSlateNotificationManager::Get().AddNotification(NotificationInfo);
-		Notification->SetCompletionState(CompletionState);
-	}
-}
-
-void SSCSEditor::PromoteToBlueprint() const
-{
-	FCreateBlueprintFromActorDialog::OpenDialog(ECreateBlueprintFromActorMode::Subclass, GetActorContext());
-}
-
-FReply SSCSEditor::OnPromoteToBlueprintClicked()
-{
-	PromoteToBlueprint();
-	return FReply::Handled();
-}
-
-/** Returns the Actor context for which we are viewing/editing the SCS.  Can return null.  Should not be cached as it may change from frame to frame. */
-AActor* SSCSEditor::GetActorContext() const
-{
-	return ActorContext.Get(nullptr);
-}
-
-void SSCSEditor::SetItemExpansionRecursive(FSCSEditorTreeNodePtrType Model, bool bInExpansionState)
-{
-	SetNodeExpansionState(Model, bInExpansionState);
-	for (const FSCSEditorTreeNodePtrType& Child : Model->GetChildren())
-	{
-		if (Child.IsValid())
-		{
-			SetItemExpansionRecursive(Child, bInExpansionState);
-		}
-	}
-}
-
-FText SSCSEditor::GetFilterText() const
-{
-	return FilterBox->GetText();
-}
-
-void SSCSEditor::OnFilterTextChanged(const FText& /*InFilterText*/)
-{
-	struct OnFilterTextChanged_Inner
-	{
-		static FSCSEditorTreeNodePtrType ExpandToFilteredChildren(SSCSEditor* SCSEditor, FSCSEditorTreeNodePtrType TreeNode)
-		{
-			FSCSEditorTreeNodePtrType NodeToFocus;
-
-			const TArray<FSCSEditorTreeNodePtrType>& Children = TreeNode->GetChildren();
-			// iterate backwards so we select from the top down
-			for (int32 ChildIndex = Children.Num() - 1; ChildIndex >= 0; --ChildIndex)
-			{
-				const FSCSEditorTreeNodePtrType& Child = Children[ChildIndex];
-				// Don't attempt to focus a separator or filtered node
-				if ((Child->GetNodeType() != FSCSEditorTreeNode::SeparatorNode) && !Child->IsFlaggedForFiltration())
-				{
-					SCSEditor->SetNodeExpansionState(TreeNode, /*bIsExpanded =*/true);
-					NodeToFocus = ExpandToFilteredChildren(SCSEditor, Child);
-				}
-			}
-
-			if (!NodeToFocus.IsValid() && !TreeNode->IsFlaggedForFiltration())
-			{
-				NodeToFocus = TreeNode;
-			}
-			return NodeToFocus;
-		}
-	};
-
-	FSCSEditorTreeNodePtrType NewSelection;
-	// iterate backwards so we select from the top down
-	for (int32 ComponentIndex = RootNodes.Num() - 1; ComponentIndex >= 0; --ComponentIndex)
-	{
-		FSCSEditorTreeNodePtrType Node = RootNodes[ComponentIndex];
-
-		bool bIsRootVisible = !RefreshFilteredState(Node, true);
-		SCSTreeWidget->SetItemExpansion(Node, bIsRootVisible);
-		if (bIsRootVisible)
-		{
-			if (!GetFilterText().IsEmpty())
-			{
-				NewSelection = OnFilterTextChanged_Inner::ExpandToFilteredChildren(this, Node);
-			}
-		}
-	}
-
-	if (!NewSelection.IsValid() && RootNodes.Num() > 0)
-	{
-		NewSelection = RootNodes[0];
-	}
-	
-	if (NewSelection.IsValid() && !SCSTreeWidget->IsItemSelected(NewSelection))
-	{
-		SelectNode(NewSelection, /*IsCntrlDown =*/false);
-	}
-	
-	UpdateTree(/*bRegenerateTreeNodes =*/false);
-}
-
-bool SSCSEditor::RefreshFilteredState(FSCSEditorTreeNodePtrType TreeNode, bool bRecursive)
-{
-	const UClass* FilterType = GetComponentTypeFilterToApply();
-
-	FString FilterText = FText::TrimPrecedingAndTrailing( GetFilterText() ).ToString();
-	TArray<FString> FilterTerms;
-	FilterText.ParseIntoArray(FilterTerms, TEXT(" "), /*CullEmpty =*/true);
-
-	TreeNode->RefreshFilteredState(FilterType, FilterTerms, bRecursive);
-	return TreeNode->IsFlaggedForFiltration();
-}
-
-void SSCSEditor::SetUICustomization(TSharedPtr<ISCSEditorUICustomization> InUICustomization)
-{
-	UICustomization = InUICustomization;
-
-	UpdateTree(true /*bRegenerateTreeNodes*/);
-}
-
-TSubclassOf<UActorComponent> SSCSEditor::GetComponentTypeFilterToApply() const
-{
-	UObject* ActorContextPtr = GetActorContext();
-	TSubclassOf<UActorComponent> ComponentType = UICustomization.IsValid() ? UICustomization->GetComponentTypeFilter(MakeArrayView(&ActorContextPtr, 1)) : nullptr;
-	if (!ComponentType)
-	{
-		ComponentType = ComponentTypeFilter.Get();
-	}
-	return ComponentType;
-}
-
-PRAGMA_ENABLE_DEPRECATION_WARNINGS
 
 #undef LOCTEXT_NAMESPACE
\ No newline at end of file
