Gameplay systems deprecation removal pass for 5.4 and earlier, I skipped anything that was still in use for data fixup or blueprint accessible.

diff --git a/Engine/Source/Editor/Kismet/Public/BlueprintActionMenuBuilder.h b/Engine/Source/Editor/Kismet/Public/BlueprintActionMenuBuilder.h
--- a/Engine/Source/Editor/Kismet/Public/BlueprintActionMenuBuilder.h
+++ b/Engine/Source/Editor/Kismet/Public/BlueprintActionMenuBuilder.h
@@ -68,9 +68,6 @@ public:
 	 * @param Flags	- Optional configuration flags.
 	 */
 	UE_API FBlueprintActionMenuBuilder(EConfigFlags Flags = EConfigFlags::DefaultConfig);
-
-	UE_DEPRECATED(5.2, "The action filter now stores a reference to the authoritative editor context. Please use the default constructor instead.")
-	UE_API FBlueprintActionMenuBuilder(TWeakPtr<FBlueprintEditor> BlueprintEditorPtr);
 	
 	// FGraphActionListBuilderBase interface
 	UE_API virtual void Empty() override;
diff --git a/Engine/Source/Editor/Kismet/Public/BlueprintEditor.h b/Engine/Source/Editor/Kismet/Public/BlueprintEditor.h
--- a/Engine/Source/Editor/Kismet/Public/BlueprintEditor.h
+++ b/Engine/Source/Editor/Kismet/Public/BlueprintEditor.h
@@ -78,7 +78,6 @@ class SFindInBlueprints;
 class SKismetInspector;
 class SMyBlueprint;
 class SReplaceNodeReferences;
-class SSCSEditor;
 class SSCSEditorViewport;
 class SSubobjectEditor;
 class SWidget;
@@ -264,21 +263,12 @@ public:
 	UE_API virtual void SummonFindAndReplaceUI() override;
 	UE_API virtual UEdGraph* GetFocusedGraph() const override;
 	UE_API virtual TSharedPtr<SGraphEditor> OpenGraphAndBringToFront(UEdGraph* Graph, bool bSetFocus = true) override;
-
-	UE_DEPRECATED(5.0, "GetSelectedSCSEditorTreeNodes has been deprecated. Use GetSelectedSubobjectEditorTreeNodes instead.")
-	UE_API virtual TArray<TSharedPtr<class FSCSEditorTreeNode>> GetSelectedSCSEditorTreeNodes() const override;
-	
-	UE_DEPRECATED(5.0, "FindAndSelectSCSEditorTreeNode has been deprecated. Use FindAndSelectSubobjectEditorTreeNode instead.")
-	virtual TSharedPtr<class FSCSEditorTreeNode> FindAndSelectSCSEditorTreeNode(const UActorComponent* InComponent, bool IsCntrlDown) override { return nullptr; }
-
 	UE_API virtual TArray<TSharedPtr<FSubobjectEditorTreeNode>> GetSelectedSubobjectEditorTreeNodes() const override;
 	UE_API virtual TSharedPtr<FSubobjectEditorTreeNode> FindAndSelectSubobjectEditorTreeNode(const UActorComponent* InComponent, bool IsCntrlDown) override;
 	UE_API virtual int32 GetNumberOfSelectedNodes() const override;
 	UE_API virtual void AnalyticsTrackNodeEvent(UBlueprint* Blueprint, UEdGraphNode *GraphNode, bool bNodeDelete = false) const override;
 	UE_API void AnalyticsTrackCompileEvent(UBlueprint* Blueprint, int32 NumErrors, int32 NumWarnings) const;
 	virtual TSharedPtr<class IClassViewerFilter> GetImportedClassViewerFilter() const override { return ImportedClassViewerFilter; }
-	UE_DEPRECATED(5.1, "Please use GetPinTypeSelectorFilters")
-	virtual TSharedPtr<class IPinTypeSelectorFilter> GetImportedPinTypeSelectorFilter() const override { return ImportedPinTypeSelectorFilter; }
 	UE_API virtual void GetPinTypeSelectorFilters(TArray<TSharedPtr<class IPinTypeSelectorFilter>>& OutFilters) const override;
 	UE_API virtual bool IsNonImportedObject(const UObject* InObject) const;
 	UE_API virtual bool IsNonImportedObject(const FSoftObjectPath& InObject) const;
@@ -321,10 +311,7 @@ public:
 	TSharedRef<SWidget> GetCompilerResults() const { return CompilerResults.ToSharedRef(); }
 	TSharedRef<SFindInBlueprints> GetFindResults() const { return FindResults.ToSharedRef(); }
 
-	/** Getters for the various optional Kismet2 widgets */
-	UE_DEPRECATED(5.0, "GetSCSEditor has been deprecated. Use GetSubobjectEditor instead.")
-	TSharedPtr<SSCSEditor> GetSCSEditor() const { return nullptr; }
-	
+	/** Getters for the various optional Kismet2 widgets */	
 	TSharedPtr<SSubobjectEditor> GetSubobjectEditor() const { return SubobjectEditor; }
 	TSharedPtr<SSCSEditorViewport> GetSubobjectViewport() const { return SubobjectViewport; }
 	TSharedPtr<SMyBlueprint> GetMyBlueprintWidget() const { return MyBlueprintWidget; }
@@ -560,14 +547,6 @@ public:
 
 	/** Delegate invoked when a new subobject is added in the subobject editor widget */
 	UE_API virtual void OnComponentAddedToBlueprint(const FSubobjectData& NewSubobjectData);
-	
-	/** Delegate invoked when the selection is changed in the subobject editor widget */
-	UE_DEPRECATED(5.0, "OnSelectionUpdated(const TArray<TSharedPtr<class FSCSEditorTreeNode>>&) has been deprecated. Use OnSelectionUpdated(const TArray<TSharedPtr<class FSubobjectEditorTreeNode>>&) instead.")
-	virtual void OnSelectionUpdated(const TArray<TSharedPtr<class FSCSEditorTreeNode>>& SelectedNodes) {}
-
-	/** Delegate invoked when an item is double clicked in the subobject editor widget */
-	UE_DEPRECATED(5.0, "OnComponentDoubleClicked(TSharedPtr<class FSCSEditorTreeNode>) has been deprecated. Use OnComponentDoubleClicked(TSharedPtr<class FSubobjectEditorTreeNode>) instead.")
-	virtual void OnComponentDoubleClicked(TSharedPtr<class FSCSEditorTreeNode> Node) {}
 
 	/** Pin visibility accessors */
 	UE_API void SetPinVisibility(SGraphEditor::EPinVisibility Visibility);
@@ -668,12 +647,6 @@ public:
 	/** Sets subobject editor UI customization */
 	UE_API void SetSubobjectEditorUICustomization(TSharedPtr<class ISCSEditorUICustomization> InSCSEditorUICustomization);
 
-	UE_DEPRECATED(5.0, "SetSCSEditorUICustomization has been deprecated. Use SetSubobjectEditorUICustomization instead.")
-	void SetSCSEditorUICustomization(TSharedPtr<class ISCSEditorUICustomization> InSCSEditorUICustomization)
-	{
-		return SetSubobjectEditorUICustomization(InSCSEditorUICustomization);
-	}
-
 	/**
 	 * Register a customization for interacting with the subobject editor
 	 * @param	InComponentName			The name of the component to customize behavior for
@@ -707,12 +680,6 @@ public:
 	 */
 	UE_API TSharedPtr<class ISCSEditorCustomization> CustomizeSubobjectEditor(const USceneComponent* InComponentToCustomize) const;
 
-	UE_DEPRECATED(5.0, "CustomizeSCSEditor has been deprecated. Use CustomizeSubobjectEditor instead.")
-	TSharedPtr<class ISCSEditorCustomization> CustomizeSCSEditor(USceneComponent* InComponentToCustomize) const
-	{
-		return CustomizeSubobjectEditor(InComponentToCustomize);
-	}
-
 	/** Adds to a list of custom objects for debugging beyond what will automatically be found/used */
 	virtual void GetCustomDebugObjects(TArray<FCustomDebugObject>& DebugList) const { }
 
@@ -1536,9 +1503,6 @@ public:
 	/** Make nodes which are unrelated to the selected nodes fade out */
 	UE_API virtual void ToggleHideUnrelatedNodes();
 	UE_API virtual bool IsToggleHideUnrelatedNodesChecked() const;
-	
-	UE_DEPRECATED(5.0, "The Toggle Hide Unrelated Nodes button is always shown now")
-	UE_API bool ShouldShowToggleHideUnrelatedNodes(bool bIsToolbar) const;
 
 	/** Make a drop down menu to control the opacity of unrelated nodes */
 	UE_API TSharedRef<SWidget> MakeHideUnrelatedNodesOptionsMenu();
diff --git a/Engine/Source/Editor/Kismet/Public/BlueprintEditorModule.h b/Engine/Source/Editor/Kismet/Public/BlueprintEditorModule.h
--- a/Engine/Source/Editor/Kismet/Public/BlueprintEditorModule.h
+++ b/Engine/Source/Editor/Kismet/Public/BlueprintEditorModule.h
@@ -178,10 +178,6 @@ public:
 	/** Return the class viewer filter associated with the current set of imported namespaces within this editor context. Default is NULL (no filter). */
 	virtual TSharedPtr<class IClassViewerFilter> GetImportedClassViewerFilter() const { return nullptr; }
 
-	/** Return the pin type selector filter associated with the current set of imported namespaces within this editor context. Default is NULL (no filter). */
-	UE_DEPRECATED(5.1, "Please use GetPinTypeSelectorFilters")
-	virtual TSharedPtr<class IPinTypeSelectorFilter> GetImportedPinTypeSelectorFilter() const { return nullptr; }
-
 	/** Get all the the pin type selector filters within this editor context. */
 	virtual void GetPinTypeSelectorFilters(TArray<TSharedPtr<class IPinTypeSelectorFilter>>& OutFilters) const {}
 	
@@ -190,11 +186,6 @@ public:
 
 	/** Return whether the given object (referenced by path) falls outside the scope of the current set of imported namespaces within this editor context. Default is FALSE (imported). */
 	virtual bool IsNonImportedObject(const FSoftObjectPath& InObject) const { return false; }
-
-	UE_DEPRECATED(5.0, "GetSelectedSCSEditorTreeNodes has been deprecated. Use GetSelectedSubobjectEditorTreeNodes instead.")
-	virtual TArray<TSharedPtr<class FSCSEditorTreeNode> >  GetSelectedSCSEditorTreeNodes() const = 0;
-	UE_DEPRECATED(5.0, "FindAndSelectSCSEditorTreeNode has been deprecated. Use FindAndSelectSubobjectEditorTreeNode instead.")
-	virtual TSharedPtr<class FSCSEditorTreeNode> FindAndSelectSCSEditorTreeNode(const class UActorComponent* InComponent, bool IsCntrlDown) = 0;
 };
 
 DECLARE_DELEGATE_RetVal_OneParam(TSharedRef<class ISCSEditorCustomization>, FSCSEditorCustomizationBuilder, TSharedRef< IBlueprintEditor > /* InBlueprintEditor */);
@@ -299,13 +290,6 @@ public:
 	 */
 	virtual FDelegateHandle RegisterVariableCustomization(FFieldClass* InFieldClass, FOnGetVariableCustomizationInstance InOnGetVariableCustomization);
 
-	/** 
-	 * Unregister a previously registered customization for BP variables
-	 * @param	InStruct				The type to create the customization for
-	 */
-	UE_DEPRECATED(5.1, "UnregisterVariableCustomization without a delegate handle is deprecated.")
-	virtual void UnregisterVariableCustomization(FFieldClass* InFieldClass);
-
 	/**
 	 * Unregister a previously registered customization for BP variables
 	 * @param	InStruct				The type to create the customization for
@@ -320,13 +304,6 @@ public:
 	 */
 	virtual FDelegateHandle RegisterLocalVariableCustomization(FFieldClass* InFieldClass, FOnGetLocalVariableCustomizationInstance InOnGetLocalVariableCustomization);
 
-	/** 
-	* Unregister a previously registered customization for BP local variables
-	* @param	InFieldClass				The type to create the customization for
-	 */
-	UE_DEPRECATED(5.1, "UnregisterLocalVariableCustomization without a delegate handle is deprecated.")
-	virtual void UnregisterLocalVariableCustomization(FFieldClass* InFieldClass);
-
 	/** 
 	 * Unregister a previously registered customization for BP local variables
 	 * @param	InFieldClass				The type to create the customization for
diff --git a/Engine/Source/Editor/Kismet/Public/BlueprintNamespaceUtilities.h b/Engine/Source/Editor/Kismet/Public/BlueprintNamespaceUtilities.h
--- a/Engine/Source/Editor/Kismet/Public/BlueprintNamespaceUtilities.h
+++ b/Engine/Source/Editor/Kismet/Public/BlueprintNamespaceUtilities.h
@@ -64,12 +64,6 @@ public:
 	 */
 	static UE_API void GetPropertyValueNamespaces(const FProperty* InProperty, const void* InContainer, TSet<FString>& OutNamespaces);
 
-	UE_DEPRECATED(5.1, "Please use the updated version that removes the InStruct parameter (no longer needed).")
-	static void GetPropertyValueNamespaces(const UStruct* InStruct, const FProperty* InProperty, const void* InContainer, TSet<FString>& OutNamespaces)
-	{
-		GetPropertyValueNamespaces(InProperty, InContainer, OutNamespaces);
-	}
-
 	/**
 	 * Gathers the set of global namespaces that are implicitly imported by all Blueprint assets.
 	 * 
@@ -77,15 +71,6 @@ public:
 	 */
 	static UE_API void GetSharedGlobalImports(TSet<FString>& OutNamespaces);
 
-	/**
-	 * Gathers the set of default namespaces that are implicitly imported by a given Blueprint asset.
-	 *
-	 * @param InBlueprint	A reference to a loaded Blueprint asset.
-	 * @param OutNamespaces	Zero or more unique namespace identifier(s) representing the Blueprint asset's default namespace set.
-	 */
-	UE_DEPRECATED(5.1, "Please use GetDefaultImportsForObject instead")
-	static UE_API void GetDefaultImportsForBlueprint(const UBlueprint* InBlueprint, TSet<FString>& OutNamespaces);
-
 	/**
 	 * Gathers the set of default namespaces that are implicitly imported by a given object's type.
 	 *
diff --git a/Engine/Source/Editor/Kismet/Public/SBlueprintNamespaceEntry.h b/Engine/Source/Editor/Kismet/Public/SBlueprintNamespaceEntry.h
--- a/Engine/Source/Editor/Kismet/Public/SBlueprintNamespaceEntry.h
+++ b/Engine/Source/Editor/Kismet/Public/SBlueprintNamespaceEntry.h
@@ -43,9 +43,6 @@ public:
 	DECLARE_DELEGATE_OneParam(FOnNamespaceSelected, const FString&);
 	DECLARE_DELEGATE_OneParam(FOnGetNamespacesToExclude, TSet<FString>&);
 
-	UE_DEPRECATED(5.1, "Use FOnGetNamespacesToExclude instead.")
-	DECLARE_DELEGATE_OneParam(FOnFilterNamespaceList, TArray<FString>&);
-
 	SLATE_BEGIN_ARGS(SBlueprintNamespaceEntry)
 	: _Font(FCoreStyle::Get().GetFontStyle(TEXT("NormalFont")))
 	, _AllowTextEntry(true)
@@ -71,15 +68,6 @@ public:
 		/** Tooltip used for excluded namespaces that are visible in the selection drop-down. */
 		SLATE_ATTRIBUTE(FText, ExcludedNamespaceTooltipText)
 
-		/** [DEPRECATED] Called to allow external code to filter out the namespace list. */
-PRAGMA_DISABLE_DEPRECATION_WARNINGS
-		FOnGetNamespacesToExclude ConvertOnFilterNamespaceListFn(const FOnFilterNamespaceList& LegacyDelegate)
-		{
-			return FOnGetNamespacesToExclude::CreateStatic(SBlueprintNamespaceEntry::HandleLegacyOnFilterNamespaceList, LegacyDelegate);
-		}
-		SLATE_EVENT_DEPRECATED(5.1, "Use OnGetNamespacesToExclude instead.", FOnFilterNamespaceList, OnFilterNamespaceList, OnGetNamespacesToExclude, ConvertOnFilterNamespaceListFn)
-PRAGMA_ENABLE_DEPRECATION_WARNINGS
-
 	SLATE_END_ARGS()
 
 	UE_API void Construct(const FArguments& InArgs);
@@ -114,9 +102,6 @@ protected:
 	UE_API void SelectNamespace(const FString& InNamespace);
 	UE_API void ExpandAllTreeViewItems(const TArray<TSharedPtr<FPathTreeNodeItem>>* NodeListPtr = nullptr);
 	UE_API const TSharedPtr<FPathTreeNodeItem>* FindTreeViewNode(const FString& NodePath, const TArray<TSharedPtr<FPathTreeNodeItem>>* NodeListPtr = nullptr) const;
-UE_API PRAGMA_DISABLE_DEPRECATION_WARNINGS
-	static void HandleLegacyOnFilterNamespaceList(TSet<FString>& OutNamespacesToExclude, FOnFilterNamespaceList LegacyDelegate);
-PRAGMA_ENABLE_DEPRECATION_WARNINGS
 
 private:
 	FString CurrentNamespace;
diff --git a/Engine/Source/Editor/Kismet/Public/SCSDiff.h b/Engine/Source/Editor/Kismet/Public/SCSDiff.h
--- a/Engine/Source/Editor/Kismet/Public/SCSDiff.h
+++ b/Engine/Source/Editor/Kismet/Public/SCSDiff.h
@@ -13,7 +13,6 @@
 class FSCSEditorTreeNode;
 class FSubobjectEditorTreeNode;
 class SKismetInspector;
-class SSCSEditor;
 class SSubobjectBlueprintEditor;
 class SWidget;
 class UBlueprint;
diff --git a/Engine/Source/Editor/Kismet/Public/SSCSEditor.h b/Engine/Source/Editor/Kismet/Public/SSCSEditor.h
--- a/Engine/Source/Editor/Kismet/Public/SSCSEditor.h
+++ b/Engine/Source/Editor/Kismet/Public/SSCSEditor.h
@@ -94,6 +94,8 @@ using FSCSEditorChildActorNodePtrType = TSharedPtr<class FSCSEditorTreeNodeChild
  * FSCSEditorTreeNode
  *
  * Wrapper class for nodes displayed in the SCS (Simple Construction Script) editor tree widget.
+ * 
+ * NOTE: Mostly deprecated, switch to FSubobjectEditorTreeNode
  */
 class FSCSEditorTreeNode : public TSharedFromThis<FSCSEditorTreeNode>
 {
@@ -325,6 +327,9 @@ public:
 	 */
 	static UE_API FSCSEditorTreeNodePtrType FactoryNodeFromComponent(UActorComponent* InComponent);
 
+	/** Returns a color to use for this type of node */
+	static UE_API FSlateColor GetColorTintForIcon(FSCSEditorTreeNodePtrType InNode);
+
 	// Destructor
 	virtual ~FSCSEditorTreeNode() = default;
 
@@ -742,799 +747,15 @@ private:
 	TArray<const UClass*> FilteredTypes;
 };
 
-//////////////////////////////////////////////////////////////////////////
-// SSCS_RowWidget
+class UE_DEPRECATED(5.8, "SSCSEditor was removed, use either SSubObjectEditor or FSCSEditorTreeNode functions") SSCS_RowWidget;
 
-class SSCS_RowWidget : public SMultiColumnTableRow<FSCSEditorTreeNodePtrType>
+class SSCS_RowWidget
 {
-public:
-	SLATE_BEGIN_ARGS( SSCS_RowWidget ){}
-	SLATE_END_ARGS()
-
-	void Construct( const FArguments& InArgs, TSharedPtr<SSCSEditor> InSCSEditor, FSCSEditorTreeNodePtrType InNodePtr, TSharedPtr<STableViewBase> InOwnerTableView  );
-	virtual ~SSCS_RowWidget();
-
-	// SMultiColumnTableRow<T> interface
-	virtual TSharedRef<SWidget> GenerateWidgetForColumn( const FName& ColumnName ) override;
-	// End of SMultiColumnTableRow<T>
-
-	// SWidget interface
-	virtual FReply OnMouseButtonDown(const FGeometry& MyGeometry, const FPointerEvent& MouseEvent) override;
-	// End of SWidget interface
-
-	/** Get the blueprint we are editing */
-	UBlueprint* GetBlueprint() const;
-
-	FText GetNameLabel() const;
-	FText GetTooltipText() const;
-	FSlateColor GetColorTintForIcon() const;
-	FSlateColor GetColorTintForText() const;
-	FString GetDocumentationLink() const;
-	FString GetDocumentationExcerptName() const;
-	
-	KISMET_API static FSlateColor GetColorTintForIcon(FSCSEditorTreeNodePtrType InNode);
-
-	FText GetAssetName() const;
-	FText GetAssetPath() const;
-	EVisibility GetAssetVisibility() const;
-
-	/* Get the node used by the row Widget */
-	virtual FSCSEditorTreeNodePtrType GetNode() const { return TreeNodePtr; };
-
-protected:
-	virtual ESelectionMode::Type GetSelectionMode() const override;
-
-	virtual const FSlateBrush* GetIconBrush() const;
-
-	static void AddToToolTipInfoBox(const TSharedRef<SVerticalBox>& InfoBox, const FText& Key, TSharedRef<SWidget> ValueIcon, const TAttribute<FText>& Value, bool bImportant);
-
-	/** Commits the new name of the component */
-	void OnNameTextCommit(const FText& InNewName, ETextCommit::Type InTextCommit);
-
-private:
-	/** Verifies the name of the component when changing it */
-	bool OnNameTextVerifyChanged(const FText& InNewText, FText& OutErrorMessage);
-
-	/** Builds a context menu popup for dropping a child node onto the scene root node */
-	TSharedPtr<SWidget> BuildSceneRootDropActionMenu(FSCSEditorTreeNodePtrType DroppedNodePtr);
-
-	/** Creates a tooltip for this row */
-	TSharedRef<SToolTip> CreateToolTipWidget() const;
-
-	/** Drag-drop handlers */
-	void HandleOnDragEnter(const FDragDropEvent& DragDropEvent);
-	void HandleOnDragLeave(const FDragDropEvent& DragDropEvent);
-	FReply HandleOnDragDetected(const FGeometry& MyGeometry, const FPointerEvent& MouseEvent);
-	TOptional<EItemDropZone> HandleOnCanAcceptDrop(const FDragDropEvent& DragDropEvent, EItemDropZone DropZone, FSCSEditorTreeNodePtrType TargetItem);
-	FReply HandleOnAcceptDrop(const FDragDropEvent& DragDropEvent, EItemDropZone DropZone, FSCSEditorTreeNodePtrType TargetItem);
-
-	/** Handler for attaching a single node to this node */
-	void OnAttachToDropAction(FSCSEditorTreeNodePtrType DroppedNodePtr)
+	UE_DEPRECATED(5.8, "Use FSCSEditorTreeNode::GetColorTintForIcon")
+	static FSlateColor GetColorTintForIcon(FSCSEditorTreeNodePtrType InNode)
 	{
-		TArray<FSCSEditorTreeNodePtrType> DroppedNodePtrs;
-		DroppedNodePtrs.Add(DroppedNodePtr);
-		OnAttachToDropAction(DroppedNodePtrs);
+		return FSCSEditorTreeNode::GetColorTintForIcon(InNode);
 	}
-
-	/** Handler for attaching one or more nodes to this node */
-	void OnAttachToDropAction(const TArray<FSCSEditorTreeNodePtrType>& DroppedNodePtrs);
-
-	/** Handler for detaching one or more nodes from the current parent and reattaching to the existing scene root node */
-	void OnDetachFromDropAction(const TArray<FSCSEditorTreeNodePtrType>& DroppedNodePtrs);
-
-	/** Handler for making the given node the new scene root node */
-	void OnMakeNewRootDropAction(FSCSEditorTreeNodePtrType DroppedNodePtr);
-	
-	/** Tasks to perform after handling a drop action */
-	void PostDragDropAction(bool bRegenerateTreeNodes);
-
-	/**
-	 * Retrieves an image brush signifying the specified component's mobility (could sometimes be NULL).
-	 * 
-	 * @returns A pointer to the FSlateBrush to use (NULL for Static and Non-SceneComponents)
-	 */
-	FSlateBrush const* GetMobilityIconImage() const;
-
-	/**
-	 * Retrieves tooltip text describing the specified component's mobility.
-	 * 
-	 * @returns An FText object containing a description of the component's mobility
-	 */
-	FText GetMobilityToolTipText() const;
-
-	/**
-	 * Retrieves tooltip text describing where the component was first introduced (for inherited components).
-	 * 
-	 * @returns An FText object containing a description of when the component was first introduced
-	 */
-	FText GetIntroducedInToolTipText() const;
-
-	/**
-	 * Retrieves tooltip text describing how the component was introduced
-	 * 
-	 * @returns An FText object containing a description of when the component was first introduced
-	 */
-	FText GetComponentAddSourceToolTipText() const;
-
-/**
-	 * Retrieves tooltip text for the specified Native Component's underlying Name
-	 *
-	 * @returns An FText object containing the Component's Name
-	 */
-	FText GetNativeComponentNameToolTipText() const;
-
-	/**
-	 * Retrieves a tooltip text describing if the component is marked Editor only or not
-	 *
-	 * @returns An FText object containing a description of if the component is marked Editor only or not
-	 */
-	FText GetComponentEditorOnlyTooltipText() const;
-
-public:
-	/** Pointer back to owning SCSEditor 2 tool */
-	TWeakPtr<SSCSEditor> SCSEditor;
-	TSharedPtr<SInlineEditableTextBlock> InlineWidget;
-private:
-	/** Pointer to node we represent */
-	FSCSEditorTreeNodePtrType TreeNodePtr;
-};
-
-class SSCS_RowWidget_ActorRoot : public SSCS_RowWidget
-{
-public:
-
-	// SMultiColumnTableRow<T> interface
-	virtual TSharedRef<SWidget> GenerateWidgetForColumn(const FName& ColumnName) override;
-	// End of SMultiColumnTableRow<T>
-
-protected:
-	/** Gets the associated actor node */
-	FSCSEditorActorNodePtrType GetActorNode() const;
-
-	/** Data accessors */
-	virtual const FSlateBrush* GetIconBrush() const override;
-
-private:
-	/** Creates a tooltip for this row */
-	TSharedRef<SToolTip> CreateToolTipWidget() const;
-
-	/** Called to validate the actor name */
-	bool OnVerifyActorLabelChanged(const FText& InLabel, FText& OutErrorMessage);
-
-	/** Data accessors */
-	FText GetActorDisplayText() const;
-	FText GetActorContextText() const;
-	FText GetActorClassNameText() const;
-	FText GetActorSuperClassNameText() const;
-	FText GetActorMobilityText() const;
-};
-
-class SSCS_RowWidget_Separator : public SSCS_RowWidget
-{
-public:
-
-	// SMultiColumnTableRow<T> interface
-	virtual TSharedRef<SWidget> GenerateWidgetForColumn(const FName& ColumnName) override;
-	// End of SMultiColumnTableRow<T>
-
-private:
-
 };
 
-//////////////////////////////////////////////////////////////////////////
-// SSCSEditorDragDropTree - implements STreeView for our specific node type and adds drag/drop functionality
-class SSCSEditorDragDropTree : public STreeView<FSCSEditorTreeNodePtrType>
-{
-public:
-	SLATE_BEGIN_ARGS( SSCSEditorDragDropTree )
-		: _SCSEditor( NULL )
-		, _OnGenerateRow()
-		, _OnGetChildren()
-		, _OnSetExpansionRecursive()
-		, _TreeItemsSource( static_cast< TArray<FSCSEditorTreeNodePtrType>* >(NULL) ) //@todo Slate Syntax: Initializing from NULL without a cast
-		, _ItemHeight(16)
-		, _OnContextMenuOpening()
-		, _OnMouseButtonDoubleClick()
-		, _OnSelectionChanged()
-		, _OnExpansionChanged()
-		, _SelectionMode(ESelectionMode::Multi)
-		, _ClearSelectionOnClick(true)
-		, _ExternalScrollbar()
-		, _OnTableViewBadState()
-		{
-			_Clipping = EWidgetClipping::ClipToBounds;
-		}
-
-		SLATE_ARGUMENT( SSCSEditor*, SCSEditor )
-
-		SLATE_EVENT( FOnGenerateRow, OnGenerateRow )
-
-		SLATE_EVENT( FOnItemScrolledIntoView, OnItemScrolledIntoView )
-
-		SLATE_EVENT( FOnGetChildren, OnGetChildren )
-
-		SLATE_EVENT( FOnSetExpansionRecursive, OnSetExpansionRecursive )
-
-		SLATE_ARGUMENT( TArray<FSCSEditorTreeNodePtrType>* , TreeItemsSource )
-
-		SLATE_ATTRIBUTE( float, ItemHeight )
-
-		SLATE_EVENT( FOnContextMenuOpening, OnContextMenuOpening )
-
-		SLATE_EVENT( FOnMouseButtonDoubleClick, OnMouseButtonDoubleClick )
-
-		SLATE_EVENT( FOnSelectionChanged, OnSelectionChanged )
-
-		SLATE_EVENT( FOnExpansionChanged, OnExpansionChanged )
-
-		SLATE_ATTRIBUTE( ESelectionMode::Type, SelectionMode )
-
-		SLATE_ARGUMENT( TSharedPtr<SHeaderRow>, HeaderRow )
-
-		SLATE_ARGUMENT ( bool, ClearSelectionOnClick )
-
-		SLATE_ARGUMENT( TSharedPtr<SScrollBar>, ExternalScrollbar )
-
-		SLATE_EVENT( FOnTableViewBadState, OnTableViewBadState )
-
-	SLATE_END_ARGS()
-	/** Object construction - mostly defers to the base STreeView */
-	void Construct( const FArguments& InArgs );
-
-	// SWidget interface
-	virtual FReply OnDragOver( const FGeometry& MyGeometry, const FDragDropEvent& DragDropEvent ) override;
-	virtual FReply OnDrop( const FGeometry& MyGeometry, const FDragDropEvent& DragDropEvent ) override;
-	// End SWidget interface
-
-private:
-	/** Pointer to the SSCSEditor that owns this widget */
-	SSCSEditor* SCSEditor;
-};
-
-//////////////////////////////////////////////////////////////////////////
-// SSCSEditor
-
-typedef SSCSEditorDragDropTree SSCSTreeType;
-
-/* Component editor mode */
-namespace EComponentEditorMode
-{
-	enum Type
-	{
-		/* View/edit the SCS in a BGPC */
-		BlueprintSCS,
-		/* View/edit the Actor instance */
-		ActorInstance
-	};
-};
-
-class UE_DEPRECATED(5.0, "SSCSEditor has been deprecated, use a child class of SSubobjectEditor instead.") SSCSEditor;
-
-class SSCSEditor : public SCompoundWidget
-{
-public:
-	DECLARE_DELEGATE_RetVal_OneParam(class USCS_Node*, FOnAddNewComponent, class UClass*);
-	DECLARE_DELEGATE_RetVal_OneParam(class USCS_Node*, FOnAddExistingComponent, class UActorComponent*);
-	DECLARE_DELEGATE_OneParam(FOnSelectionUpdated, const TArray<FSCSEditorTreeNodePtrType>&);
-	DECLARE_DELEGATE_OneParam(FOnItemDoubleClicked, const FSCSEditorTreeNodePtrType);
-	DECLARE_DELEGATE_OneParam(FOnHighlightPropertyInDetailsView, const class FPropertyPath&);
-
-	SLATE_BEGIN_ARGS( SSCSEditor )
-		:_EditorMode(EComponentEditorMode::BlueprintSCS)
-		,_IsDiffing(false)
-		,_ActorContext(nullptr)
-		,_PreviewActor(nullptr)
-		,_AllowEditing(true)
-		,_HideComponentClassCombo(false)
-		,_ComponentTypeFilter()
-		,_OnSelectionUpdated()
-		,_OnHighlightPropertyInDetailsView()
-		{}
-
-		SLATE_ARGUMENT(EComponentEditorMode::Type, EditorMode)
-		SLATE_ARGUMENT(bool, IsDiffing)
-		SLATE_ATTRIBUTE(class AActor*, ActorContext)
-		SLATE_ATTRIBUTE(class AActor*, PreviewActor)
-		SLATE_ATTRIBUTE(bool, AllowEditing)
-		SLATE_ATTRIBUTE(bool, HideComponentClassCombo)
-		SLATE_ATTRIBUTE(TSubclassOf<UActorComponent>, ComponentTypeFilter)
-		SLATE_EVENT(FOnSelectionUpdated, OnSelectionUpdated)
-		SLATE_EVENT(FOnItemDoubleClicked, OnItemDoubleClicked)
-		SLATE_EVENT(FOnHighlightPropertyInDetailsView, OnHighlightPropertyInDetailsView)
-		SLATE_EVENT(FSimpleDelegate, OnObjectReplaced)
-
-	SLATE_END_ARGS()
-
-	UE_API void Construct(const FArguments& InArgs);
-
-	UE_API ~SSCSEditor();
-
-	/** SWidget interface */
-	UE_API virtual FReply OnKeyDown( const FGeometry& MyGeometry, const FKeyEvent& InKeyEvent );
-
-	/** Used by tree control - make a widget for a table row from a node */
-	UE_API TSharedRef<ITableRow> MakeTableRowWidget(FSCSEditorTreeNodePtrType InNodePtr, const TSharedRef<STableViewBase>& OwnerTable);
-
-	/** Used by tree control - get children for a specified node */
-	UE_API void OnGetChildrenForTree(FSCSEditorTreeNodePtrType InNodePtr, TArray<FSCSEditorTreeNodePtrType>& OutChildren);
-
-	/** Returns true if editing is allowed */
-	UE_API bool IsEditingAllowed() const;
-
-	/** gets the actor root */
-	UE_API FSCSEditorActorNodePtrType GetActorNode() const;
-
-	/** get the root scene node */
-	UE_API FSCSEditorTreeNodePtrType GetSceneRootNode() const;
-
-	UE_API void SetSceneRootNode(FSCSEditorTreeNodePtrType NewSceneRootNode);
-
-	/** Adds a component to the SCS tree */
-	struct FAddNewComponentParams
-	{
-		FAddNewComponentParams()
-			: bSkipMarkBlueprintModified(false)
-			, bSetFocusToNewItem(true)
-			, bConformTransformToParent(true)
-		{
-		}
-		
-		/** Optionally skip marking this blueprint as modified (e.g. if we're handling that externally */
-		bool bSkipMarkBlueprintModified;
-		/** Whether the newly created component should be focused */
-		bool bSetFocusToNewItem;
-		/** Whether the newly created component should keep its transform, or conform it to its parent */
-		bool bConformTransformToParent;
-	};
-
-	/** Adds a component to the SCS Table
-	   @param NewComponentClass				(In) The class to add
-	   @param Asset       					(In) Optional asset to assign to the component
-	   @param Params       					(In) Parameter block of optional behavior flags
-	 */
-	UE_API UActorComponent* AddNewComponent(UClass* NewComponentClass, UObject* Asset, const FAddNewComponentParams Params = FAddNewComponentParams());
-
-	struct FAddedNodeDetails
-	{
-		FSCSEditorTreeNodePtrType NewNodePtr;
-		FSCSEditorTreeNodePtrType ParentNodePtr;
-	};
-
-	/** Adds a new SCS Node to the component Table
-	   @param OutNodeDetails (Out) Struct to be populated by the tree node pointers of the new node and its parent
-	   @param OngoingCreateTransaction (In) The transaction containing the creation of the node. The transaction will remain ongoing until the node gets its initial name from user.
-	   @param NewNode	(In) The SCS node to add
-	   @param Asset		(In) Optional asset to assign to the component
-	   @param bMarkBlueprintModified (In) Whether or not to mark the Blueprint as structurally modified
-	   @param bSetFocusToNewItem (In) Select the new item and activate the inline rename widget (default is true) */
-	UE_API void AddNewNode(FAddedNodeDetails& OutNodeDetails, TUniquePtr<FScopedTransaction> OngoingCreateTransaction, USCS_Node* NewNode, UObject* Asset, bool bMarkBlueprintModified, bool bSetFocusToNewItem = true);
-
-	/** Adds a new SCS Node to the component Table
-	   @param OngoingCreateTransaction (In) The transaction containing the creation of the node. The transaction will remain ongoing until the node gets its initial name from user.
-	   @param NewNode	(In) The SCS node to add
-	   @param Asset		(In) Optional asset to assign to the component
-	   @param bMarkBlueprintModified (In) Whether or not to mark the Blueprint as structurally modified
-	   @param bSetFocusToNewItem (In) Select the new item and activate the inline rename widget (default is true)
-	   @return The reference of the newly created ActorComponent */
-	UE_API UActorComponent* AddNewNode(TUniquePtr<FScopedTransaction> OngoingCreateTransaction, USCS_Node* NewNode, UObject* Asset, bool bMarkBlueprintModified, bool bSetFocusToNewItem = true);
-
-	/** Adds a new component instance node to the component Table
-		@param OngoingCreateTransaction (In) The transaction containing the creation of the node. The transaction will remain ongoing until the node gets its initial name from user.
-		@param NewInstanceComponent	(In) The component being added to the actor instance
-		@param InParentNodePtr (In) The node this component will be added to
-		@param Asset (In) Optional asset to assign to the component
-		@param bSetFocusToNewItem (In) Select the new item and activate the inline rename widget (default is true) */
-	UE_API void AddNewNodeForInstancedComponent(TUniquePtr<FScopedTransaction> OngoingCreateTransaction, UActorComponent* NewInstanceComponent, FSCSEditorTreeNodePtrType InParentNodePtr, UObject* Asset, bool bSetFocusToNewItem = true);
-	
-	/** Returns true if the specified component is currently selected */
-	UE_API bool IsComponentSelected(const UPrimitiveComponent* PrimComponent) const;
-
-	/** Assigns a selection override delegate to the specified component */
-	UE_API void SetSelectionOverride(UPrimitiveComponent* PrimComponent) const;
-
-	/** Cut selected node(s) */
-	UE_API void CutSelectedNodes();
-	UE_API bool CanCutNodes() const;
-
-	/** Copy selected node(s) */
-	UE_API void CopySelectedNodes();
-	UE_API bool CanCopyNodes() const;
-
-	/** Pastes previously copied node(s) */
-	UE_API void PasteNodes();
-	UE_API bool CanPasteNodes() const;
-
-	/** Callbacks to duplicate the selected component */
-	UE_API bool CanDuplicateComponent() const;
-	UE_API void OnDuplicateComponent();
-
-	/** Removes existing selected component nodes from the SCS */
-	UE_API void OnDeleteNodes();
-	UE_API bool CanDeleteNodes() const;
-
-	/** Callbacks to find references of the selected component */
-	UE_API void OnFindReferences(bool bSearchAllBlueprints, const EGetFindReferenceSearchStringFlags Flags);
-
-	/** Removes an existing component node from the tree */
-	UE_API void RemoveComponentNode(FSCSEditorTreeNodePtrType InNodePtr);
-
-	/** Called when selection in the tree changes */
-	UE_API void OnTreeSelectionChanged(FSCSEditorTreeNodePtrType InSelectedNodePtr, ESelectInfo::Type SelectInfo);
-
-	/** Called when the Actor is selected. */
-	UE_API void OnActorSelected(const ECheckBoxState NewCheckedState);
-
-	/** Called to determine if actor is selected. */
-	UE_API ECheckBoxState OnIsActorSelected() const;
-
-	/** Update any associated selection (e.g. details view) from the passed in nodes */
-	UE_API void UpdateSelectionFromNodes(const TArray<FSCSEditorTreeNodePtrType> &SelectedNodes );
-
-	/** Refresh the tree control to reflect changes in the SCS */
-	UE_API void UpdateTree(bool bRegenerateTreeNodes = true);
-
-	/** Dumps out the tree view contents to the log (used to assist with debugging widget hierarchy issues) */
-	UE_API void DumpTree();
-
-	/** Forces the details panel to refresh on the same objects */
-	UE_API void RefreshSelectionDetails();
-
-	/** Clears the current selection */
-	UE_API void ClearSelection();
-
-	/** Get the currently selected tree nodes */
-	UE_API TArray<FSCSEditorTreeNodePtrType> GetSelectedNodes() const;
-
-	/** Get the number of currently selected tree nodes */
-	int32 GetNumSelectedNodes() const { return SCSTreeWidget->GetNumItemsSelected(); }
-
-	/**
-	 * Fills out an events section in ui.
-	 * @param Menu								the menu to add the events section into
-	 * @param Blueprint							the active blueprint context being edited
-	 * @param SelectedClass						the common component class to build the events list from
-	 * @param CanExecuteActionDelegate			the delegate to query whether or not to execute the UI action
-	 * @param GetSelectedObjectsDelegate		the delegate to fill the currently select variables / components
-	 */
-	static UE_API void BuildMenuEventsSection( FMenuBuilder& Menu, UBlueprint* Blueprint, UClass* SelectedClass, FCanExecuteAction CanExecuteActionDelegate, FGetSelectedObjectsDelegate GetSelectedObjectsDelegate );
-
-	/**
-	 * Given an actor component, attempts to find an associated tree node.
-	 *
-	 * @param ActorComponent The component associated with the node.
-	 * @param bIncludeAttachedComponents Whether or not to include components attached to each node in the search (default is true).
-	 * @return A shared pointer to a tree node. The pointer will be invalid if no match could be found.
-	 */
-	UE_API FSCSEditorTreeNodePtrType GetNodeFromActorComponent(const UActorComponent* ActorComponent, bool bIncludeAttachedComponents = true) const;
-
-	/** Select the root of the tree */
-	UE_API void SelectRoot();
-
-	/** Select the given tree node */
-	UE_API void SelectNode(FSCSEditorTreeNodePtrType InNodeToSelect, bool IsCntrlDown);
-
-	/**
-	 * Set the expansion state of a node
-	 *
-	 * @param InNodeToChange	The node to be expanded/collapsed
-	 * @param bIsExpanded		True to expand the node, false to collapse it
-	 */
-	UE_API void SetNodeExpansionState(FSCSEditorTreeNodePtrType InNodeToChange, const bool bIsExpanded);
-
-	/**
-	 * Highlight a tree node and, optionally, a property with in it
-	 *
-	 * @param TreeNodeName		Name of the treenode to be highlighted
-	 * @param Property	The name of the property to be highlighted in the details view
-	 * @return True if the node was found in this Editor, otherwise false
-	 */
-	UE_API void HighlightTreeNode( FName TreeNodeName, const class FPropertyPath& Property );
-	/**
-	 * Highlight a tree node and, optionally, a property with in it
-	 *
-	 * @param Node		 A Reference to the Node SCS_Node to be highlighted
-	 * @param Property	The name of the property to be highlighted in the details view
-	 */
-	UE_API void HighlightTreeNode( const USCS_Node* Node, FName Property );
-
-	/**
-	 * Function to save current state of SimpleConstructionScript and nodes associated with it.
-	 *
-	 * @param: SimpleContructionScript object reference.
-	 */
-	static UE_API void SaveSCSCurrentState( USimpleConstructionScript* SCSObj );
-
-	/**
-	 * Function to save the current state of SCS_Node and its children
-	 *
-	 * @param: Reference of the SCS_Node to be saved
-	 */
-	static UE_API void SaveSCSNode(USCS_Node* Node);
-
-	/** Is this node still used by the Simple Construction Script */
-	UE_API bool IsNodeInSimpleConstructionScript(USCS_Node* Node) const;
-
-	/**
-	 * Fills the supplied array with the currently selected objects
-	 * @param OutSelectedItems The array to fill.
-	 */
-	UE_API void GetSelectedItemsForContextMenu(TArray<FComponentEventConstructionData>& OutSelectedItems) const;
-
-	/** @return Array of the editable objects selected in the tree */
-	UE_API TArray<UObject*> GetSelectedEditableObjects() const;
-
-	/** Provides access to the Blueprint context that's being edited */
-	UE_API class UBlueprint* GetBlueprint() const;
-
-	/** @return The current editor mode (editing live actors or editing blueprints) */
-	EComponentEditorMode::Type GetEditorMode() const { return EditorMode; }
-
-	/** Try to handle a drag-drop operation */
-	UE_API FReply TryHandleAssetDragDropOperation(const FDragDropEvent& DragDropEvent);
-
-	/** Handler for recursively expanding/collapsing items */
-	UE_API void SetItemExpansionRecursive(FSCSEditorTreeNodePtrType Model, bool bInExpansionState);
-
-	/** Callback for the action trees to get the filter text */
-	UE_API FText GetFilterText() const;
-
-	/** Called at the end of each frame. */
-	UE_API void OnPostTick(float);
-
-	/** Return the button widgets that can add components or create/edit blueprints */
-	UE_API TSharedPtr<SWidget> GetToolButtonsBox();
-
-	/** Sets UI customizations of this SCSEditor. */
-	UE_API void SetUICustomization(TSharedPtr<ISCSEditorUICustomization> InUICustomization);
-
-protected:
-	UE_API FSCSEditorTreeNodePtrType FindOrCreateParentForExistingComponent(UActorComponent* InActorComponent, FSCSEditorActorNodePtrType ActorRootNode);
-	UE_API FSCSEditorTreeNodePtrType FindParentForNewComponent(UActorComponent* NewComponent) const;
-	UE_API FSCSEditorTreeNodePtrType FindParentForNewNode(USCS_Node* NewNode) const;
-
-	/** Add a component from the selection in the combo box */
-	UE_API UActorComponent* PerformComboAddClass(TSubclassOf<UActorComponent> ComponentClass, EComponentCreateAction::Type ComponentCreateAction, UObject* AssetOverride);
-
-	/** Called to display context menu when right clicking on the widget */
-	UE_API TSharedPtr< SWidget > CreateContextMenu();
-
-	/** Registers context menu by name for later access */
-	UE_API void RegisterContextMenu();
-
-	/** Populate context menu on the fly */
-	UE_API void PopulateContextMenu(UToolMenu* InMenu);
-
-	/** Called when the level editor requests a component to be renamed. */
-	UE_API void OnLevelComponentRequestRename(const UActorComponent* InComponent);
-
-	/** Checks to see if renaming is allowed on the selected component */
-	UE_API bool CanRenameComponent() const;
-
-	/**
-	 * Requests a rename on the selected component just after creation so that the user can provide the initial
-	 * component name (overwriting the default generated one), which is considered part of the creation process.
-	 * @param OngoingCreateTransaction The ongoing transaction started when the component was created.
-	 */
-	UE_API void OnRenameComponent(TUniquePtr<FScopedTransaction> OngoingCreateTransaction);
-
-	/**
-	 * Requests a rename on the selected component.
-	 */
-	UE_API void OnRenameComponent();
-
-	/** Called when component objects are replaced following construction script execution */
-	UE_API void OnObjectsReplaced(const TMap<UObject*, UObject*>& OldToNewInstanceMap);
-
-	/** Helper method to update component pointers held by the given actor node's subtree */
-	UE_API void ReplaceComponentReferencesInTree(FSCSEditorActorNodePtrType InActorNode, const TMap<UObject*, UObject*>& OldToNewInstanceMap, bool& OutHasChanges);
-
-	/** Update component pointers held by tree nodes if components have been replaced following construction script execution */
-	UE_API void ReplaceComponentReferencesInTree(const TArray<FSCSEditorTreeNodePtrType>& Nodes, const TMap<UObject*, UObject*>& OldToNewInstanceMap, bool& OutHasChanges);
-
-	/**
-	 * Function to create events for the current selection
-	 * @param Blueprint						the active blueprint context
-	 * @param EventName						the event to add
-	 * @param GetSelectedObjectsDelegate	the delegate to gather information about current selection
-	 * @param NodeIndex						an index to a specified node to add event for or < 0 for all selected nodes.
-	 */
-	static UE_API void CreateEventsForSelection(UBlueprint* Blueprint, FName EventName, FGetSelectedObjectsDelegate GetSelectedObjectsDelegate);
-
-	/**
-	 * Function to construct an event for a node
-	 * @param Blueprint						the nodes blueprint
-	 * @param EventName						the event to add
-	 * @param EventData						the event data structure describing the node
-	 */
-	static UE_API void ConstructEvent(UBlueprint* Blueprint, const FName EventName, const FComponentEventConstructionData EventData);
-
-	/**
-	 * Function to view an event for a node
-	 * @param Blueprint						the nodes blueprint
-	 * @param EventName						the event to view
-	 * @param EventData						the event data structure describing the node
-	 */
-	static UE_API void ViewEvent(UBlueprint* Blueprint, const FName EventName, const FComponentEventConstructionData EventData);
-
-	/** Helper method to add a tree node for the given SCS node */
-	UE_API FSCSEditorTreeNodePtrType AddTreeNode(USCS_Node* InSCSNode, FSCSEditorTreeNodePtrType InParentNodePtr, const bool bIsInheritedSCS);
-
-	/** Helper method to add a tree node for the given actor component */
-	UE_API FSCSEditorTreeNodePtrType AddTreeNodeFromComponent(UActorComponent* InSceneComponent, FSCSEditorTreeNodePtrType InParentTreeNode = FSCSEditorTreeNodePtrType());
-	
-	/** Helper method to add a tree node for the given node's child actor node, if present */
-	UE_API FSCSEditorTreeNodePtrType AddTreeNodeFromChildActor(FSCSEditorTreeNodePtrType InNodePtr);
-	
-	/** Helper method to recursively find a tree node for the given SCS node starting at the given tree node */
-	UE_API FSCSEditorTreeNodePtrType FindTreeNode(const USCS_Node* InSCSNode, FSCSEditorTreeNodePtrType InStartNodePtr = FSCSEditorTreeNodePtrType()) const;
-
-	/** Helper method to recursively find a tree node for the given scene component starting at the given tree node */
-	UE_API FSCSEditorTreeNodePtrType FindTreeNode(const UActorComponent* InComponent, FSCSEditorTreeNodePtrType InStartNodePtr = FSCSEditorTreeNodePtrType()) const;
-
-	/** Helper method to recursively find a tree node for the given variable or instance name starting at the given tree node */
-	UE_API FSCSEditorTreeNodePtrType FindTreeNode(const FName& InVariableOrInstanceName, FSCSEditorTreeNodePtrType InStartNodePtr = FSCSEditorTreeNodePtrType()) const;
-
-	/** Callback when a component item is scrolled into view */
-	UE_API void OnItemScrolledIntoView( FSCSEditorTreeNodePtrType InItem, const TSharedPtr<ITableRow>& InWidget);
-
-	/** Callback when a component item is double clicked. */
-	UE_API void HandleItemDoubleClicked(FSCSEditorTreeNodePtrType InItem);
-
-	/** Recursively visits the given node + its children and invokes the given function for each. */
-	UE_API void DepthFirstTraversal(const FSCSEditorTreeNodePtrType& InNodePtr, TSet<FSCSEditorTreeNodePtrType>& OutVisitedNodes, const TFunctionRef<void(const FSCSEditorTreeNodePtrType&)> InFunction) const;
-
-	/** Returns the set of expandable nodes that are currently collapsed in the UI */
-	UE_API void GetCollapsedNodes(const FSCSEditorTreeNodePtrType& InNodePtr, TSet<FSCSEditorTreeNodePtrType>& OutCollapsedNodes) const;
-
-	/** @return The visibility of the promote to blueprint button (only visible with an actor instance that is not created from a blueprint)*/
-	UE_API EVisibility GetPromoteToBlueprintButtonVisibility() const;
-
-	/** @return The visibility of the Edit Blueprint button (only visible with an actor instance that is created from a blueprint)*/
-	UE_API EVisibility GetEditBlueprintButtonVisibility() const;
-
-	/** @return The visibility of the Add Component combo button */
-	UE_API EVisibility GetComponentClassComboButtonVisibility() const;
-
-	/** @return The visibility of the components tree */
-	UE_API EVisibility GetComponentsTreeVisibility() const;
-
-	/** @return The visibility of the components filter box */
-	UE_API EVisibility GetComponentsFilterBoxVisibility() const;
-
-	/** @return the tooltip describing how many properties will be applied to the blueprint */
-	UE_API FText OnGetApplyChangesToBlueprintTooltip() const;
-
-	/** @return the tooltip describing how many properties will be reset to the blueprint default*/
-	UE_API FText OnGetResetToBlueprintDefaultsTooltip() const;
-
-	/** Opens the blueprint editor for the blueprint being viewed by the scseditor */
-	UE_API void OnOpenBlueprintEditor(bool bForceCodeEditing) const;
-
-	/** Propagates instance changes to the blueprint */
-	UE_API void OnApplyChangesToBlueprint() const;
-
-	/** Resets instance changes to the blueprint default */
-	UE_API void OnResetToBlueprintDefaults();
-
-	/** Converts the current actor instance to a blueprint */
-	UE_API void PromoteToBlueprint() const;
-
-	/** Called when the promote to blueprint button is clicked */
-	UE_API FReply OnPromoteToBlueprintClicked();
-
-	/** gets a root nodes of the tree */
-	UE_API const TArray<FSCSEditorTreeNodePtrType>& GetRootNodes() const;
-
-	/**
-	 * Creates a new C++ component from the specified class type
-	 * The user will be prompted to pick a new subclass name and code will be recompiled
-	 *
-	 * @return The new class that was created
-	 */
-	UE_API UClass* CreateNewCPPComponent(TSubclassOf<UActorComponent> ComponentClass);
-	
-	/**
-	 * Creates a new Blueprint component from the specified class type
-	 * The user will be prompted to pick a new subclass name and a blueprint asset will be created
-	 *
-	 * @return The new class that was created
-	 */
-	UE_API UClass* CreateNewBPComponent(TSubclassOf<UActorComponent> ComponentClass);
-
-	/** Recursively updates the filtered state for each component item */
-	UE_API void OnFilterTextChanged(const FText& InFilterText);
-
-	/** 
-	 * Compares the filter bar's text with the item's component name. Use 
-	 * bRecursive to refresh the state of child nodes as well. Returns true if 
-	 * the node is set to be filtered out 
-	 */
-	UE_API bool RefreshFilteredState(FSCSEditorTreeNodePtrType TreeNode, bool bRecursive);
-
-	/** Helper method to construct the subtree for the given actor (root) node. */
-	UE_API void BuildSubTreeForActorNode(FSCSEditorActorNodePtrType InActorNode);
-
-	/** @return Type of component to filter the tree view with or nullptr if there's no filter. */
-	UE_API TSubclassOf<UActorComponent> GetComponentTypeFilterToApply() const;
-
-public:
-	/** Tree widget */
-	TSharedPtr<SSCSTreeType> SCSTreeWidget;
-
-	/** Command list for handling actions in the SSCSEditor */
-	TSharedPtr< FUICommandList > CommandList;
-
-	/** Name of a node that has been requested to be renamed */
-	FName DeferredRenameRequest;
-
-	/** Scope the creation of a component which ends when the initial component 'name' is given/accepted by the user, which can be several frames after the component was actually created. */
-	TUniquePtr<FScopedTransaction> DeferredOngoingCreateTransaction;
-
-	/** Used to unregister from the post tick event. */
-	FDelegateHandle PostTickHandle;
-
-	/** Attribute that provides access to the Actor context for which we are viewing/editing the SCS. */
-	TAttribute<class AActor*> ActorContext;
-
-	/** Attribute that provides access to a "preview" Actor context (may not be same as the Actor context that's being edited. */
-	TAttribute<class AActor*> PreviewActor;
-
-	/** Attribute to indicate whether or not editing is allowed. */
-	TAttribute<bool> AllowEditing;
-
-	/** Attribute to indicate whether or not the "Add Component" button is visible. If true, new components cannot be added to the Blueprint. */
-	TAttribute<bool> HideComponentClassCombo;
-
-	/** Attribute to limit visible nodes to a particular component type when filtering the tree view. */
-	TAttribute<TSubclassOf<UActorComponent>> ComponentTypeFilter;
-
-	/** Delegate to invoke on selection update. */
-	FOnSelectionUpdated OnSelectionUpdated;
-
-	/** Delegate to invoke when an item in the tree is double clicked. */
-	FOnItemDoubleClicked OnItemDoubleClicked;
-
-	/** Delegate to invoke when the given property should be highlighted in the details view (e.g. diff). */
-	FOnHighlightPropertyInDetailsView OnHighlightPropertyInDetailsView;
-
-	/** Delegate to invoke when objects within the SCS tree are replaced (eg, via re-instancing from a BP compile) */
-	FSimpleDelegate OnObjectReplaced;
-
-	/** Returns the Actor context for which we are viewing/editing the SCS.  Can return null.  Should not be cached as it may change from frame to frame. */
-	UE_API class AActor* GetActorContext() const;
-
-private:
-	/** Indicates which editor mode we're in. */
-	EComponentEditorMode::Type EditorMode;
-
-	/** Root set of tree */
-	TArray<FSCSEditorTreeNodePtrType> RootNodes;
-
-	/* Root Tree Node*/
-	TSharedPtr<FExtender> ActorMenuExtender;
-
-	/** Flag to enable/disable component editing */
-	bool bEnableComponentEditing;
-
-	/** Gate to prevent changing the selection while selection change is being broadcast. */
-	bool bUpdatingSelection;
-
-	/** Controls whether or not to allow calls to UpdateTree() */
-	bool bAllowTreeUpdates;
-
-	/** TRUE if this SCSEditor is currently the target of a diff */
-	bool bIsDiffing;
-
-	/** The filter box that handles filtering for the tree. */
-	TSharedPtr< SSearchBox > FilterBox;
-
-	/** The tools buttons box **/
-	TSharedPtr<SHorizontalBox> ButtonBox;
-
-	/** SCSEditor UI customizations */
-	TSharedPtr<ISCSEditorUICustomization> UICustomization;
-
-	/** SCSEditor UI extension */
-	TSharedPtr<class SExtensionPanel> ExtensionPanel; };
-
 #undef UE_API
diff --git a/Engine/Source/Editor/KismetWidgets/Public/CreateBlueprintFromActorDialog.h b/Engine/Source/Editor/KismetWidgets/Public/CreateBlueprintFromActorDialog.h
--- a/Engine/Source/Editor/KismetWidgets/Public/CreateBlueprintFromActorDialog.h
+++ b/Engine/Source/Editor/KismetWidgets/Public/CreateBlueprintFromActorDialog.h
@@ -36,13 +36,6 @@ public:
 	 */
 	static KISMETWIDGETS_API void OpenDialog(ECreateBlueprintFromActorMode CreateMode, AActor* InActorOverride = nullptr, bool bInReplaceActors = true);
 
-	UE_DEPRECATED(4.25, "Use version of OpenDialog that takes the CreateMode enum")
-	static void OpenDialog(bool bInHarvest, AActor* InActorOverride = nullptr)
-	{
-		OpenDialog(bInHarvest ? ECreateBlueprintFromActorMode::Harvest : ECreateBlueprintFromActorMode::Subclass, InActorOverride);
-	}
-
-
 	/** 
 	 * Static function that returns which create modes are valid given the current selection set
 	 */
diff --git a/Engine/Source/Editor/UMGEditor/Private/WidgetBlueprint.cpp b/Engine/Source/Editor/UMGEditor/Private/WidgetBlueprint.cpp
--- a/Engine/Source/Editor/UMGEditor/Private/WidgetBlueprint.cpp
+++ b/Engine/Source/Editor/UMGEditor/Private/WidgetBlueprint.cpp
@@ -51,9 +51,6 @@
 
 #define LOCTEXT_NAMESPACE "UMG"
 
-PRAGMA_DISABLE_DEPRECATION_WARNINGS;
-FWidgetBlueprintDelegates::FGetAssetTags FWidgetBlueprintDelegates::GetAssetTags;
-PRAGMA_ENABLE_DEPRECATION_WARNINGS;
 FWidgetBlueprintDelegates::FGetAssetTagsWithContext FWidgetBlueprintDelegates::GetAssetTagsWithContext;
 
 FEditorPropertyPathSegment::FEditorPropertyPathSegment()
@@ -716,26 +713,11 @@ void UWidgetBlueprint::PreSave(FObjectPreSaveContext ObjectSaveContext)
 	}
 }
 
-void UWidgetBlueprint::GetAssetRegistryTags(TArray<FAssetRegistryTag>& OutTags) const
-{
-	PRAGMA_DISABLE_DEPRECATION_WARNINGS;
-	Super::GetAssetRegistryTags(OutTags);
-	PRAGMA_ENABLE_DEPRECATION_WARNINGS;
-}
 
 void UWidgetBlueprint::GetAssetRegistryTags(FAssetRegistryTagsContext Context) const
 {
 	Super::GetAssetRegistryTags(Context);
 
-	PRAGMA_DISABLE_DEPRECATION_WARNINGS;
-	TArray<UObject::FAssetRegistryTag> DeprecatedFunctionTags;
-	FWidgetBlueprintDelegates::GetAssetTags.Broadcast(this, DeprecatedFunctionTags);
-	for (UObject::FAssetRegistryTag& Tag : DeprecatedFunctionTags)
-	{
-		Context.AddTag(MoveTemp(Tag));
-	}
-	PRAGMA_ENABLE_DEPRECATION_WARNINGS;
-
 	// Add AvailableNamedSlots,  also available on generated class, to the WidgetBlueprint
 	if (const UWidgetBlueprintGeneratedClass* WidgetBPGeneratedClass = Cast<const UWidgetBlueprintGeneratedClass>(GeneratedClass))
 	{
diff --git a/Engine/Source/Editor/UMGEditor/Public/WidgetBlueprint.h b/Engine/Source/Editor/UMGEditor/Public/WidgetBlueprint.h
--- a/Engine/Source/Editor/UMGEditor/Public/WidgetBlueprint.h
+++ b/Engine/Source/Editor/UMGEditor/Public/WidgetBlueprint.h
@@ -40,8 +40,6 @@ public:
 
 	// called by UWdgetBlueprint::GetAssetRegistryTags()
 	static UE_API FGetAssetTagsWithContext GetAssetTagsWithContext;
-	UE_DEPRECATED(5.4, "Subscribe to GetAssetTagsWithContext instead.")
-	static UE_API FGetAssetTags GetAssetTags;
 };
 
 
@@ -273,8 +271,6 @@ public:
 
 #if WITH_EDITOR
 	UE_API virtual void GetAssetRegistryTags(FAssetRegistryTagsContext Context) const override;
-	UE_DEPRECATED(5.4, "Implement the version that takes FAssetRegistryTagsContext instead.")
-	UE_API virtual void GetAssetRegistryTags(TArray<FAssetRegistryTag>& OutTags) const override;
 	UE_API virtual void NotifyGraphRenamed(class UEdGraph* Graph, FName OldName, FName NewName) override;
 	UE_API virtual EDataValidationResult IsDataValid(class FDataValidationContext& Context) const override;
 	UE_API bool DetectSlateWidgetLeaks(class FDataValidationContext& Context) const;
diff --git a/Engine/Source/Editor/UnrealEd/Private/Kismet2/ChildActorComponentEditorUtils.cpp b/Engine/Source/Editor/UnrealEd/Private/Kismet2/ChildActorComponentEditorUtils.cpp
--- a/Engine/Source/Editor/UnrealEd/Private/Kismet2/ChildActorComponentEditorUtils.cpp
+++ b/Engine/Source/Editor/UnrealEd/Private/Kismet2/ChildActorComponentEditorUtils.cpp
@@ -4,7 +4,6 @@
 #include "Settings/BlueprintEditorProjectSettings.h"
 #include "ToolMenus.h"
 #include "SSCSEditor.h"
-#include "SSCSEditorMenuContext.h"
 #include "SubobjectEditorMenuContext.h"
 #include "SSubobjectEditor.h"
 
diff --git a/Engine/Source/Editor/UnrealEd/Private/Kismet2/KismetDebugUtilities.cpp b/Engine/Source/Editor/UnrealEd/Private/Kismet2/KismetDebugUtilities.cpp
--- a/Engine/Source/Editor/UnrealEd/Private/Kismet2/KismetDebugUtilities.cpp
+++ b/Engine/Source/Editor/UnrealEd/Private/Kismet2/KismetDebugUtilities.cpp
@@ -20,7 +20,6 @@
 #include "Editor/UnrealEdEngine.h"
 #include "Settings/EditorExperimentalSettings.h"
 #include "CallStackViewer.h"
-#include "WatchPointViewer.h"
 #include "Animation/AnimBlueprintGeneratedClass.h"
 #include "UnrealEdGlobals.h"
 #include "Kismet2/Breakpoint.h"
diff --git a/Engine/Source/Runtime/Engine/Classes/Engine/Blueprint.h b/Engine/Source/Runtime/Engine/Classes/Engine/Blueprint.h
--- a/Engine/Source/Runtime/Engine/Classes/Engine/Blueprint.h
+++ b/Engine/Source/Runtime/Engine/Classes/Engine/Blueprint.h
@@ -638,11 +638,12 @@ public:
 	UPROPERTY(transient)
 	TMap<FName, FName> OldToNewComponentTemplateNames;
 
-	/** Array of extensions for this blueprint */
-	UE_DEPRECATED(5.1, "Please do not access this member directly; Instead use: UBlueprint::GetExtensions / UBlueprint::AddExtension / UBlueprint::RemoveExtension[At].")
+private:
+	/** Array of extensions for this blueprint. To access use: UBlueprint::GetExtensions / UBlueprint::AddExtension / UBlueprint::RemoveExtension[At] */
 	UPROPERTY()
 	TArray<TObjectPtr<UBlueprintExtension>> Extensions;
 
+public:
 	/** Fires whenever BP extension added */
 	FOnExtensionAdded OnExtensionAdded;
 
@@ -911,8 +912,6 @@ public:
 #endif
 	ENGINE_API virtual bool Modify(bool bAlwaysMarkDirty = true) override;
 	ENGINE_API virtual void GetAssetRegistryTags(FAssetRegistryTagsContext Context) const override;
-	UE_DEPRECATED(5.4, "Implement the version that takes FAssetRegistryTagsContext instead.")
-	ENGINE_API virtual void GetAssetRegistryTags(TArray<FAssetRegistryTag>& OutTags) const override;
 	ENGINE_API virtual void GetAssetRegistryTagMetadata(TMap<FName, FAssetRegistryTagMetadata>& OutMetadata) const override;
 protected:
 	ENGINE_API virtual void ThreadedPostLoadAssetRegistryTagsOverride(FPostLoadAssetRegistryTagsContext& Context) const override;
diff --git a/Engine/Source/Runtime/Engine/Classes/Engine/BlueprintGeneratedClass.h b/Engine/Source/Runtime/Engine/Classes/Engine/BlueprintGeneratedClass.h
--- a/Engine/Source/Runtime/Engine/Classes/Engine/BlueprintGeneratedClass.h
+++ b/Engine/Source/Runtime/Engine/Classes/Engine/BlueprintGeneratedClass.h
@@ -592,8 +592,6 @@ public:
 	ENGINE_API virtual void PostInitProperties() override;
 	ENGINE_API virtual void GetPreloadDependencies(TArray<UObject*>& OutDeps) override;
 	ENGINE_API virtual void GetAssetRegistryTags(FAssetRegistryTagsContext Context) const override;
-	UE_DEPRECATED(5.4, "Implement the version that takes FAssetRegistryTagsContext instead.")
-	ENGINE_API virtual void GetAssetRegistryTags(TArray<FAssetRegistryTag>& OutTags) const override;
 #if WITH_EDITOR
 	ENGINE_API virtual void GetAdditionalAssetDataObjectsForCook(FArchiveCookContext& CookContext, 
 		TArray<UObject*>& OutObjects) const override;
diff --git a/Engine/Source/Runtime/Engine/Classes/Engine/CurveTable.h b/Engine/Source/Runtime/Engine/Classes/Engine/CurveTable.h
--- a/Engine/Source/Runtime/Engine/Classes/Engine/CurveTable.h
+++ b/Engine/Source/Runtime/Engine/Classes/Engine/CurveTable.h
@@ -84,8 +84,6 @@ public:
 
 #if WITH_EDITORONLY_DATA
 	virtual void GetAssetRegistryTags(FAssetRegistryTagsContext Context) const override;
-	UE_DEPRECATED(5.4, "Implement the version that takes FAssetRegistryTagsContext instead.")
-	virtual void GetAssetRegistryTags(TArray<FAssetRegistryTag>& OutTags) const override;
 	virtual void PostInitProperties() override;
 	virtual void PostLoad() override;
 
diff --git a/Engine/Source/Runtime/Engine/Classes/Engine/DataTable.h b/Engine/Source/Runtime/Engine/Classes/Engine/DataTable.h
--- a/Engine/Source/Runtime/Engine/Classes/Engine/DataTable.h
+++ b/Engine/Source/Runtime/Engine/Classes/Engine/DataTable.h
@@ -151,8 +151,6 @@ public:
 #if WITH_EDITORONLY_DATA
 	ENGINE_API FTopLevelAssetPath GetRowStructPathName() const;
 	ENGINE_API virtual void GetAssetRegistryTags(FAssetRegistryTagsContext Context) const override;
-	UE_DEPRECATED(5.4, "Implement the version that takes FAssetRegistryTagsContext instead.")
-	ENGINE_API virtual void GetAssetRegistryTags(TArray<FAssetRegistryTag>& OutTags) const override;
 	ENGINE_API virtual void PostInitProperties() override;
 	ENGINE_API virtual void PostLoad() override;
 	//~ End UObject Interface
diff --git a/Engine/Source/Runtime/Engine/Classes/Engine/Engine.h b/Engine/Source/Runtime/Engine/Classes/Engine/Engine.h
--- a/Engine/Source/Runtime/Engine/Classes/Engine/Engine.h
+++ b/Engine/Source/Runtime/Engine/Classes/Engine/Engine.h
@@ -3149,8 +3149,6 @@ public:
 		bool bSkipCompilerGeneratedDefaults;
 		bool bNotifyObjectReplacement;
 		bool bClearReferences;
-		UE_DEPRECATED(5.4, "This isn't used anymore by the code.")
-		bool bDontClearReferenceIfNewerClassExists;
 		bool bReplaceInternalReferenceUponRead; // While reading back object ptr, immediately replace them if they are in the replacement map.
 
 		// In cases where the SourceObject will no longer be able to look up its correct Archetype, it can be supplied
@@ -3585,24 +3583,6 @@ public:
 	 */
 	ENGINE_API virtual void CheckAndHandleStaleWorldObjectReferences(FWorldContext* InWorldContext = nullptr);
 
-	/**
-	 * Attempts to find what is referencing a world that should have been garbage collected
-	 * @param ObjectToFindReferencesTo World or its package (or any object from the world package that should've been destroyed)
-	 * @param Verbosity Verbosity (can be fatal or non-fatal) with which to print the error message with
-	 */
-	UE_DEPRECATED(5.1, "Please use FReferenceChainSearch::FindAndPrintStaleReferencesToObject")
-	static ENGINE_API void FindAndPrintStaleReferencesToObject(UObject* ObjectToFindReferencesTo, ELogVerbosity::Type Verbosity);
-
-	/**
-	 * Attempts to find a reference chain leading to a world that should have been garbage collected
-	 * @param ObjectToFindReferencesTo World or its package (or any object from the world package that should've been destroyed)
-	 * @param Options Determines how the stale references messages should be logged
-	 */
-	UE_DEPRECATED(5.3, "Please use FReferenceChainSearch::FindAndPrintStaleReferencesToObject")
-	static ENGINE_API FString FindAndPrintStaleReferencesToObject(UObject* ObjectToFindReferencesTo, EPrintStaleReferencesOptions Options);
-	UE_DEPRECATED(5.3, "Please use FReferenceChainSearch::FindAndPrintStaleReferencesToObjects")
-	static ENGINE_API TArray<FString> FindAndPrintStaleReferencesToObjects(TConstArrayView<UObject*> ObjectsToFindReferencesTo, EPrintStaleReferencesOptions Options);
-
 	ENGINE_API FWorldContext& CreateNewWorldContext(EWorldType::Type WorldType);
 
 	ENGINE_API virtual void DestroyWorldContext(UWorld * InWorld);
diff --git a/Engine/Source/Runtime/Engine/Classes/Engine/World.h b/Engine/Source/Runtime/Engine/Classes/Engine/World.h
--- a/Engine/Source/Runtime/Engine/Classes/Engine/World.h
+++ b/Engine/Source/Runtime/Engine/Classes/Engine/World.h
@@ -1314,10 +1314,11 @@ public:
 	/** Whether actors have been initialized for play */
 	uint8 bActorsInitialized:1;
 
-	/** Whether BeginPlay has been called on actors */
-	UE_DEPRECATED(5.4, "Public access to bBegunPlay is deprecated. Please update your code to use the public accessors GetBegunPlay() & SetBegunPlay().")
+private:
+	/** Whether BeginPlay has been called on actors. Use the public accessors GetBegunPlay() & SetBegunPlay() */
 	uint8 bBegunPlay:1;
 
+public:
 	/** Sets whether BeginPlay has been called for actors in the world. Use BeginPlay and EndPlay below to start the process of changing this. */
 	UE_API void SetBegunPlay(bool bHasBegunPlay);
 
@@ -3068,8 +3069,6 @@ public:
 	UE_API virtual EDataValidationResult IsDataValid(FDataValidationContext& Context) const override;
 	UE_API virtual bool Rename(const TCHAR* NewName = NULL, UObject* NewOuter = NULL, ERenameFlags Flags = REN_None) override;
 	UE_API virtual void GetAssetRegistryTags(FAssetRegistryTagsContext Context) const override;
-	UE_DEPRECATED(5.4, "Implement the version that takes FAssetRegistryTagsContext instead.")
-	UE_API virtual void GetAssetRegistryTags(TArray<FAssetRegistryTag>& OutTags) const override;
 	UE_API virtual void GetAssetRegistryTagMetadata(TMap<FName, FAssetRegistryTagMetadata>& OutMetadata) const override;
 	UE_API virtual void GetExtendedAssetRegistryTagsForSave(const ITargetPlatform* TargetPlatform, TArray<FAssetRegistryTag>& OutTags) const override;
 protected:
@@ -3271,8 +3270,7 @@ public:
 	 * operations that need to call InitWorld during the Load of the World's package should break that rule and not call CleanupWorld.
 	 */
 	static UE_API const FName KeepInitializedDuringLoadTag;
-	UE_DEPRECATED(5.2, "Call IsInitialized instead.")
-	bool IsInitializedAndNeedsCleanup() const { return bIsWorldInitialized; }
+
 	/** Returns whether InitWorld has ever been called since this World was created.  */
 	bool HasEverBeenInitialized() const { return bHasEverBeenInitialized; }
 
@@ -4540,8 +4538,6 @@ public:
 
 	// called by UWorld::GetAssetRegistryTags()
 	static UE_API FWorldGetAssetTagsWithContext GetAssetTagsWithContext;
-	UE_DEPRECATED(5.4, "Subscribe to GetAssetTagsWithContext instead")
-	static UE_API FWorldGetAssetTags GetAssetTags;
 
 #if WITH_EDITOR
 	// Delegate called when levelscript actions need refreshing
@@ -4704,8 +4700,6 @@ inline bool UWorld::IsNetMode(ENetMode Mode) const
 #endif
 }
 
-UE_DEPRECATED(5.0, "Please use LexToString(EWorldType::Type Type) instead")
-FString ENGINE_API ToString(EWorldType::Type Type);
 FString ENGINE_API ToString(ENetMode NetMode);
 
 #undef UE_API
diff --git a/Engine/Source/Runtime/Engine/Classes/GameFramework/Actor.h b/Engine/Source/Runtime/Engine/Classes/GameFramework/Actor.h
--- a/Engine/Source/Runtime/Engine/Classes/GameFramework/Actor.h
+++ b/Engine/Source/Runtime/Engine/Classes/GameFramework/Actor.h
@@ -1455,9 +1455,6 @@ public:
 	//~=============================================================================
 	// DataLayers functions.
 #if WITH_EDITOR
-protected:
-	UE_DEPRECATED(5.4, "Use ActorTypeSupportsDataLayer instead.")
-	virtual bool IsDataLayerTypeSupported(TSubclassOf<UDataLayerInstance> DataLayerType) const final { return false; }
 private:
 	virtual bool ActorTypeIsMainWorldOnly() const { return false; }
 	ENGINE_API TArray<const UDataLayerAsset*> ResolveDataLayerAssets(const TArray<TSoftObjectPtr<UDataLayerAsset>>& InDataLayerAssets) const;
@@ -2347,9 +2344,6 @@ public:
 	static ENGINE_API TDelegate<FRemoteServerId(const AActor*)> GetPreferredServerDelegate;
 #endif
 	ENGINE_API virtual void BeginDestroy() override;
-#if !UE_STRIP_DEPRECATED_PROPERTIES
-	ENGINE_API virtual bool IsReadyForFinishDestroy() override;
-#endif
 	ENGINE_API virtual bool Rename( const TCHAR* NewName=nullptr, UObject* NewOuter=nullptr, ERenameFlags Flags=REN_None ) override;
 	ENGINE_API virtual void PostRename( UObject* OldOuter, const FName OldName ) override;
 	ENGINE_API virtual bool CanBeInCluster() const override;
@@ -4403,12 +4397,6 @@ public:
 	static ENGINE_API FOnProcessEvent ProcessEventDelegate;
 #endif
 
-#if !UE_STRIP_DEPRECATED_PROPERTIES
-	/** A fence to track when the primitive is detached from the scene in the rendering thread. */
-	UE_DEPRECATED(5.1, "AActor::DetachFence has been deprecated. If you are relying on it for render thread synchronization in a subclass of actor, add your own fence to that class instead.")
-	FRenderCommandFence DetachFence;
-#endif
-
 private:
 	/** Helper that already assumes the Hit info is reversed, and avoids creating a temp FHitResult if possible. */
 	ENGINE_API void InternalDispatchBlockingHit(UPrimitiveComponent* MyComp, UPrimitiveComponent* OtherComp, bool bSelfMoved, FHitResult const& Hit);
diff --git a/Engine/Source/Runtime/Engine/Classes/GameFramework/GameModeBase.h b/Engine/Source/Runtime/Engine/Classes/GameFramework/GameModeBase.h
--- a/Engine/Source/Runtime/Engine/Classes/GameFramework/GameModeBase.h
+++ b/Engine/Source/Runtime/Engine/Classes/GameFramework/GameModeBase.h
@@ -356,12 +356,6 @@ public:
 	 */
 	ENGINE_API virtual APlayerController* SpawnPlayerController(ENetRole InRemoteRole, const FString& Options);
 
-	UE_DEPRECATED(4.20, "SpawnPlayerController with Location and Rotation is deprecated, call or override the version that takes an Option string instead")
-	ENGINE_API virtual APlayerController* SpawnPlayerController(ENetRole InRemoteRole, FVector const& SpawnLocation, FRotator const& SpawnRotation);
-	UE_DEPRECATED(4.20, "SpawnReplayPlayerController is deprecated, replay controller spawning is handled inside the new version of the SpawnPlayerController function")
-	ENGINE_API virtual APlayerController* SpawnReplayPlayerController(ENetRole InRemoteRole, FVector const& SpawnLocation, FRotator const& SpawnRotation);
-
-
 	/** Signals that a player is ready to enter the game, which may start it up */
 	UFUNCTION(BlueprintNativeEvent, Category=Game)
 	ENGINE_API void HandleStartingNewPlayer(APlayerController* NewPlayer);
diff --git a/Engine/Source/Runtime/Engine/Classes/GameFramework/PlayerController.h b/Engine/Source/Runtime/Engine/Classes/GameFramework/PlayerController.h
--- a/Engine/Source/Runtime/Engine/Classes/GameFramework/PlayerController.h
+++ b/Engine/Source/Runtime/Engine/Classes/GameFramework/PlayerController.h
@@ -2122,12 +2122,6 @@ public:
 	/** returns whether the client has completely loaded the server's current world (valid on server only) */
 	ENGINE_API bool HasClientLoadedCurrentWorld();
 
-	/** forces a full replication check of the specified Actor on only the client that owns this PlayerController
-	 * this function has no effect if this PC is not a remote client or if the Actor is not relevant to that client
-	 */
-	UE_DEPRECATED(5.3, "Deprecated in favor of using ForceNetUpdate on the target actor instead.")
-	ENGINE_API void ForceSingleNetUpdateFor(class AActor* Target);
-
 	/** Set the view target
 	 * @param A - new actor to set as view target
 	 * @param TransitionParams - parameters to use for controlling the transition
diff --git a/Engine/Source/Runtime/Engine/Private/Actor.cpp b/Engine/Source/Runtime/Engine/Private/Actor.cpp
--- a/Engine/Source/Runtime/Engine/Private/Actor.cpp
+++ b/Engine/Source/Runtime/Engine/Private/Actor.cpp
@@ -922,15 +922,6 @@ void AActor::BeginDestroy()
 	Super::BeginDestroy();
 }
 
-#if !UE_STRIP_DEPRECATED_PROPERTIES
-bool AActor::IsReadyForFinishDestroy()
-{
-PRAGMA_DISABLE_DEPRECATION_WARNINGS
-	return Super::IsReadyForFinishDestroy() && DetachFence.IsFenceComplete();
-PRAGMA_ENABLE_DEPRECATION_WARNINGS
-}
-#endif
-
 void AActor::Serialize(FArchive& Ar)
 {
 	Ar.UsingCustomVersion(FUE5MainStreamObjectVersion::GUID);
diff --git a/Engine/Source/Runtime/Engine/Private/Blueprint.cpp b/Engine/Source/Runtime/Engine/Private/Blueprint.cpp
--- a/Engine/Source/Runtime/Engine/Private/Blueprint.cpp
+++ b/Engine/Source/Runtime/Engine/Private/Blueprint.cpp
@@ -1047,13 +1047,6 @@ bool UBlueprint::CanAlwaysRecompileWhilePlayingInEditor() const
 	return false;
 }
 
-void UBlueprint::GetAssetRegistryTags(TArray<FAssetRegistryTag>& OutTags) const
-{
-	PRAGMA_DISABLE_DEPRECATION_WARNINGS;
-	Super::GetAssetRegistryTags(OutTags);
-	PRAGMA_ENABLE_DEPRECATION_WARNINGS;
-}
-
 void UBlueprint::GetAssetRegistryTags(FAssetRegistryTagsContext Context) const
 {
 	// We use Generated instead of Skeleton because the CDO data is more accurate on Generated
diff --git a/Engine/Source/Runtime/Engine/Private/BlueprintGeneratedClass.cpp b/Engine/Source/Runtime/Engine/Private/BlueprintGeneratedClass.cpp
--- a/Engine/Source/Runtime/Engine/Private/BlueprintGeneratedClass.cpp
+++ b/Engine/Source/Runtime/Engine/Private/BlueprintGeneratedClass.cpp
@@ -531,13 +531,6 @@ void UBlueprintGeneratedClass::InitializeFieldNotifies()
 	}
 }
 
-void UBlueprintGeneratedClass::GetAssetRegistryTags(TArray<FAssetRegistryTag>& OutTags) const
-{
-	PRAGMA_DISABLE_DEPRECATION_WARNINGS;
-	Super::GetAssetRegistryTags(OutTags);
-	PRAGMA_ENABLE_DEPRECATION_WARNINGS;
-}
-
 void UBlueprintGeneratedClass::GetAssetRegistryTags(FAssetRegistryTagsContext Context) const
 {
 	Super::GetAssetRegistryTags(Context);
diff --git a/Engine/Source/Runtime/Engine/Private/Components/PrimitiveComponent.cpp b/Engine/Source/Runtime/Engine/Private/Components/PrimitiveComponent.cpp
--- a/Engine/Source/Runtime/Engine/Private/Components/PrimitiveComponent.cpp
+++ b/Engine/Source/Runtime/Engine/Private/Components/PrimitiveComponent.cpp
@@ -1940,16 +1940,6 @@ void UPrimitiveComponent::BeginDestroy()
 
 	// Use a fence to keep track of when the rendering thread executes this scene detachment.
 	DetachFence.BeginFence();
-	
-#if !UE_STRIP_DEPRECATED_PROPERTIES
-	AActor* Owner = GetOwner();
-	if(Owner)
-	{
-PRAGMA_DISABLE_DEPRECATION_WARNINGS
-		Owner->DetachFence.BeginFence();
-PRAGMA_ENABLE_DEPRECATION_WARNINGS
-	}
-#endif
 }
 
 void UPrimitiveComponent::OnComponentDestroyed(bool bDestroyingHierarchy)
diff --git a/Engine/Source/Runtime/Engine/Private/CurveTable.cpp b/Engine/Source/Runtime/Engine/Private/CurveTable.cpp
--- a/Engine/Source/Runtime/Engine/Private/CurveTable.cpp
+++ b/Engine/Source/Runtime/Engine/Private/CurveTable.cpp
@@ -273,13 +273,6 @@ void UCurveTable::FinishDestroy()
 }
 
 #if WITH_EDITORONLY_DATA
-void UCurveTable::GetAssetRegistryTags(TArray<FAssetRegistryTag>& OutTags) const
-{
-	PRAGMA_DISABLE_DEPRECATION_WARNINGS;
-	Super::GetAssetRegistryTags(OutTags);
-	PRAGMA_ENABLE_DEPRECATION_WARNINGS;
-}
-
 void UCurveTable::GetAssetRegistryTags(FAssetRegistryTagsContext Context) const
 {
 	if (AssetImportData)
diff --git a/Engine/Source/Runtime/Engine/Private/DataTable.cpp b/Engine/Source/Runtime/Engine/Private/DataTable.cpp
--- a/Engine/Source/Runtime/Engine/Private/DataTable.cpp
+++ b/Engine/Source/Runtime/Engine/Private/DataTable.cpp
@@ -344,13 +344,6 @@ FTopLevelAssetPath UDataTable::GetRowStructPathName() const
 	return (RowStruct) ? RowStruct->GetStructPathName() : RowStructPathName;
 }
 
-void UDataTable::GetAssetRegistryTags(TArray<FAssetRegistryTag>& OutTags) const
-{
-	PRAGMA_DISABLE_DEPRECATION_WARNINGS;
-	Super::GetAssetRegistryTags(OutTags);
-	PRAGMA_ENABLE_DEPRECATION_WARNINGS;
-}
-
 void UDataTable::GetAssetRegistryTags(FAssetRegistryTagsContext Context) const
 {
 	if (AssetImportData)
diff --git a/Engine/Source/Runtime/Engine/Private/GameModeBase.cpp b/Engine/Source/Runtime/Engine/Private/GameModeBase.cpp
--- a/Engine/Source/Runtime/Engine/Private/GameModeBase.cpp
+++ b/Engine/Source/Runtime/Engine/Private/GameModeBase.cpp
@@ -740,27 +740,12 @@ APlayerController* AGameModeBase::Login(UPlayer* NewPlayer, ENetRole InRemoteRol
 
 APlayerController* AGameModeBase::SpawnPlayerController(ENetRole InRemoteRole, const FString& Options)
 {
-// calling the deprecated functions for backward compatibility, should call SpawnPlayerControllerCommon directly in the future.
-PRAGMA_DISABLE_DEPRECATION_WARNINGS
 	if (Options.Contains(FString("SpectatorOnly=1")) && ReplaySpectatorPlayerControllerClass != nullptr)
 	{
-		return SpawnReplayPlayerController(InRemoteRole, FVector::ZeroVector, FRotator::ZeroRotator);
+		return SpawnPlayerControllerCommon(InRemoteRole, FVector::ZeroVector, FRotator::ZeroRotator, ReplaySpectatorPlayerControllerClass);
 	}
 	
-	return SpawnPlayerController(InRemoteRole, FVector::ZeroVector, FRotator::ZeroRotator);
-PRAGMA_ENABLE_DEPRECATION_WARNINGS
-}
-
-// deprecated
-APlayerController* AGameModeBase::SpawnPlayerController(ENetRole InRemoteRole, FVector const& SpawnLocation, FRotator const& SpawnRotation)
-{
-	return SpawnPlayerControllerCommon(InRemoteRole, SpawnLocation, SpawnRotation, PlayerControllerClass);
-}
-
-// deprecated
-APlayerController* AGameModeBase::SpawnReplayPlayerController(ENetRole InRemoteRole, FVector const& SpawnLocation, FRotator const& SpawnRotation)
-{
-	return SpawnPlayerControllerCommon(InRemoteRole, SpawnLocation, SpawnRotation, ReplaySpectatorPlayerControllerClass);
+	return SpawnPlayerControllerCommon(InRemoteRole, FVector::ZeroVector, FRotator::ZeroRotator, PlayerControllerClass);
 }
 
 APlayerController* AGameModeBase::SpawnPlayerControllerCommon(ENetRole InRemoteRole, FVector const& SpawnLocation, FRotator const& SpawnRotation, TSubclassOf<APlayerController> InPlayerControllerClass)
diff --git a/Engine/Source/Runtime/Engine/Private/PlayerController.cpp b/Engine/Source/Runtime/Engine/Private/PlayerController.cpp
--- a/Engine/Source/Runtime/Engine/Private/PlayerController.cpp
+++ b/Engine/Source/Runtime/Engine/Private/PlayerController.cpp
@@ -747,18 +747,6 @@ bool APlayerController::HasClientLoadedCurrentWorld()
 	}
 }
 
-void APlayerController::ForceSingleNetUpdateFor(AActor* Target)
-{
-	if (Target == NULL)
-	{
-		UE_LOG(LogPlayerController, Warning, TEXT("PlayerController::ForceSingleNetUpdateFor(): No Target specified"));
-	}
-	else
-	{
-		Target->ForceNetUpdate();
-	}
-}
-
 void APlayerController::SmoothTargetViewRotation(APawn* TargetPawn, float DeltaSeconds)
 {
 	BlendedTargetViewRotation = FMath::RInterpTo(BlendedTargetViewRotation, TargetViewRotation, DeltaSeconds, SmoothTargetViewRotationSpeed);
diff --git a/Engine/Source/Runtime/Engine/Private/UnrealEngine.cpp b/Engine/Source/Runtime/Engine/Private/UnrealEngine.cpp
--- a/Engine/Source/Runtime/Engine/Private/UnrealEngine.cpp
+++ b/Engine/Source/Runtime/Engine/Private/UnrealEngine.cpp
@@ -17383,46 +17383,6 @@ void UEngine::CheckAndHandleStaleWorldObjectReferences(FWorldContext* WorldConte
 	}
 }
 
-void UEngine::FindAndPrintStaleReferencesToObject(UObject* ObjectToFindReferencesTo, ELogVerbosity::Type Verbosity)
-{
-	EPrintStaleReferencesOptions Options = EPrintStaleReferencesOptions::None;
-	switch (Verbosity)
-	{
-	case ELogVerbosity::NoLogging:
-		Options = EPrintStaleReferencesOptions::None;
-		break;
-	case ELogVerbosity::Log:
-		Options = EPrintStaleReferencesOptions::Log;
-		break;
-	case ELogVerbosity::Display:
-		Options = EPrintStaleReferencesOptions::Display;
-		break;
-	case ELogVerbosity::Warning:
-		Options = EPrintStaleReferencesOptions::Warning;
-		break;
-	case ELogVerbosity::Error:
-		Options = EPrintStaleReferencesOptions::Error;
-		break;
-	case ELogVerbosity::Fatal:
-		Options = EPrintStaleReferencesOptions::Fatal;
-		break;
-	default:
-		checkf(false, TEXT("Unsupported Verbosity: %d"), Verbosity);
-		break;
-	}
-	FReferenceChainSearch::FindAndPrintStaleReferencesToObject(ObjectToFindReferencesTo, Options);
-}
-
-FString UEngine::FindAndPrintStaleReferencesToObject(UObject* ObjectToFindReferencesTo, EPrintStaleReferencesOptions Options)
-{
-	return FReferenceChainSearch::FindAndPrintStaleReferencesToObject(ObjectToFindReferencesTo, Options);
-}
-
-TArray<FString> UEngine::FindAndPrintStaleReferencesToObjects(TConstArrayView<UObject*> ObjectsToFindReferencesTo, EPrintStaleReferencesOptions Options)
-{
-	return FReferenceChainSearch::FindAndPrintStaleReferencesToObjects(ObjectsToFindReferencesTo, Options);
-}
-
 /*-----------------------------------------------------------------------------
 Async persistent level map change.
 -----------------------------------------------------------------------------*/
@@ -18374,7 +18334,6 @@ UEngine::FCopyPropertiesForUnrelatedObjectsParams::FCopyPropertiesForUnrelatedOb
 	, bSkipCompilerGeneratedDefaults(false)
 	, bNotifyObjectReplacement(false)
 	, bClearReferences(true)
-	, bDontClearReferenceIfNewerClassExists(false)
 	, bReplaceInternalReferenceUponRead(false)
 	, SourceObjectArchetype(nullptr)
 	, OptionalReplacementMappings(nullptr)
diff --git a/Engine/Source/Runtime/Engine/Private/UserDefinedEnum.cpp b/Engine/Source/Runtime/Engine/Private/UserDefinedEnum.cpp
--- a/Engine/Source/Runtime/Engine/Private/UserDefinedEnum.cpp
+++ b/Engine/Source/Runtime/Engine/Private/UserDefinedEnum.cpp
@@ -132,13 +132,6 @@ void UUserDefinedEnum::PostEditChangeProperty(FPropertyChangedEvent& PropertyCha
 	}
 }
 
-void UUserDefinedEnum::GetAssetRegistryTags(TArray<FAssetRegistryTag>& OutTags) const
-{
-	PRAGMA_DISABLE_DEPRECATION_WARNINGS;
-	Super::GetAssetRegistryTags(OutTags);
-	PRAGMA_ENABLE_DEPRECATION_WARNINGS;
-}
-
 void UUserDefinedEnum::GetAssetRegistryTags(FAssetRegistryTagsContext Context) const
 {
 	Super::GetAssetRegistryTags(Context);
diff --git a/Engine/Source/Runtime/Engine/Private/World.cpp b/Engine/Source/Runtime/Engine/Private/World.cpp
--- a/Engine/Source/Runtime/Engine/Private/World.cpp
+++ b/Engine/Source/Runtime/Engine/Private/World.cpp
@@ -806,9 +806,6 @@ FWorldDelegates::FLevelComponentsEvent FWorldDelegates::LevelComponentsCleared;
 FWorldDelegates::FLevelComponentsEvent FWorldDelegates::LevelComponentsUpdated;
 FWorldDelegates::FLevelOffsetEvent FWorldDelegates::PostApplyLevelOffset;
 FWorldDelegates::FLevelTransformEvent FWorldDelegates::PostApplyLevelTransform;
-PRAGMA_DISABLE_DEPRECATION_WARNINGS;
-FWorldDelegates::FWorldGetAssetTags FWorldDelegates::GetAssetTags;
-PRAGMA_ENABLE_DEPRECATION_WARNINGS;
 FWorldDelegates::FWorldGetAssetTagsWithContext FWorldDelegates::GetAssetTagsWithContext;
 FWorldDelegates::FOnWorldTickStart FWorldDelegates::OnWorldTickStart;
 FWorldDelegates::FOnWorldTickEnd FWorldDelegates::OnWorldTickEnd;
@@ -10208,13 +10205,6 @@ void UWorld::OnAddExtraObjectsToDelete(const TArray<UObject*>& InObjectsToDelete
 	}
 }
 
-void UWorld::GetAssetRegistryTags(TArray<FAssetRegistryTag>& OutTags) const
-{
-	PRAGMA_DISABLE_DEPRECATION_WARNINGS;
-	Super::GetAssetRegistryTags(OutTags);
-	PRAGMA_ENABLE_DEPRECATION_WARNINGS;
-}
-
 void UWorld::GetAssetRegistryTags(FAssetRegistryTagsContext Context) const
 {
 	Super::GetAssetRegistryTags(Context);
@@ -10295,14 +10285,6 @@ void UWorld::GetAssetRegistryTags(FAssetRegistryTagsContext Context) const
 		}
 	}
 
-	PRAGMA_DISABLE_DEPRECATION_WARNINGS;
-	TArray<UObject::FAssetRegistryTag> DeprecatedTags;
-	FWorldDelegates::GetAssetTags.Broadcast(this, DeprecatedTags);
-	for (UObject::FAssetRegistryTag& Tag : DeprecatedTags)
-	{
-		Context.AddTag(MoveTemp(Tag));
-	}
-	PRAGMA_ENABLE_DEPRECATION_WARNINGS;
 	FWorldDelegates::GetAssetTagsWithContext.Broadcast(this, Context);
 }
 
@@ -10873,22 +10855,6 @@ bool FAsyncPreRegisterDDCRequest::GetAsynchronousResults(TArray<uint8>& OutData)
 }
 #endif
 
-FString ToString(EWorldType::Type Type)
-{
-	switch (Type)
-	{
-	case EWorldType::None: return TEXT("None");
-	case EWorldType::Game: return TEXT("Game");
-	case EWorldType::Editor: return TEXT("Editor");
-	case EWorldType::PIE: return TEXT("PIE");
-	case EWorldType::EditorPreview: return TEXT("EditorPreview");
-	case EWorldType::GamePreview: return TEXT("GamePreview");
-	case EWorldType::GameRPC: return TEXT("GameRPC");
-	case EWorldType::Inactive: return TEXT("Inactive");
-	default: return TEXT("Unknown");
-	}
-}
-
 FString ENGINE_API ToString(ENetMode NetMode)
 {
 	switch (NetMode)
diff --git a/Engine/Source/Runtime/Engine/Public/ComponentInstanceDataCache.h b/Engine/Source/Runtime/Engine/Public/ComponentInstanceDataCache.h
--- a/Engine/Source/Runtime/Engine/Public/ComponentInstanceDataCache.h
+++ b/Engine/Source/Runtime/Engine/Public/ComponentInstanceDataCache.h
@@ -73,10 +73,6 @@ struct FInstanceCacheDataBase
 
 	ENGINE_API virtual void AddReferencedObjects(FReferenceCollector& Collector);
 
-	UE_DEPRECATED(5.3, "Use GetUniqueTransientObject instead")
-	/** Get (or create) the unique transient outer for the duplicated objects created for this object */
-	UObject* GetUniqueTransientPackage() { return GetUniqueTransientObject(UObject::StaticClass()); }
-
 	/** Get (or create) the unique transient outer for the duplicated objects created for this object */
 	ENGINE_API UObject* GetUniqueTransientObject(UClass* Class);
 
diff --git a/Engine/Source/Runtime/UMG/Private/WidgetBlueprintGeneratedClass.cpp b/Engine/Source/Runtime/UMG/Private/WidgetBlueprintGeneratedClass.cpp
--- a/Engine/Source/Runtime/UMG/Private/WidgetBlueprintGeneratedClass.cpp
+++ b/Engine/Source/Runtime/UMG/Private/WidgetBlueprintGeneratedClass.cpp
@@ -543,12 +543,6 @@ void UWidgetBlueprintGeneratedClass::GetExtensions(TArray<UWidgetBlueprintGenera
 }
 
 #if WITH_EDITOR
-void UWidgetBlueprintGeneratedClass::GetAssetRegistryTags(TArray<FAssetRegistryTag>& OutTags) const
-{
-	PRAGMA_DISABLE_DEPRECATION_WARNINGS;
-	Super::GetAssetRegistryTags(OutTags);
-	PRAGMA_ENABLE_DEPRECATION_WARNINGS;
-}
 
 void UWidgetBlueprintGeneratedClass::GetAssetRegistryTags(FAssetRegistryTagsContext Context) const
 {
diff --git a/Engine/Source/Runtime/UMG/Public/Blueprint/WidgetBlueprintGeneratedClass.h b/Engine/Source/Runtime/UMG/Public/Blueprint/WidgetBlueprintGeneratedClass.h
--- a/Engine/Source/Runtime/UMG/Public/Blueprint/WidgetBlueprintGeneratedClass.h
+++ b/Engine/Source/Runtime/UMG/Public/Blueprint/WidgetBlueprintGeneratedClass.h
@@ -170,8 +170,6 @@ public:
 	UMG_API virtual bool NeedsLoadForServer() const override;
 #if WITH_EDITOR
 	UMG_API virtual void GetAssetRegistryTags(FAssetRegistryTagsContext Context) const override;
-	UE_DEPRECATED(5.4, "Implement the version that takes FAssetRegistryTagsContext instead.")
-	UMG_API virtual void GetAssetRegistryTags(TArray<FAssetRegistryTag>& OutTags) const override;
 	UMG_API virtual void CollectSaveOverrides(FObjectCollectSaveOverridesContext SaveContext) override;
 #endif
 	//~ End UObject interface
