[BlueprintCompilation]
- Fix potential deadlock when saving package during BP compilation

#Synced-CL 51934851

diff --git a/Engine/Source/Editor/Kismet/Private/BlueprintCompilationManager.cpp b/Engine/Source/Editor/Kismet/Private/BlueprintCompilationManager.cpp
--- a/Engine/Source/Editor/Kismet/Private/BlueprintCompilationManager.cpp
+++ b/Engine/Source/Editor/Kismet/Private/BlueprintCompilationManager.cpp
@@ -454,6 +454,7 @@ void FBlueprintCompilationManagerImpl::CompileSynchronouslyImpl(const FBPCompile
 		}
 	}
 
+	TArray<UPackage*> PackagesToSave;
 	if (CompiledBlueprintsToSave.Num() > 0 && !bRegenerateSkeletonOnly)
 	{
 		if (!bSkipSaving)
@@ -462,7 +463,6 @@ void FBlueprintCompilationManagerImpl::CompileSynchronouslyImpl(const FBPCompile
 			const bool bSaveBlueprintsAfterCompile = Settings->SaveOnCompile == SoC_Always;
             const bool bSaveBlueprintAfterCompileSucceeded = Settings->SaveOnCompile == SoC_SuccessOnly;
 			
-			TArray<UPackage*> PackagesToSave;
 			for (UBlueprint* BP : CompiledBlueprintsToSave)
 			{
 				// The status of the blueprint might have changed during PostCDOCompiled
@@ -472,13 +472,21 @@ void FBlueprintCompilationManagerImpl::CompileSynchronouslyImpl(const FBPCompile
 				}
 			}
 
-			FEditorFileUtils::PromptForCheckoutAndSave(PackagesToSave, /*bCheckDirty =*/true, /*bPromptToSave =*/false);
 		}
 		CompiledBlueprintsToSave.Empty();
 	}
 
 	// We've done our GC, so release old CDO references
 	OldCDOs.Empty();
+
+#if WITH_EDITOR
+	ScopeLock.Unlock();
+#endif
+	// We need to save outside the compilation lock because saving can require some additional loading.
+	if (PackagesToSave.Num())
+	{
+		FEditorFileUtils::PromptForCheckoutAndSave(PackagesToSave, /*bCheckDirty =*/true, /*bPromptToSave =*/false);
+	}
 }
 
 enum class ECompilationManagerJobType
