Enabled viewport interactions in the Blueprint Editor
Fixes arcball click-through

#Synced-CL 52103551

diff --git a/Engine/Source/Editor/Kismet/Private/BlueprintEditor.cpp b/Engine/Source/Editor/Kismet/Private/BlueprintEditor.cpp
--- a/Engine/Source/Editor/Kismet/Private/BlueprintEditor.cpp
+++ b/Engine/Source/Editor/Kismet/Private/BlueprintEditor.cpp
@@ -2527,7 +2527,12 @@ void FBlueprintEditor::InitBlueprintEditor(
 	const bool bCreateDefaultToolbar = true;
 	const FName BlueprintEditorAppName = FName(TEXT("BlueprintEditorApp"));
 	InitAssetEditor(Mode, InitToolkitHost, BlueprintEditorAppName, FTabManager::FLayout::NullLayout, bCreateDefaultStandaloneMenu, bCreateDefaultToolbar, Objects);
-	
+
+	if (UModeManagerInteractiveToolsContext* ToolsContext = GetEditorModeManager().GetInteractiveToolsContext())
+	{
+		ToolsContext->SetViewportInteractionsEnabled(true);
+	}
+
 	CommonInitialization(InBlueprints, bShouldOpenInDefaultsMode);
 
 	InitalizeExtenders();
