A Batch Processor framework for running script functions on large datasets

#Synced-CL 51986075

diff --git a/Engine/Source/Editor/Kismet/Kismet.Build.cs b/Engine/Source/Editor/Kismet/Kismet.Build.cs
--- a/Engine/Source/Editor/Kismet/Kismet.Build.cs
+++ b/Engine/Source/Editor/Kismet/Kismet.Build.cs
@@ -58,6 +58,7 @@ public class Kismet : ModuleRules
 				"WorkspaceMenuStructure",
 				"DeveloperSettings",
 				"ToolMenus",
+				"Sockets",
 				"SubobjectEditor",
 				"SubobjectDataInterface",
 				"ToolWidgets",
diff --git a/Engine/Source/Editor/Kismet/Private/BatchProcess/BatchProcessCoordinator.cpp b/Engine/Source/Editor/Kismet/Private/BatchProcess/BatchProcessCoordinator.cpp
new file mode 100644
--- /dev/null
+++ b/Engine/Source/Editor/Kismet/Private/BatchProcess/BatchProcessCoordinator.cpp
@@ -0,0 +1,856 @@
+// Copyright Epic Games, Inc. All Rights Reserved.
+
+#include "BatchProcess/BatchProcessCoordinator.h"
+
+#include "HAL/FileManager.h"
+#include "HAL/PlatformTime.h"
+#include "Misc/App.h"
+#include "Misc/OutputDeviceRedirector.h"
+#include "Misc/Paths.h"
+#include "BatchProcess/BatchProcessLog.h"
+#include "BatchProcess/BatchProcessSocketHelpers.h"
+#include "Policies/CondensedJsonPrintPolicy.h"
+#include "Serialization/JsonReader.h"
+#include "Serialization/JsonSerializer.h"
+#include "Serialization/JsonWriter.h"
+#include "SocketSubsystem.h"
+#include "Sockets.h"
+
+DEFINE_LOG_CATEGORY(LogBatchProcess);
+
+namespace UE::Private
+{
+	FUtf8String ReadStringFromFile(const TCHAR* Filename);
+	const TCHAR* const BatchCoordinatorName   = TEXT("BatchProcessCoordinator");
+	const TCHAR* const CoordinatorIPStr       = TEXT("127.0.0.1");
+	const uint32 CoordinatorIP                = 0x7f000001;
+	const int32  MaxConnections               = 512;
+	const float  TimeoutLengthSeconds         = 1.f;
+	const float  JobStuckTimeoutSeconds       = 120.f; // kill a Running worker after 2 min silence
+	const int32  MaxJobRetries                = 3;     // re-enqueue limit before marking a job failed
+}
+
+FUtf8String UE::Private::ReadStringFromFile(const TCHAR* Filename)
+{
+	TUniquePtr<FArchive> Reader(IFileManager::Get().CreateFileReader(Filename));
+	if (!Reader.IsValid())
+	{
+		return FUtf8String();
+	}
+	FUtf8String ResultString;
+	int64 TotalSize = Reader->TotalSize();
+	if(TotalSize > 0 && TotalSize < TNumericLimits<int32>::Max())
+	{
+		ResultString.GetCharArray().Reserve(TotalSize+1);
+		ResultString.GetCharArray().SetNum(TotalSize);
+		Reader->Serialize(ResultString.GetCharArray().GetData(), TotalSize);
+		ResultString.GetCharArray().Add(UTF8CHAR('\0'));
+	}
+	return ResultString;
+}
+
+FBatchProcessCoordinator::FBatchProcessCoordinator(const FString* ClientAddress, int32 InNumWorkers)
+	: NumWorkers(InNumWorkers)
+	, bIsCompleted(false)
+	, bIsSealed(false)
+{
+	// create socket
+	ISocketSubsystem* SocketSubsystem = ISocketSubsystem::Get();
+	check(SocketSubsystem);
+	ListenAddr = SocketSubsystem->GetLocalBindAddr(*GLog);
+	ListenAddr->SetIp(UE::Private::CoordinatorIP);
+	ListenAddr->SetPort(UE::Private::BatchProcessSocketHelpers::DefaultPortNumber);
+	CoordinatorSocket = SocketSubsystem->CreateSocket(NAME_Stream, UE::Private::BatchCoordinatorName, ListenAddr->GetProtocolType());
+	CoordinatorSocket->SetReuseAddr();
+
+	const bool bBindSuccessful = CoordinatorSocket->Bind(*ListenAddr);
+	check(bBindSuccessful);
+
+	const bool bListenSuccessful = CoordinatorSocket->Listen(UE::Private::MaxConnections);
+	check(bListenSuccessful);
+
+	// ClientAddress provided: we're in "dummy coordinator" mode — someone is launching a
+	// coordinator just to exercise a single client (e.g. to debug a worker process).
+	if(ClientAddress)
+	{
+		// connect back to the client and tell it our address so it can connect:
+		TSharedPtr<class FInternetAddr> ClientAddr = SocketSubsystem->GetLocalBindAddr(*GLog);
+		bool bIpValid = false;
+		ClientAddr->SetIp(**ClientAddress, bIpValid);
+		ClientAddr->SetPort(61580); // debug client listen port — see BatchProcessWorker
+		check(bIpValid);
+
+		FSocket* ToClient = SocketSubsystem->CreateSocket(NAME_Stream, UE::Private::BatchCoordinatorName, ClientAddr->GetProtocolType());
+		bool bConnected = ToClient->Connect(*ClientAddr);
+		check(bConnected);
+		check(UE::Private::BatchProcessSocketHelpers::SendString(*ToClient, UE::Private::CoordinatorIPStr));
+
+		ActiveWorkers.Add(
+			{
+				EWorkerState::Starting,
+				FProcHandle(),
+				FPlatformTime::Seconds(),
+				TEXT("DummyWorker")
+			}
+		);
+
+		SocketSubsystem->DestroySocket(ToClient);
+
+		check(GetNumWorkers() == 1); // dummy mode only supports a single worker
+	}
+}
+
+FBatchProcessCoordinator::~FBatchProcessCoordinator()
+{
+	Abort();
+}
+
+void FBatchProcessCoordinator::SubmitFile(const TCHAR* Filename)
+{
+	FUtf8String Command = UE::Private::ReadStringFromFile(Filename);
+	if(Command == FUtf8String())
+	{
+		HandleError(FString::Printf(TEXT("Failed to read job file: %s"), Filename));
+		return;
+	}
+	Submit(Command);
+}
+
+void FBatchProcessCoordinator::Submit(const FUtf8String& JsonCommand)
+{
+	TSharedPtr<FJsonObject> JsonObject;
+	TSharedRef<TJsonReader<UTF8CHAR>> JsonReader = TJsonReaderFactory<UTF8CHAR>::Create(JsonCommand);
+	if (!FJsonSerializer::Deserialize(JsonReader, JsonObject) || !JsonObject.IsValid())
+	{
+		HandleError(TEXT("Failed to parse job JSON"));
+		return;
+	}
+
+	const bool bFoundFunction = JsonObject->TryGetStringField(TEXT("Function"), FunctionAsString);
+	if(!bFoundFunction)
+	{
+		HandleError(TEXT("Job JSON missing required field: \"function\""));
+		return;
+	}
+
+	const TArray<TSharedPtr<FJsonValue>>* ArgumentsField;
+	if(!JsonObject->TryGetArrayField(TEXT("Arguments"), ArgumentsField))
+	{
+		HandleError(TEXT("Job JSON missing required field: \"arguments\""));
+		return;
+	}
+
+	for (const TSharedPtr<FJsonValue>& Arg : *ArgumentsField)
+	{
+		AppendJob(Arg);
+	}
+}
+
+void FBatchProcessCoordinator::AppendJob(TSharedPtr<FJsonValue> Argument)
+{
+	FBatchJobResult& Result = JobResults.AddDefaulted_GetRef();
+
+	// Serialize the argument to a compact string for use in diagnostic output
+	// when the job produces no log lines of its own.
+	const TSharedPtr<FJsonObject>* ObjPtr = nullptr;
+	if (Argument.IsValid() && Argument->TryGetObject(ObjPtr))
+	{
+		TSharedRef<TJsonWriter<TCHAR, TCondensedJsonPrintPolicy<TCHAR>>> Writer =
+			TJsonWriterFactory<TCHAR, TCondensedJsonPrintPolicy<TCHAR>>::Create(&Result.ArgumentSummary);
+		FJsonSerializer::Serialize((*ObjPtr).ToSharedRef(), Writer);
+	}
+
+	JobArguments.Add(MoveTemp(Argument));
+	JobRetryCount.Add(0);
+}
+
+void FBatchProcessCoordinator::Seal()
+{
+	bIsSealed = true;
+}
+
+void FBatchProcessCoordinator::Abort()
+{
+	HandleCleanup();
+}
+
+bool FBatchProcessCoordinator::Tick()
+{
+	AcceptConnections();
+	ReceiveHandshakes();
+	LaunchNewWorkers();
+	UpdateActiveWorkers();
+	CheckCompletion();
+
+	return !bIsCompleted;
+}
+
+bool FBatchProcessCoordinator::HasAnyErrors() const
+{
+	if(!ErrorRecord.IsEmpty())
+	{
+		return true;
+	}
+	for(const FWorkerStateData& Worker : ActiveWorkers)
+	{
+		if(IsErrorState(Worker.WorkerState))
+		{
+			return true;
+		}
+	}
+	return false;
+}
+
+FString FBatchProcessCoordinator::GetErrorString() const
+{
+	return ErrorRecord;
+}
+
+int32 FBatchProcessCoordinator::GetNumJobs() const
+{
+	return JobResults.Num();
+}
+
+const FBatchJobResult& FBatchProcessCoordinator::GetJobResult(int32 Idx) const
+{
+	return JobResults[Idx];
+}
+
+const FString& FBatchProcessCoordinator::GetFunctionAsString() const
+{
+	return FunctionAsString;
+}
+
+void FBatchProcessCoordinator::HandleCleanup()
+{
+	ISocketSubsystem* SocketSubsystem = ISocketSubsystem::Get();
+
+	if(CoordinatorSocket)
+	{
+		SocketSubsystem->DestroySocket(CoordinatorSocket);
+		CoordinatorSocket = nullptr;
+	}
+
+	for (FSocket* Socket : PendingWorkerSockets)
+	{
+		SocketSubsystem->DestroySocket(Socket);
+	}
+	PendingWorkerSockets.Reset();
+
+	for (FWorkerStateData& Worker : ActiveWorkers)
+	{
+		if (Worker.Socket)
+		{
+			SocketSubsystem->DestroySocket(Worker.Socket);
+			Worker.Socket = nullptr;
+		}
+	}
+}
+
+void FBatchProcessCoordinator::HandleError(const FString& Error)
+{
+	if(!ErrorRecord.IsEmpty())
+	{
+		ErrorRecord += TEXT("\n");
+	}
+	ErrorRecord += Error;
+}
+
+int32 FBatchProcessCoordinator::GetNumWorkers() const
+{
+	return NumWorkers;
+}
+
+int32 FBatchProcessCoordinator::GetNumLaunchingWorkers() const
+{
+	int32 Count = 0;
+	for (const FWorkerStateData& Worker : ActiveWorkers)
+	{
+		if (Worker.WorkerState == EWorkerState::Starting)
+		{
+			++Count;
+		}
+	}
+	return Count;
+}
+
+int32 FBatchProcessCoordinator::GetNumActiveWorkers() const
+{
+	int32 Count = 0;
+	for (const FWorkerStateData& Worker : ActiveWorkers)
+	{
+		if (Worker.WorkerState == EWorkerState::Starting
+		 || Worker.WorkerState == EWorkerState::WaitingForWork
+		 || Worker.WorkerState == EWorkerState::Running)
+		{
+			++Count;
+		}
+	}
+	return Count;
+}
+
+void FBatchProcessCoordinator::RegisterPendingWorker()
+{
+	// An invalid FProcHandle marks the slot as a "dummy" — FindCrashedWorkers skips it.
+	ActiveWorkers.Add(
+		{
+			EWorkerState::Starting,
+			FProcHandle(),
+			FPlatformTime::Seconds(),
+			TEXT("PrelaunchedWorker")
+		}
+	);
+}
+
+void FBatchProcessCoordinator::AcceptConnections()
+{
+	bool bHasPendingConnection = false;
+	if (!CoordinatorSocket->HasPendingConnection(bHasPendingConnection) || !bHasPendingConnection)
+	{
+		return;
+	}
+
+	FSocket* ClientSocket = CoordinatorSocket->Accept(UE::Private::BatchCoordinatorName);
+	if(!ClientSocket)
+	{
+		return;
+	}
+
+	PendingWorkerSockets.Add(ClientSocket);
+}
+
+void FBatchProcessCoordinator::ReceiveHandshakes()
+{
+	// Workers send their index as a 4-byte handshake. Don't block the update loop:
+	// leave a socket in PendingWorkerSockets until its 4 bytes have arrived.
+	for(int32 Idx = PendingWorkerSockets.Num() - 1; Idx >= 0; --Idx)
+	{
+		FSocket* Socket = PendingWorkerSockets[Idx];
+
+		uint32 PendingBytes = 0;
+		if (!Socket->HasPendingData(PendingBytes))
+		{
+			// No data in the receive buffer. A remotely-closed connection becomes
+			// readable (select fires) but yields 0 bytes — detect and discard such
+			// sockets so dead pending entries don't accumulate.
+			if (Socket->Wait(ESocketWaitConditions::WaitForRead, FTimespan::Zero()))
+			{
+				HandleError(TEXT("Pending worker closed connection before sending handshake"));
+				ISocketSubsystem::Get()->DestroySocket(Socket);
+				PendingWorkerSockets.RemoveAt(Idx);
+			}
+			continue;
+		}
+		if (PendingBytes < sizeof(uint32))
+		{
+			continue;
+		}
+
+		uint32 WorkerId = 0;
+		if (!UE::Private::BatchProcessSocketHelpers::RecvUint32(*Socket, WorkerId))
+		{
+			HandleError(TEXT("Worker connected but failed to send handshake ID"));
+			ISocketSubsystem::Get()->DestroySocket(Socket);
+			PendingWorkerSockets.RemoveAt(Idx);
+			continue;
+		}
+
+		if (!ActiveWorkers.IsValidIndex((int32)WorkerId))
+		{
+			HandleError(FString::Printf(TEXT("Worker sent out-of-range ID %u (have %d slots)"),
+				WorkerId, ActiveWorkers.Num()));
+			ISocketSubsystem::Get()->DestroySocket(Socket);
+			PendingWorkerSockets.RemoveAt(Idx);
+			continue;
+		}
+
+		FWorkerStateData& Worker = ActiveWorkers[WorkerId];
+		if (Worker.Socket != nullptr)
+		{
+			HandleError(FString::Printf(TEXT("Worker ID %u connected twice"), WorkerId));
+			ISocketSubsystem::Get()->DestroySocket(Socket);
+			PendingWorkerSockets.RemoveAt(Idx);
+			continue;
+		}
+
+		Worker.Socket = Socket;
+		const double BootElapsed = FPlatformTime::Seconds() - Worker.LastStateChangeTime;
+		UE_LOG(LogBatchProcess, Display, TEXT("Worker %d ready  (boot: %.1fs)"), WorkerId, BootElapsed);
+		Worker.WorkerState = EWorkerState::WaitingForWork;
+		PendingWorkerSockets.RemoveAt(Idx);
+	}
+}
+
+void FBatchProcessCoordinator::LaunchNewWorkers()
+{
+	if(CompletionCounter == JobResults.Num())
+	{
+		return;
+	}
+
+	// Never launch more workers than there are remaining jobs.
+	// Count only live workers (Starting/WaitingForWork/Running) — ActiveWorkers may also
+	// contain terminal entries from previous restarts (FailedToShutdown etc.) that no longer
+	// consume a worker slot.
+	const int32 RemainingJobs = JobResults.Num() - CompletionCounter;
+	if (RemainingJobs <= 0)
+	{
+		return;
+	}
+	int32 NumLiveWorkers = 0;
+	for (const FWorkerStateData& W : ActiveWorkers)
+	{
+		if (W.WorkerState == EWorkerState::Starting
+		 || W.WorkerState == EWorkerState::WaitingForWork
+		 || W.WorkerState == EWorkerState::Running)
+		{
+			++NumLiveWorkers;
+		}
+	}
+	const int32 NumToLaunch = FMath::Min(GetNumWorkers(), RemainingJobs) - NumLiveWorkers;
+	if(NumToLaunch <= 0)
+	{
+		return;
+	}
+
+	const FString LaunchExecutableName = FPlatformProcess::ExecutablePath();
+	FString ProjectPath = FPaths::GetProjectFilePath();
+
+	for(int32 I = 0; I < NumToLaunch; ++I)
+	{
+		// The worker identifies itself during handshake by its index in ActiveWorkers,
+		// so capture the index before appending the new entry.
+		const int32 NewWorkerIndex = ActiveWorkers.Num();
+		FString ScratchLaunchString = FString::Printf(
+			TEXT("%s -run=BatchWorkerCommandlet %s:%d %d -coordinatorpid=%u -unattended -nopause -nosplash -NoExceptionHandler -nullrhi -nosound -notexturestreaming"),
+			*ProjectPath, UE::Private::CoordinatorIPStr, UE::Private::BatchProcessSocketHelpers::DefaultPortNumber, NewWorkerIndex,
+			FPlatformProcess::GetCurrentProcessId());
+		ActiveWorkers.Add(
+			{
+				EWorkerState::Starting,
+				FPlatformProcess::CreateProc(
+					*LaunchExecutableName,
+					*ScratchLaunchString,
+					false,
+					true,
+					true,
+					nullptr,
+					0,
+					nullptr,
+					nullptr),
+				FPlatformTime::Seconds(),
+				MoveTemp(ScratchLaunchString)
+			}
+		);
+	}
+}
+
+void FBatchProcessCoordinator::UpdateActiveWorkers()
+{
+	FindCrashedWorkers();
+	FindStuckWorkers();
+	ProcessResultsFromWorkers(); // transitions finished workers to WaitingForWork
+	GiveWorkToWorkers();         // immediately picks up any newly-idle workers
+}
+
+void FBatchProcessCoordinator::CheckCompletion()
+{
+	// Do not complete until Seal() has been called — the job list may still grow.
+	if (!bIsSealed)
+	{
+		return;
+	}
+
+	if(CompletionCounter == JobResults.Num())
+	{
+		// All jobs are done.  Close every still-connected worker socket so that
+		// PerformWork() on the worker side gets an EOF and returns cleanly.
+		for(FWorkerStateData& Worker : ActiveWorkers)
+		{
+			if(Worker.Socket &&
+				(Worker.WorkerState == EWorkerState::Running ||
+				 Worker.WorkerState == EWorkerState::WaitingForWork))
+			{
+				Worker.Socket->Close();
+				Worker.ChangeState(EWorkerState::ShuttingDown);
+			}
+		}
+		bIsCompleted = true;
+		return;
+	}
+
+	const auto IsTerminalState = [](EWorkerState WorkerState)
+	{
+		return WorkerState == EWorkerState::FailedToStart ||
+			WorkerState == EWorkerState::FailedToShutdown ||
+			WorkerState == EWorkerState::ShuttingDown;
+	};
+
+	bool bAllWorkersAreDone = true;
+	for(FWorkerStateData& Worker : ActiveWorkers)
+	{
+		if(!IsTerminalState(Worker.WorkerState))
+		{
+			bAllWorkersAreDone = false;
+		}
+	}
+
+	if(bAllWorkersAreDone && ActiveWorkers.Num() > 0)
+	{
+		FailAllRemainingJobs();
+		bIsCompleted = true;
+	}
+	else
+	{
+		check(!bIsCompleted);
+	}
+}
+
+void FBatchProcessCoordinator::FindCrashedWorkers()
+{
+	for(FWorkerStateData& Worker : ActiveWorkers)
+	{
+		if( FPlatformProcess::IsProcRunning(Worker.Handle) ||
+			!Worker.Handle.IsValid()) // dummy worker
+		{
+			continue;
+		}
+
+		if(Worker.WorkerState == EWorkerState::Starting)
+		{
+			HandleError(TEXT("Batch worker failed to start"));
+			Worker.ChangeState(EWorkerState::FailedToStart);
+		}
+		else
+		{
+			Worker.Socket->Close();
+
+			HandleError(FString::Printf(TEXT("Worker crashed: %s"), *Worker.LaunchString));
+
+			// Re-enqueue the in-flight job so another worker can pick it up.
+			if (Worker.CurrentJobIndex >= 0)
+			{
+				const int32 JobIdx = Worker.CurrentJobIndex;
+				Worker.CurrentJobIndex = -1;
+				if (JobRetryCount[JobIdx]++ < UE::Private::MaxJobRetries)
+				{
+					RetryQueue.Add(JobIdx);
+				}
+				else
+				{
+					JobResults[JobIdx].State           = EBatchJobState::Completed;
+					JobResults[JobIdx].bEncounteredError = true;
+					JobResults[JobIdx].Output          = TEXT("Worker crashed repeatedly on this job");
+					++CompletionCounter;
+				}
+			}
+
+			Worker.ChangeState(EWorkerState::FailedToShutdown);
+		}
+	}
+}
+
+void FBatchProcessCoordinator::FindStuckWorkers()
+{
+	const double CurrentTime = FPlatformTime::Seconds();
+	for(FWorkerStateData& Worker : ActiveWorkers)
+	{
+		if(Worker.WorkerState == EWorkerState::ShuttingDown)
+		{
+			if(Worker.LastStateChangeTime + UE::Private::TimeoutLengthSeconds < CurrentTime)
+			{
+				HandleError(TEXT("Worker failed to shut down in time"));
+				Worker.ChangeState(EWorkerState::FailedToShutdown);
+			}
+		}
+		else if (Worker.WorkerState == EWorkerState::Running
+			&& Worker.LastResultTime + UE::Private::JobStuckTimeoutSeconds < CurrentTime)
+		{
+			HandleError(FString::Printf(TEXT("Worker silent for %.0fs, terminating: %s"),
+				UE::Private::JobStuckTimeoutSeconds, *Worker.LaunchString));
+
+			// Re-enqueue the in-flight job so another worker can pick it up.
+			if (Worker.CurrentJobIndex >= 0)
+			{
+				const int32 JobIdx = Worker.CurrentJobIndex;
+				Worker.CurrentJobIndex = -1;
+				if (JobRetryCount[JobIdx]++ < UE::Private::MaxJobRetries)
+				{
+					RetryQueue.Add(JobIdx);
+				}
+				else
+				{
+					JobResults[JobIdx].State             = EBatchJobState::Completed;
+					JobResults[JobIdx].bEncounteredError = true;
+					JobResults[JobIdx].Output            = TEXT("Worker timed out repeatedly on this job");
+					++CompletionCounter;
+				}
+			}
+
+			FPlatformProcess::TerminateProc(Worker.Handle, true);
+			if (Worker.Socket)
+			{
+				Worker.Socket->Close();
+			}
+			Worker.ChangeState(EWorkerState::FailedToShutdown);
+		}
+	}
+}
+
+void FBatchProcessCoordinator::GiveWorkToWorkers()
+{
+	for (FWorkerStateData& Worker : ActiveWorkers)
+	{
+		if (Worker.WorkerState != EWorkerState::WaitingForWork)
+		{
+			continue;
+		}
+
+		// Pick the next job: drain the retry queue first (crashed/timed-out jobs),
+		// then advance the sequential dispatch pointer.
+		int32 JobIdx = -1;
+		if (RetryQueue.Num() > 0)
+		{
+			JobIdx = RetryQueue.Last();
+			RetryQueue.RemoveAt(RetryQueue.Num() - 1, 1, EAllowShrinking::No);
+		}
+		else if (NextJobToDispatch < JobArguments.Num())
+		{
+			JobIdx = NextJobToDispatch++;
+		}
+		else
+		{
+			break; // no more jobs to dispatch right now
+		}
+
+		// Build a single-element batch packet.
+		TArray<TSharedPtr<FJsonValue>> Arguments;
+		Arguments.Add(JobArguments[JobIdx]);
+
+		TSharedRef<FJsonObject> Object = ::MakeShared<FJsonObject>();
+		Object->SetArrayField(TEXT("Arguments"), Arguments);
+		Object->SetStringField(TEXT("Function"),   FunctionAsString);
+		Object->SetNumberField(TEXT("StartIndex"), JobIdx);
+
+		if (!UE::Private::BatchProcessSocketHelpers::SendJsonValue(*Worker.Socket, Object))
+		{
+			HandleError(FString::Printf(TEXT("Failed to send work to worker: %s"), *Worker.LaunchString));
+			RetryQueue.Add(JobIdx); // put the job back so another worker can pick it up
+			Worker.ChangeState(EWorkerState::FailedToShutdown);
+			continue;
+		}
+
+		Worker.CurrentJobIndex = JobIdx;
+		Worker.LastResultTime  = FPlatformTime::Seconds();
+		Worker.ChangeState(EWorkerState::Running);
+		JobResults[JobIdx].State = EBatchJobState::InProgress;
+	}
+}
+
+void FBatchProcessCoordinator::ProcessResultsFromWorkers()
+{
+	// Phase 1 — drain bytes.
+	// Read all currently-available socket bytes into each worker's per-worker buffer.
+	// We do this across all workers before parsing so that no single worker's large
+	// result payload can starve the others: every worker's send window is kept open
+	// regardless of how quickly individual messages are parsed.
+	for (FWorkerStateData& Worker : ActiveWorkers)
+	{
+		if (Worker.WorkerState != EWorkerState::Running)
+		{
+			continue;
+		}
+
+		uint32 Available = 0;
+		while (Worker.Socket->HasPendingData(Available) && Available > 0)
+		{
+			const int32 Offset = Worker.RecvBuffer.AddUninitialized((int32)Available);
+			int32 BytesRead = 0;
+			if (!Worker.Socket->Recv(Worker.RecvBuffer.GetData() + Offset, (int32)Available, BytesRead)
+				|| BytesRead == 0)
+			{
+				Worker.RecvBuffer.SetNum(Offset, EAllowShrinking::No); // undo the uninitialised reservation
+				HandleError(FString::Printf(TEXT("Lost connection to worker mid-batch: %s"), *Worker.LaunchString));
+				Worker.ChangeState(EWorkerState::FailedToShutdown);
+				break;
+			}
+			Worker.RecvBuffer.SetNum(Offset + BytesRead, EAllowShrinking::No);
+		}
+	}
+
+	// Phase 2 — parse messages.
+	// Walk each worker's buffer and process every complete length-prefixed JSON frame.
+	// Incomplete frames (body not yet fully received) are left in the buffer for the
+	// next tick — the coordinator never blocks waiting for more socket data.
+	for (FWorkerStateData& Worker : ActiveWorkers)
+	{
+		if (Worker.WorkerState != EWorkerState::Running)
+		{
+			continue;
+		}
+
+		while (true)
+		{
+			const int32 Remaining = Worker.RecvBuffer.Num() - Worker.RecvBufferReadOffset;
+
+			// Need at least the 4-byte big-endian length header.
+			if (Remaining < (int32)sizeof(uint32))
+			{
+				break;
+			}
+
+			const uint8* const Buf = Worker.RecvBuffer.GetData() + Worker.RecvBufferReadOffset;
+			const uint32 MsgLen = (uint32(Buf[0]) << 24)
+			                    | (uint32(Buf[1]) << 16)
+			                    | (uint32(Buf[2]) <<  8)
+			                    |  uint32(Buf[3]);
+
+			// Guard against a corrupt length header before any arithmetic.
+			if (MsgLen > (uint32)(TNumericLimits<int32>::Max() - (int32)sizeof(uint32)))
+			{
+				HandleError(TEXT("Worker sent oversized message length (corrupt header)"));
+				Worker.ChangeState(EWorkerState::FailedToShutdown);
+				break;
+			}
+
+			// Wait until the full body has arrived.
+			if (Remaining < (int32)sizeof(uint32) + (int32)MsgLen)
+			{
+				break;
+			}
+
+			// Parse the JSON body directly from the buffer — no copy needed.
+			const TStringView<char> JsonStr{
+				reinterpret_cast<const char*>(Buf + sizeof(uint32)),
+				(int32)MsgLen
+			};
+			TSharedRef<TJsonStringViewReader<char>> Reader =
+				TJsonStringViewReader<char>::Create(JsonStr);
+			TSharedPtr<FJsonValue> Result;
+			if (!FJsonSerializer::Deserialize(*Reader, Result))
+			{
+				HandleError(TEXT("Worker sent malformed result JSON"));
+				Worker.ChangeState(EWorkerState::FailedToShutdown);
+				break;
+			}
+
+			// Advance past this frame before any early-out so the buffer stays consistent.
+			Worker.RecvBufferReadOffset += (int32)(sizeof(uint32) + MsgLen);
+
+			const TSharedPtr<FJsonObject>* JsonObjectPtr;
+			if (!Result || !Result->TryGetObject(JsonObjectPtr))
+			{
+				HandleError(TEXT("Worker sent malformed result JSON"));
+				Worker.ChangeState(EWorkerState::FailedToShutdown);
+				break;
+			}
+			const FJsonObject& JsonObject = *(JsonObjectPtr->Get());
+
+			bool bShutdown = false;
+			JsonObject.TryGetBoolField(TEXT("Shutdown"), bShutdown);
+			if (bShutdown)
+			{
+				Worker.ChangeState(EWorkerState::ShuttingDown);
+				break;
+			}
+
+			// Normal per-job result.
+			double JobIndex = 0.0;
+			JsonObject.TryGetNumberField(TEXT("JobIndex"), JobIndex);
+			FBatchJobResult& JobResult = JobResults[(int32)JobIndex];
+			JobResult.State = EBatchJobState::Completed;
+			JsonObject.TryGetStringField(TEXT("Output"), JobResult.Output);
+			bool bError = false;
+			JsonObject.TryGetBoolField(TEXT("Error"), bError);
+			JobResult.bEncounteredError = bError;
+
+			Worker.CurrentJobIndex = -1;
+			Worker.LastResultTime  = FPlatformTime::Seconds();
+			++CompletionCounter;
+
+			// Recycle the worker: put it back to WaitingForWork so GiveWorkToWorkers
+			// can dispatch the next job immediately (same tick).
+			Worker.ChangeState(EWorkerState::WaitingForWork);
+			// Do NOT break — there may be more complete frames already buffered.
+		}
+
+		// Compact the buffer: discard already-processed bytes with a single memmove
+		// so the buffer doesn't grow unboundedly across ticks.
+		if (Worker.RecvBufferReadOffset > 0)
+		{
+			const int32 Unprocessed = Worker.RecvBuffer.Num() - Worker.RecvBufferReadOffset;
+			if (Unprocessed > 0)
+			{
+				FMemory::Memmove(
+					Worker.RecvBuffer.GetData(),
+					Worker.RecvBuffer.GetData() + Worker.RecvBufferReadOffset,
+					Unprocessed);
+			}
+			Worker.RecvBuffer.SetNum(Unprocessed, EAllowShrinking::No);
+			Worker.RecvBufferReadOffset = 0;
+		}
+	}
+}
+
+void FBatchProcessCoordinator::FailAllRemainingJobs()
+{
+	auto Fail = [&](int32 JobIdx)
+	{
+		if (JobResults[JobIdx].State != EBatchJobState::Completed)
+		{
+			JobResults[JobIdx].State             = EBatchJobState::Completed;
+			JobResults[JobIdx].bEncounteredError = true;
+			JobResults[JobIdx].Output            = TEXT("All workers failed before this job could run");
+			++CompletionCounter;
+		}
+	};
+
+	for (int32 Idx : RetryQueue)
+	{
+		Fail(Idx);
+	}
+	RetryQueue.Reset();
+
+	for (int32 Idx = NextJobToDispatch; Idx < JobResults.Num(); ++Idx)
+	{
+		Fail(Idx);
+	}
+}
+
+bool FBatchProcessCoordinator::Recv(FSocket* FromSocket, TArray<uint8>& Bytes)
+{
+	uint32 BufferSize = 0;
+	if (FromSocket->HasPendingData(BufferSize))
+	{
+		Bytes.SetNumUninitialized(BufferSize);
+
+		int32 BytesRead = 0;
+		if (!FromSocket->Recv(Bytes.GetData(), Bytes.Num(), BytesRead))
+		{
+			HandleError(FString(TEXT("Batch worker read failed: ") + GetSocketReadableErrorCode()));
+			return false;
+		}
+		return true;
+	}
+
+	return false;
+}
+
+FString FBatchProcessCoordinator::GetSocketReadableErrorCode()
+{
+	ISocketSubsystem* SocketSubsystem = ISocketSubsystem::Get(PLATFORM_SOCKETSUBSYSTEM);
+	const ESocketErrors LastError = SocketSubsystem->GetLastErrorCode();
+	return FString::Printf(TEXT("%s (%d)"), SocketSubsystem->GetSocketError(LastError), static_cast<int32>(LastError));
+}
+
+bool FBatchProcessCoordinator::IsErrorState(EWorkerState WorkerState)
+{
+	return WorkerState == EWorkerState::FailedToStart ||
+		WorkerState == EWorkerState::FailedToShutdown;
+}
+
+void FBatchProcessCoordinator::FWorkerStateData::ChangeState(EWorkerState NewState)
+{
+	WorkerState = NewState;
+	LastStateChangeTime = FPlatformTime::Seconds();
+}
diff --git a/Engine/Source/Editor/Kismet/Private/BatchProcess/BatchProcessCoordinator.h b/Engine/Source/Editor/Kismet/Private/BatchProcess/BatchProcessCoordinator.h
new file mode 100644
--- /dev/null
+++ b/Engine/Source/Editor/Kismet/Private/BatchProcess/BatchProcessCoordinator.h
@@ -0,0 +1,156 @@
+// Copyright Epic Games, Inc. All Rights Reserved.
+
+#pragma once
+
+#include "CoreMinimal.h"
+
+#include "Dom/JsonValue.h"
+#include "HAL/PlatformProcess.h"
+
+class FSocket;
+
+enum class EBatchJobState
+{
+	Pending,
+	Scheduled,
+	InProgress,
+	Completed
+};
+
+struct FBatchJobResult
+{
+	EBatchJobState State        = EBatchJobState::Pending;
+	bool bEncounteredError      = false;
+	FString Output;          // log output captured by the worker during this job's function call
+	FString ArgumentSummary; // compact JSON of the original argument, for diagnostic display
+};
+
+struct FBatchProcessCoordinator
+{
+	FBatchProcessCoordinator(const FString* ClientAddress, int32 InNumWorkers = 1);
+	~FBatchProcessCoordinator();
+
+	/**
+	 * Parse a job description from a file or JSON string.
+	 * The JSON must contain a "function" (UFunction path) and "arguments" (array of arg objects).
+	 * Call Seal() once no more jobs will be added.
+	 */
+	void SubmitFile(const TCHAR* Filename);
+	void Submit(const FUtf8String& JsonCommand);
+
+	/**
+	 * Append a single argument object to the pending job queue.
+	 * May be called at any point before Seal(), including while Tick() is running.
+	 */
+	void AppendJob(TSharedPtr<FJsonValue> Argument);
+
+	/**
+	 * Signal that no more jobs will be appended.
+	 * The coordinator will complete once all queued jobs finish.
+	 */
+	void Seal();
+
+	/** Cancel all in-flight work and shut down workers. */
+	void Abort();
+
+	/**
+	 * Registers a pre-existing worker slot without launching a process.
+	 * The slot accepts an incoming connection via the normal handshake.
+	 * Useful for in-process testing and the "dummy coordinator" debug mode.
+	 */
+	void RegisterPendingWorker();
+
+	/** @return true if still in progress, false when complete */
+	bool Tick();
+
+	bool HasAnyErrors() const;
+	FString GetErrorString() const;
+	int32 GetNumJobs() const;
+	int32 GetNumWorkers() const;
+	int32 GetNumLaunchingWorkers() const;
+	/** Returns the number of workers that are Starting, WaitingForWork, or Running. */
+	int32 GetNumActiveWorkers() const;
+	const FBatchJobResult& GetJobResult(int32 Idx) const;
+	const FString& GetFunctionAsString() const;
+
+private:
+	void HandleCleanup();
+	void HandleError(const FString& Error);
+	void AcceptConnections();
+	void ReceiveHandshakes();
+	void LaunchNewWorkers();
+	void UpdateActiveWorkers();
+	void CheckCompletion();
+
+	void FindCrashedWorkers();
+	void FindStuckWorkers();
+	void GiveWorkToWorkers();
+	void ProcessResultsFromWorkers();
+	void FailAllRemainingJobs();
+
+	bool Recv(FSocket* FromSocket, TArray<uint8>& Bytes);
+	static FString GetSocketReadableErrorCode();
+
+	FBatchProcessCoordinator(const FBatchProcessCoordinator&) = delete;
+	FBatchProcessCoordinator& operator=(const FBatchProcessCoordinator&) = delete;
+	FBatchProcessCoordinator(FBatchProcessCoordinator&&) = delete;
+	FBatchProcessCoordinator& operator=(FBatchProcessCoordinator&&) = delete;
+
+	TArray<FBatchJobResult>          JobResults;
+	TArray<TSharedPtr<FJsonValue>>   JobArguments; // one entry per job
+	FString FunctionAsString;
+
+	enum class EWorkerState
+	{
+		Starting,
+		WaitingForWork,
+		Running,
+		ShuttingDown,
+		FailedToStart,
+		FailedToShutdown,
+	};
+	static bool IsErrorState(EWorkerState WorkerState);
+	struct FWorkerStateData
+	{
+		EWorkerState WorkerState      = EWorkerState::Starting;
+		FProcHandle  Handle;
+		double       LastStateChangeTime = 0.0;
+		FString      LaunchString;
+		FSocket*     Socket           = nullptr;
+
+		// The global job index this worker is currently executing; -1 when idle.
+		// Used for crash/timeout recovery: if the worker dies, this job is re-enqueued.
+		int32  CurrentJobIndex    = -1;
+
+		// Timestamp of the last result received from this worker.
+		// Used by FindStuckWorkers to detect Running workers that have gone silent.
+		double LastResultTime     = 0.0;
+
+		// Per-worker byte buffer for incoming socket data.
+		// The coordinator drains bytes non-blockingly into this buffer and
+		// parses complete framed messages from it, preventing a slow/large
+		// result from one worker from blocking reads of all other workers.
+		TArray<uint8> RecvBuffer;
+		int32 RecvBufferReadOffset = 0;
+
+		void ChangeState(EWorkerState NewState);
+	};
+
+	TArray<FSocket*>        PendingWorkerSockets;
+	TArray<FWorkerStateData> ActiveWorkers;
+	FString ErrorRecord;
+
+	int32          NumWorkers         = 1;
+	int32          NextJobToDispatch  = 0;  // next sequential job index to send
+	int32          CompletionCounter  = 0;
+	TArray<int32>  RetryQueue;              // jobs to re-attempt after worker crash/timeout
+	TArray<int8>   JobRetryCount;           // per-job retry count (parallel to JobResults)
+
+	// socket info:
+	TSharedPtr<class FInternetAddr> ListenAddr;
+	FSocket* CoordinatorSocket;
+
+	// flags:
+	bool bIsCompleted;
+	bool bIsSealed;  // true once no more jobs will be appended; enables completion detection
+};
diff --git a/Engine/Source/Editor/Kismet/Private/BatchProcess/BatchProcessLibrary.cpp b/Engine/Source/Editor/Kismet/Private/BatchProcess/BatchProcessLibrary.cpp
new file mode 100644
--- /dev/null
+++ b/Engine/Source/Editor/Kismet/Private/BatchProcess/BatchProcessLibrary.cpp
@@ -0,0 +1,35 @@
+// Copyright Epic Games, Inc. All Rights Reserved.
+
+#include "BatchProcess/BatchProcessLibrary.h"
+
+#include "BatchProcess/BatchProcessRunner.h"
+#include "Policies/CondensedJsonPrintPolicy.h"
+#include "Serialization/JsonWriter.h"
+
+FString UBatchProcessLibrary::RunBatch(const FString& JobJson, int32 NumWorkers)
+{
+	UE::Private::FBatchRunOptions Options;
+	Options.NumWorkers = NumWorkers;
+
+	const UE::Private::FBatchRunResult RunResult = UE::Private::RunBatchFromJson(JobJson, Options);
+
+	// Serialise results to a JSON array for the Python caller.
+	FString Out;
+	TSharedRef<TJsonWriter<TCHAR, TCondensedJsonPrintPolicy<TCHAR>>> Writer =
+		TJsonWriterFactory<TCHAR, TCondensedJsonPrintPolicy<TCHAR>>::Create(&Out);
+
+	Writer->WriteArrayStart();
+	for (int32 Idx = 0; Idx < RunResult.Jobs.Num(); ++Idx)
+	{
+		const FBatchJobResult& Job = RunResult.Jobs[Idx];
+		Writer->WriteObjectStart();
+		Writer->WriteValue(TEXT("Passed"),           !Job.bEncounteredError);
+		Writer->WriteValue(TEXT("Output"),            Job.Output);
+		Writer->WriteValue(TEXT("ArgumentSummary"),  Job.ArgumentSummary);
+		Writer->WriteObjectEnd();
+	}
+	Writer->WriteArrayEnd();
+	Writer->Close();
+
+	return Out;
+}
diff --git a/Engine/Source/Editor/Kismet/Private/BatchProcess/BatchProcessLibrary.h b/Engine/Source/Editor/Kismet/Private/BatchProcess/BatchProcessLibrary.h
new file mode 100644
--- /dev/null
+++ b/Engine/Source/Editor/Kismet/Private/BatchProcess/BatchProcessLibrary.h
@@ -0,0 +1,109 @@
+// Copyright Epic Games, Inc. All Rights Reserved.
+
+#pragma once
+
+#include "Kismet/BlueprintFunctionLibrary.h"
+
+#include "BatchProcessLibrary.generated.h"
+
+/**
+ * Blueprint/Python-callable wrapper around FBatchProcessCoordinator.
+ *
+ * RunBatch() lets you drive the batch processor from a single Python script:
+ * discover assets, build a job description, fan out to N workers, and get
+ * structured results back — all without leaving the current editor process.
+ *
+ * Example — stringify every Blueprint in the project with JsonObjectGraph
+ * note that for snapshot_asset to be reachable you must register it in
+ * init_unreal.py thusly: import snapshot_assets
+ 
+@code{.py}
+# Save in /Game/Content/Python/snapshot_assets.py
+import json
+import os
+import unreal
+
+# Declare a uclass with a function that we want to run, this simple function just
+# takes an asset path, loads it, stringifies it, and writes the string to a file:
[REDACTED_EMAIL]()
+class SnapshotAsset(unreal.Object):
+	@unreal.ufunction(params=[str])
+	def snapshot_asset(self, asset_path: str) -> None:
+		try:
+			asset = unreal.load_asset(asset_path)
+		except Exception as exc:
+			unreal.log_error(f"FAIL: {asset_path}  (load exception: {exc})")
+			return
+
+		if asset is None:
+			unreal.log_error(f"FAIL: {asset_path}  (asset not found)")
+			return
+
+		json_str = unreal.JsonObjectGraphFunctionLibrary.stringify(
+			[asset], unreal.JsonStringifyOptions(unreal.JsonStringifyFlags.FILTER_EDITOR_ONLY_DATA))
+
+		# Build a nested output path from the package path, stripping the leading '/'.
+		# e.g. "/Game/Baz/Foo" -> "Game/Baz/Foo.json" and write to that file:
+		relative = asset_path.lstrip("/")
+		out_path = os.path.join(
+			unreal.Paths.project_saved_dir(), "SnapshotAsset", f"{relative}.json"
+		)
+		os.makedirs(os.path.dirname(out_path), exist_ok=True)
+
+		with open(out_path, "w", encoding="utf-8") as f:
+			f.write(json_str)
+
+		unreal.log(f"PASS: {asset_path}  -> {out_path}")
+
+    
+# Discover all assets of the specified type, run snapshot_asset
+def _run_snapshot(asset_class: str)-> None:
+	registry = unreal.AssetRegistryHelpers.get_asset_registry()
+	registry.wait_for_completion()
+	asset_filter = unreal.ARFilter(
+		class_names=[asset_class], recursive_paths=True, recursive_classes=True)
+	asset_paths = [str(a.package_name)
+					for a in registry.get_assets(asset_filter)]
+    
+	# Build the job description.
+	job = {
+		"Function": "/Game/Python/snapshot_assets_PY.SnapshotAsset:snapshot_asset",
+		"Arguments": [{"asset_path": p} for p in asset_paths],
+	}
+    
+	# Run the batch — blocks until all jobs complete.
+	results_json = unreal.BatchProcessLibrary.run_batch(json.dumps(job), num_workers=8)
+	results = json.loads(results_json)
+    
+	passed = [r for r in results if r["Passed"]]
+	failed = [r for r in results if not r["Passed"]]
+	print(f"{len(passed)}/{len(results)} passed")
+	for r in failed:
+		print(f"  FAIL: {r['ArgumentSummary']}\n    {r['Output']}")
+
+def run_blueprints_snapshot() -> None:
+    _run_snapshot("Blueprint")
+	
+if __name__ == '__main__': # critical, prevents forkbomb if init_unreal is used to import this script
+    run_blueprints_snapshot()
+@endcode
+ */
+UCLASS()
+class UBatchProcessLibrary : public UBlueprintFunctionLibrary
+{
+	GENERATED_BODY()
+
+public:
+	/**
+	 * Run a batch of jobs across multiple worker processes and return structured results.
+	 *
+	 * @param JobJson    Full job description as a JSON string:
+	 *                   {"Function": "<UFunction path>", "Arguments": [{...}, ...]}
+	 * @param NumWorkers Number of worker processes to spawn.
+	 *                   Pass 0 (default) to auto-select based on core count.
+	 * @return           JSON array string — one object per job, in submission order:
+	 *                   [{"Passed": bool, "Output": "...", "ArgumentSummary": "..."}, ...]
+	 */
+	UFUNCTION(BlueprintCallable, Category = "BatchProcess")
+	static FString RunBatch(const FString& JobJson, int32 NumWorkers = 0);
+};
diff --git a/Engine/Source/Editor/Kismet/Private/BatchProcess/BatchProcessLog.h b/Engine/Source/Editor/Kismet/Private/BatchProcess/BatchProcessLog.h
new file mode 100644
--- /dev/null
+++ b/Engine/Source/Editor/Kismet/Private/BatchProcess/BatchProcessLog.h
@@ -0,0 +1,7 @@
+// Copyright Epic Games, Inc. All Rights Reserved.
+
+#pragma once
+
+#include "Logging/LogMacros.h"
+
+DECLARE_LOG_CATEGORY_EXTERN(LogBatchProcess, Log, All);
diff --git a/Engine/Source/Editor/Kismet/Private/BatchProcess/BatchProcessRunner.cpp b/Engine/Source/Editor/Kismet/Private/BatchProcess/BatchProcessRunner.cpp
new file mode 100644
--- /dev/null
+++ b/Engine/Source/Editor/Kismet/Private/BatchProcess/BatchProcessRunner.cpp
@@ -0,0 +1,292 @@
+// Copyright Epic Games, Inc. All Rights Reserved.
+
+#include "BatchProcess/BatchProcessRunner.h"
+
+#include "HAL/FileManager.h"
+#include "HAL/PlatformMisc.h"
+#include "HAL/PlatformProcess.h"
+#include "HAL/PlatformTime.h"
+#include "Math/UnrealMathUtility.h"
+#include "Misc/FileHelper.h"
+#include "Misc/Paths.h"
+#include "BatchProcess/BatchProcessLog.h"
+#if PLATFORM_WINDOWS
+#include "Windows/WindowsHWrapper.h"
+#else
+#include <unistd.h> // isatty, STDOUT_FILENO
+#endif
+
+namespace UE::Private
+{
+
+#if PLATFORM_WINDOWS
+#ifndef ENABLE_VIRTUAL_TERMINAL_PROCESSING
+#define ENABLE_VIRTUAL_TERMINAL_PROCESSING 0x0004
+#endif
+
+/**
+ * Returns a writable console HANDLE with VT processing enabled, or INVALID_HANDLE_VALUE.
+ *
+ * UnrealEditor.exe is a GUI-subsystem process. When launched from a terminal the
+ * inherited STD_OUTPUT_HANDLE is sometimes a pipe or has been replaced by UE's own
+ * log infrastructure, so GetConsoleMode on it fails.  Falling back to CONOUT$ opens
+ * the active console screen buffer directly and works as long as the process is
+ * attached to an interactive terminal (cmd.exe, PowerShell, Windows Terminal, etc.).
+ */
+static HANDLE AcquireConsoleHandle()
+{
+	HANDLE hOut = ::GetStdHandle(STD_OUTPUT_HANDLE);
+	DWORD Mode = 0;
+	if (hOut != INVALID_HANDLE_VALUE && hOut != nullptr && ::GetConsoleMode(hOut, &Mode))
+	{
+		::SetConsoleMode(hOut, Mode | ENABLE_VIRTUAL_TERMINAL_PROCESSING);
+		return hOut;
+	}
+
+	hOut = ::CreateFileW(L"CONOUT$", GENERIC_READ | GENERIC_WRITE,
+	                     FILE_SHARE_READ | FILE_SHARE_WRITE,
+	                     nullptr, OPEN_EXISTING, 0, nullptr);
+	if (hOut == INVALID_HANDLE_VALUE) { return INVALID_HANDLE_VALUE; }
+
+	if (::GetConsoleMode(hOut, &Mode))
+	{
+		::SetConsoleMode(hOut, Mode | ENABLE_VIRTUAL_TERMINAL_PROCESSING);
+		return hOut;
+	}
+
+	::CloseHandle(hOut);
+	return INVALID_HANDLE_VALUE;
+}
+
+static HANDLE GetConsoleHandle()
+{
+	static HANDLE hConsole = AcquireConsoleHandle();
+	return hConsole;
+}
+#endif // PLATFORM_WINDOWS
+
+/** Returns true if attached to an interactive terminal. Result is cached. */
+static bool IsInteractiveTerminal()
+{
+	static int8 CachedResult = -1;
+	if (CachedResult >= 0) { return CachedResult != 0; }
+#if PLATFORM_WINDOWS
+	CachedResult = (GetConsoleHandle() != INVALID_HANDLE_VALUE) ? 1 : 0;
+#else
+	CachedResult = (isatty(STDOUT_FILENO) != 0) ? 1 : 0;
+#endif
+	return CachedResult != 0;
+}
+
+/**
+ * Writes Text to the terminal.
+ * On Windows uses WriteConsoleW so \r and VT sequences reach the console even
+ * when UE has redirected the C-runtime stdout.
+ */
+static void WriteToTerminal(const TCHAR* Text)
+{
+#if PLATFORM_WINDOWS
+	HANDLE hConsole = GetConsoleHandle();
+	if (hConsole != INVALID_HANDLE_VALUE)
+	{
+		DWORD Written = 0;
+		::WriteConsoleW(hConsole, Text, (DWORD)FCString::Strlen(Text), &Written, nullptr);
+		return;
+	}
+#endif
+	FTCHARToUTF8 Utf8(Text);
+	fwrite(Utf8.Get(), 1, Utf8.Length(), stdout);
+	fflush(stdout);
+}
+
+static void PrintProgress(const FBatchProcessCoordinator& Coordinator, float JobsPerSec)
+{
+	const int32 Total = Coordinator.GetNumJobs();
+	if (Total == 0) { return; }
+
+	int32 Completed = 0;
+	int32 Failed    = 0;
+	for (int32 Idx = 0; Idx < Total; ++Idx)
+	{
+		const FBatchJobResult& Result = Coordinator.GetJobResult(Idx);
+		if (Result.State == EBatchJobState::Completed)
+		{
+			++Completed;
+			if (Result.bEncounteredError) { ++Failed; }
+		}
+	}
+	const int32 Succeeded     = Completed - Failed;
+	const int32 Pct           = (Completed * 100) / Total;
+	const int32 ActiveWorkers = Coordinator.GetNumActiveWorkers();
+	const int32 MaxWorkers    = Coordinator.GetNumWorkers();
+
+	constexpr int32 BarWidth = 30;
+	const int32 Filled = (Completed * BarWidth) / Total;
+	FString Bar;
+	Bar.Reserve(BarWidth);
+	for (int32 i = 0;      i < Filled;   ++i) { Bar += TEXT('#'); }
+	for (int32 i = Filled; i < BarWidth; ++i) { Bar += TEXT(' '); }
+
+	if (IsInteractiveTerminal())
+	{
+		// \r[bar]\033[K redraws the current line.
+		// The trailing \r leaves the cursor at column 0 so that any UE_LOG output
+		// that fires before the next update (e.g. "Worker N ready") starts at the
+		// left margin and overwrites the bar cleanly, rather than appending to it
+		// mid-line.  Between log messages the bar updates in-place; when a log
+		// fires, the bar scrolls one line.
+		WriteToTerminal(*FString::Printf(
+			TEXT("\r[%s] %4d/%-4d  %3d%%  |  ok: %d  fail: %d  w:%d/%d  %.1f/s\033[K\r"),
+			*Bar, Completed, Total, Pct, Succeeded, Failed, ActiveWorkers, MaxWorkers, JobsPerSec));
+		return;
+	}
+	UE_LOG(LogBatchProcess, Display,
+		TEXT("[%s] %4d/%-4d  %3d%%  |  ok: %d  fail: %d  w:%d/%d  %.1f/s"),
+		*Bar, Completed, Total, Pct, Succeeded, Failed, ActiveWorkers, MaxWorkers, JobsPerSec);
+}
+
+FBatchRunResult RunBatchFromJson(const FString& JobJson, const FBatchRunOptions& Options)
+{
+	const int32 ResolvedWorkers = Options.NumWorkers > 0
+		? Options.NumWorkers
+		: FMath::Max(1, FPlatformMisc::NumberOfCores() / 4);
+
+	const FString ProjectName    = FPaths::GetBaseFilename(FPaths::GetProjectFilePath());
+	const FString ExecutableName = FPaths::GetCleanFilename(FPlatformProcess::ExecutablePath());
+
+	// Persist the job JSON so the run is trivially reproducible.
+	const FString TempDir  = FPaths::ProjectSavedDir() / TEXT("BatchProcess");
+	const FString TempFile = TempDir / TEXT("last_run.json");
+	IFileManager::Get().MakeDirectory(*TempDir, /*Tree=*/true);
+	FFileHelper::SaveStringToFile(JobJson, *TempFile);
+
+	UE_LOG(LogBatchProcess, Display,
+		TEXT("%s %s -run=BatchProcessCommandlet \"%s\" -NumWorkers=%d"),
+		*ExecutableName, *ProjectName, *TempFile, ResolvedWorkers);
+
+	// Set up and run the coordinator.
+	const double SetupStart = FPlatformTime::Seconds();
+	FBatchProcessCoordinator Coordinator(Options.ClientAddress, ResolvedWorkers);
+	Coordinator.Submit(FUtf8String(TCHAR_TO_UTF8(*JobJson)));
+	Coordinator.Seal();
+	const double SetupElapsed = FPlatformTime::Seconds() - SetupStart;
+
+	const int32 NumJobs = Coordinator.GetNumJobs();
+
+	UE_LOG(LogBatchProcess, Display,
+		TEXT("Starting %d job(s) across %d worker(s)  (setup: %.2fs)"),
+		NumJobs, Coordinator.GetNumWorkers(), SetupElapsed);
+
+	const double RunStart = FPlatformTime::Seconds();
+
+	// Escape the function path once for use in per-failure reproduce commands.
+	const FString EscapedFunctionPath = Coordinator.GetFunctionAsString()
+		.Replace(TEXT("\\"), TEXT("\\\\"))
+		.Replace(TEXT("\""), TEXT("\\\""));
+
+	// Track which completed jobs have already been reported so we never double-print.
+	TArray<bool> bReported;
+	bReported.Init(false, NumJobs);
+	int32 FailsPrinted = 0;
+	int32 TotalFailed  = 0;
+
+	// Scans all jobs for newly-completed failures and prints the first 10 immediately.
+	auto ReportNewFailures = [&]()
+	{
+		for (int32 Idx = 0; Idx < NumJobs; ++Idx)
+		{
+			if (bReported[Idx]) { continue; }
+			const FBatchJobResult& Job = Coordinator.GetJobResult(Idx);
+			if (Job.State != EBatchJobState::Completed) { continue; }
+			bReported[Idx] = true;
+			if (!Job.bEncounteredError) { continue; }
+			++TotalFailed;
+
+			if (FailsPrinted < 10)
+			{
+				// Only print the first line of output — the full capture is in the
+				// results file.  Dumping the entire multi-line string via a single
+				// UE_LOG causes embedded newlines to leave the cursor mid-line.
+				const FString FirstOutputLine = [&]() -> FString
+				{
+					if (Job.Output.IsEmpty())
+					{
+						return FString::Printf(TEXT("(no output) arg: %s"), *Job.ArgumentSummary);
+					}
+					const int32 NL = Job.Output.Find(TEXT("\n"));
+					return NL != INDEX_NONE ? Job.Output.Left(NL) : Job.Output;
+				}();
+				UE_LOG(LogBatchProcess, Error, TEXT("[FAIL job %4d]  %s"), Idx, *FirstOutputLine);
+
+				const FString CommandJson = FString::Printf(
+					TEXT("{\"Function\":\"%s\",\"Arguments\":[%s]}"),
+					*EscapedFunctionPath, *Job.ArgumentSummary);
+				const FString EscapedCommandJson = CommandJson.Replace(TEXT("\""), TEXT("\\\""));
+				UE_LOG(LogBatchProcess, Display,
+					TEXT("Repro command:  %s %s -run=BatchWorkerCommandlet -singlecommand=%s"),
+					*ExecutableName, *ProjectName, *EscapedCommandJson);
+				++FailsPrinted;
+			}
+		}
+	};
+
+	double LastProgressTime    = RunStart;
+	int32  LastCompletedForRate = 0;
+	while (Coordinator.Tick())
+	{
+		FPlatformProcess::Sleep(0.016f);
+		const double Now = FPlatformTime::Seconds();
+		if (Now - LastProgressTime >= 0.5)
+		{
+			// Count completions and compute rate over the elapsed interval.
+			int32 CurrentCompleted = 0;
+			for (int32 Idx = 0; Idx < NumJobs; ++Idx)
+			{
+				if (Coordinator.GetJobResult(Idx).State == EBatchJobState::Completed)
+				{
+					++CurrentCompleted;
+				}
+			}
+			const double Elapsed   = Now - LastProgressTime;
+			const float  JobsPerSec = (float)((CurrentCompleted - LastCompletedForRate) / Elapsed);
+
+			LastCompletedForRate = CurrentCompleted;
+			LastProgressTime     = Now;
+
+			ReportNewFailures();
+			PrintProgress(Coordinator, JobsPerSec);
+		}
+	}
+
+	// Final pass: report remaining failures, then advance past the bar line.
+	ReportNewFailures();
+	if (IsInteractiveTerminal()) { WriteToTerminal(TEXT("\n")); }
+
+	if (TotalFailed > 10)
+	{
+		UE_LOG(LogBatchProcess, Display, TEXT("  ... and %d more failures"), TotalFailed - 10);
+	}
+
+	const double TotalElapsed = FPlatformTime::Seconds() - RunStart;
+	const double JobsPerSec   = TotalElapsed > 0.0 ? NumJobs / TotalElapsed : 0.0;
+	UE_LOG(LogBatchProcess, Display,
+		TEXT("Completed %d job(s) in %.1fs  (%.1f jobs/sec across %d worker(s))"),
+		NumJobs, TotalElapsed, JobsPerSec, Coordinator.GetNumWorkers());
+
+	// Collect results.
+	FBatchRunResult Result;
+	Result.FunctionPath           = Coordinator.GetFunctionAsString();
+	Result.TempFilePath           = TempFile;
+	Result.bHasCoordinatorErrors  = Coordinator.HasAnyErrors();
+	Result.CoordinatorErrorString = Coordinator.GetErrorString();
+
+	Result.Jobs.Reserve(NumJobs);
+	for (int32 Idx = 0; Idx < NumJobs; ++Idx)
+	{
+		Result.Jobs.Add(Coordinator.GetJobResult(Idx));
+	}
+
+	return Result;
+}
+
+} // namespace UE::Private
diff --git a/Engine/Source/Editor/Kismet/Private/BatchProcess/BatchProcessRunner.h b/Engine/Source/Editor/Kismet/Private/BatchProcess/BatchProcessRunner.h
new file mode 100644
--- /dev/null
+++ b/Engine/Source/Editor/Kismet/Private/BatchProcess/BatchProcessRunner.h
@@ -0,0 +1,44 @@
+// Copyright Epic Games, Inc. All Rights Reserved.
+
+#pragma once
+
+#include "CoreMinimal.h"
+#include "BatchProcess/BatchProcessCoordinator.h"
+
+namespace UE::Private
+{
+
+struct FBatchRunOptions
+{
+	/** Number of workers to spawn. 0 = auto-select based on core count. */
+	int32 NumWorkers = 0;
+
+	/**
+	 * When non-null, launches the coordinator in listen mode so a single
+	 * external client can connect for debugging (-launchforclient).
+	 */
+	const FString* ClientAddress = nullptr;
+};
+
+struct FBatchRunResult
+{
+	TArray<FBatchJobResult> Jobs;        // one per submitted job, in submission order
+	FString FunctionPath;                // the UFUNCTION path parsed from the job JSON
+	FString TempFilePath;                // Saved/BatchProcess/last_run.json written by the runner
+	bool    bHasCoordinatorErrors = false;
+	FString CoordinatorErrorString;
+};
+
+/**
+ * Core batch execution helper shared by BatchProcessCommandlet and BatchProcessLibrary.
+ *
+ * Given a raw JSON job description string this function:
+ *   1. Persists the JSON to Saved/BatchProcess/last_run.json and prints a rerun command.
+ *   2. Spins up an FBatchProcessCoordinator, submits all jobs, and runs the tick loop.
+ *   3. Displays a live progress bar in interactive consoles.
+ *   4. Prints the first 10 failures as errors, each with a -singlecommand reproduce line.
+ *   5. Returns all results for the caller to handle (results file, exit code, JSON, etc.).
+ */
+FBatchRunResult RunBatchFromJson(const FString& JobJson, const FBatchRunOptions& Options = FBatchRunOptions());
+
+} // namespace UE::Private
diff --git a/Engine/Source/Editor/Kismet/Private/BatchProcess/BatchProcessSocketHelpers.cpp b/Engine/Source/Editor/Kismet/Private/BatchProcess/BatchProcessSocketHelpers.cpp
new file mode 100644
--- /dev/null
+++ b/Engine/Source/Editor/Kismet/Private/BatchProcess/BatchProcessSocketHelpers.cpp
@@ -0,0 +1,159 @@
+// Copyright Epic Games, Inc. All Rights Reserved.
+
+#include "BatchProcess/BatchProcessSocketHelpers.h"
+
+#include "Dom/JsonValue.h"
+#include "Serialization/JsonReader.h"
+#include "Serialization/JsonSerializer.h"
+#include "Sockets.h"
+
+bool UE::Private::BatchProcessSocketHelpers::SendUint32(FSocket& Socket, uint32 Arg)
+{
+	uint8 Data[sizeof(uint32)];
+	Data[0] = Arg >> 24;
+	Data[1] = Arg >> 16;
+	Data[2] = Arg >> 8;
+	Data[3] = Arg >> 0;
+	for (int32 Sent = 0; Sent != sizeof(uint32);)
+	{
+		int32 BytesSent;
+		if (!Socket.Send(Data + Sent, sizeof(uint32) - Sent, BytesSent))
+		{
+			return false;
+		}
+		if (BytesSent == 0)
+		{
+			return false;
+		}
+		Sent += BytesSent;
+	}
+	return true;
+}
+
+bool UE::Private::BatchProcessSocketHelpers::RecvUint32(FSocket& Socket, uint32& Arg)
+{
+	uint8 Data[sizeof(uint32)];
+	for (int32 Read = 0; Read != sizeof(uint32);)
+	{
+		int32 BytesRead;
+		if (!Socket.Recv(Data + Read, sizeof(uint32) - Read, BytesRead))
+		{
+			return false;
+		}
+		if (BytesRead == 0)
+		{
+			return false;
+		}
+		Read += BytesRead;
+	}
+	uint32 Result = 0;
+	Result |= (uint32)Data[0] << 24;
+	Result |= (uint32)Data[1] << 16;
+	Result |= (uint32)Data[2] << 8;
+	Result |= (uint32)Data[3] << 0;
+	Arg = Result;
+	return true;
+}
+
+bool UE::Private::BatchProcessSocketHelpers::RecvCharArray(FSocket& Socket, uint32 Num, TUniquePtr<char[]>& OutCharArray)
+{
+	OutCharArray = ::MakeUnique<char[]>(Num);
+	for (uint32 Read = 0; Read != Num;)
+	{
+		int32 BytesRead;
+		if (!Socket.Recv(reinterpret_cast<uint8*>(OutCharArray.Get()) + Read, static_cast<int32>(Num - Read), BytesRead))
+		{
+			return false;
+		}
+		if (BytesRead == 0)
+		{
+			return false;
+		}
+		Read += BytesRead;
+	}
+	return true;
+}
+
+bool UE::Private::BatchProcessSocketHelpers::RecvString(FSocket& Socket, FString& OutString)
+{
+	uint32 Num = 0;
+	if (!RecvUint32(Socket, Num))
+	{
+		return false;
+	}
+	if(Num > (uint32)(TNumericLimits<int32>::Max()))
+	{
+		return false;
+	}
+	TUniquePtr<char[]> CharArray;
+	if (!RecvCharArray(Socket, Num, CharArray))
+	{
+		return false;
+	}
+	OutString = FString(FUtf8StringView(CharArray.Get(), (int32)Num));
+	return true;
+}
+
+bool UE::Private::BatchProcessSocketHelpers::SendJsonValue(FSocket& Socket, const TSharedRef<FJsonObject>& Object)
+{
+	FString String;
+	auto Writer = TJsonWriterFactory<>::Create(&String);
+	if (!FJsonSerializer::Serialize(Object, Writer))
+	{
+		return false;
+	}
+	if (!SendString(Socket, String))
+	{
+		return false;
+	}
+	return true;
+}
+
+bool UE::Private::BatchProcessSocketHelpers::RecvJsonValue(FSocket& Socket, TSharedPtr<FJsonValue>& Value)
+{
+	uint32 Num = 0;
+	if (!RecvUint32(Socket, Num))
+	{
+		return false;
+	}
+	if(Num > (uint32)(TNumericLimits<int32>::Max()))
+	{
+		return false;
+	}
+	TUniquePtr<char[]> Array;
+	if (!RecvCharArray(Socket, Num, Array))
+	{
+		return false;
+	}
+	TStringView<char> StringView{Array.Get(), (int32)Num};
+	TSharedRef<TJsonStringViewReader<char>> Reader = TJsonStringViewReader<char>::Create(StringView);
+	if (!FJsonSerializer::Deserialize(*Reader, Value))
+	{
+		return false;
+	}
+	return true;
+}
+
+bool UE::Private::BatchProcessSocketHelpers::SendString(FSocket& Socket, const FString& String)
+{
+	auto Utf8String = ::StringCast<UTF8CHAR>(*String);
+	int32 Len = FCStringUtf8::Strlen(Utf8String.Get());
+	if (!SendUint32(Socket, static_cast<uint32>(Len)))
+	{
+		return false;
+	}
+	for (int32 Sent = 0; Sent != Len;)
+	{
+		int32 BytesSent;
+		if (!Socket.Send(reinterpret_cast<const uint8*>(Utf8String.Get()) + Sent, Len - Sent, BytesSent))
+		{
+			return false;
+		}
+		if (BytesSent == 0)
+		{
+			return false;
+		}
+		Sent += BytesSent;
+	}
+	return true;
+}
diff --git a/Engine/Source/Editor/Kismet/Private/BatchProcess/BatchProcessSocketHelpers.h b/Engine/Source/Editor/Kismet/Private/BatchProcess/BatchProcessSocketHelpers.h
new file mode 100644
--- /dev/null
+++ b/Engine/Source/Editor/Kismet/Private/BatchProcess/BatchProcessSocketHelpers.h
@@ -0,0 +1,23 @@
+// Copyright Epic Games, Inc. All Rights Reserved.
+
+#pragma once
+
+#include "CoreMinimal.h"
+
+#include "Dom/JsonObject.h"
+
+class FSocket;
+
+namespace UE::Private::BatchProcessSocketHelpers
+{
+	/** Port the coordinator listens on and workers connect to. */
+	constexpr int32 DefaultPortNumber = 61579;
+
+	bool SendUint32(FSocket& Socket, uint32 Arg);
+	bool RecvUint32(FSocket& Socket, uint32& Num);
+	bool RecvCharArray(FSocket& Socket, uint32 Num, TUniquePtr<char[]>& OutCharArray);
+	bool RecvString(FSocket& Socket, FString& OutString);
+	bool SendJsonValue(FSocket& Socket, const TSharedRef<FJsonObject>& Object);
+	bool RecvJsonValue(FSocket& Socket, TSharedPtr<FJsonValue>& Value);
+	bool SendString(FSocket& Socket, const FString& String);
+}
diff --git a/Engine/Source/Editor/Kismet/Private/BatchProcess/BatchProcessWorker.cpp b/Engine/Source/Editor/Kismet/Private/BatchProcess/BatchProcessWorker.cpp
new file mode 100644
--- /dev/null
+++ b/Engine/Source/Editor/Kismet/Private/BatchProcess/BatchProcessWorker.cpp
@@ -0,0 +1,434 @@
+// Copyright Epic Games, Inc. All Rights Reserved.
+
+#include "BatchProcess/BatchProcessWorker.h"
+
+#include "HAL/FileManager.h"
+#include "HAL/PlatformMemory.h"
+#include "JsonObjectConverter.h"
+#include "Misc/App.h"
+#include "Serialization/JsonReader.h"
+#include "Serialization/JsonSerializer.h"
+#include "Misc/Paths.h"
+#include "BatchProcess/BatchProcessLog.h"
+#include "BatchProcess/BatchProcessSocketHelpers.h"
+#include "UObject/SoftObjectPtr.h"
+#include "UObject/StructOnScope.h"
+#include "UObject/UObjectGlobals.h"
+#include "SocketSubsystem.h"
+#include "Sockets.h"
+
+
+/** Captures all UE_LOG output produced during a single job's ProcessEvent call. */
+struct FJobOutputCapture : public FOutputDevice
+{
+	FString Buffer;
+	bool    bHasError = false;
+
+	virtual void Serialize(const TCHAR* V, ELogVerbosity::Type Verbosity, const FName& Category) override
+	{
+		if (Verbosity <= ELogVerbosity::Error)
+		{
+			bHasError = true;
+		}
+		if (!Buffer.IsEmpty())
+		{
+			Buffer += TEXT('\n');
+		}
+		Buffer += FString::Printf(TEXT("[%s] %s"), *Category.ToString(), V);
+	}
+	virtual bool CanBeUsedOnAnyThread() const override { return true; }
+};
+
+namespace UE::Private
+{
+	const TCHAR* const BatchWorkerName    = TEXT("BatchProcessWorker");
+	// Facilities for running a single client for debugging; runs
+	// a coordinator for consistency/completeness:
+	const TCHAR* const DebugClientAddress = TEXT("127.0.0.1");
+	const uint32 DebugClientIP            = 0x7f000001;
+	const int32  DebugClientPort          = 61580;
+	const int32  MaxDebugConnections      = 128;
+}
+
+FBatchProcessWorker::FBatchProcessWorker()
+	: ClientState(EClientState::Starting)
+{
+}
+
+FBatchProcessWorker::~FBatchProcessWorker()
+{
+	HandleCleanup();
+}
+
+bool FBatchProcessWorker::LaunchLocalCoordinator(const FString& Command, FString& OutCoordinatorAddr, const FString& ExtraArgs)
+{
+	FString ProjPath = FPaths::GetProjectFilePath();
+
+	const FString LaunchExecutableName = FPlatformProcess::ExecutablePath();
+	const FString Params = FString::Printf(TEXT("\"%s\" -run=BatchProcessCommandlet \"%s\" -launchforclient=%s %s"),
+		*ProjPath, *Command, UE::Private::DebugClientAddress, *ExtraArgs);
+	LocalCoordinator = FPlatformProcess::CreateProc(
+		*LaunchExecutableName,
+		*Params,
+		false,
+		true,
+		true,
+		nullptr,
+		0,
+		nullptr,
+		nullptr);
+
+	// Wait for the coordinator to start — it will connect back to us with its address,
+	// then execution continues as normal.
+	ISocketSubsystem* SocketSubsystem = ISocketSubsystem::Get();
+	check(SocketSubsystem);
+	TSharedPtr<FInternetAddr> ListenAddr = SocketSubsystem->GetLocalBindAddr(*GLog);
+	ListenAddr->SetIp(UE::Private::DebugClientIP);
+	ListenAddr->SetPort(UE::Private::DebugClientPort);
+	FSocket* LocalCoordinatorSocket = SocketSubsystem->CreateSocket(NAME_Stream, UE::Private::BatchWorkerName, ListenAddr->GetProtocolType());
+	LocalCoordinatorSocket->SetReuseAddr();
+
+	const bool bBindSuccessful = LocalCoordinatorSocket->Bind(*ListenAddr);
+	if(!bBindSuccessful)
+	{
+		UE_LOG(LogBatchProcess, Error, TEXT("Failed to bind debug coordinator listener"));
+		return false;
+	}
+
+	const bool bListenSuccessful = LocalCoordinatorSocket->Listen(UE::Private::MaxDebugConnections);
+	if(!bListenSuccessful)
+	{
+		UE_LOG(LogBatchProcess, Error, TEXT("Failed to listen for debug coordinator"));
+		return false;
+	}
+
+	// Wait for coordinator to start up — generous timeout to allow the editor to initialise.
+	const float  Timeout   = 180.f; // 3 minutes
+	const double StartTime = FPlatformTime::Seconds();
+	while(true)
+	{
+		if((StartTime + Timeout) < FPlatformTime::Seconds())
+		{
+			UE_LOG(LogBatchProcess, Error, TEXT("Timed out waiting for debug coordinator to start"));
+			return false;
+		}
+		else if(!FPlatformProcess::IsProcRunning(LocalCoordinator))
+		{
+			UE_LOG(LogBatchProcess, Error, TEXT("Debug coordinator exited unexpectedly"));
+			return false;
+		}
+
+		bool bHasPendingConnection = false;
+		if(LocalCoordinatorSocket->HasPendingConnection(bHasPendingConnection) && bHasPendingConnection)
+		{
+			FSocket* DummyCoordinatorClientSocket = LocalCoordinatorSocket->Accept(UE::Private::BatchWorkerName);
+			check(DummyCoordinatorClientSocket);
+			check(UE::Private::BatchProcessSocketHelpers::RecvString(*DummyCoordinatorClientSocket, OutCoordinatorAddr));
+			SocketSubsystem->DestroySocket(DummyCoordinatorClientSocket);
+			SocketSubsystem->DestroySocket(LocalCoordinatorSocket);
+			break;
+		}
+
+		FPlatformProcess::Sleep(1.f/60.f);
+	}
+
+	return true;
+}
+
+bool FBatchProcessWorker::PerformWork(const FString& CoordinatorAddrStr, int32 WorkerId, uint32 CoordinatorPid)
+{
+	if (!Connect(CoordinatorAddrStr, WorkerId))
+	{
+		ClientState = EClientState::Error;
+		return false;
+	}
+
+	// Coordinator should notice us connect and send a work description.
+	while(true)
+	{
+		TSharedPtr<FJsonValue> Value;
+		if (!UE::Private::BatchProcessSocketHelpers::RecvJsonValue(*ClientSocket, Value))
+		{
+			// Coordinator closed the connection.  If we were idle (finished a
+			// batch and waiting for more), this is the normal completion signal —
+			// the coordinator closes worker sockets once all jobs are done.
+			const bool bSuccess = (ClientState == EClientState::Idle);
+			HandleCleanup();
+			ClientState = bSuccess ? EClientState::ShuttingDown : EClientState::Error;
+			return bSuccess;
+		}
+
+		if(!Value)
+		{
+			HandleCleanup();
+			ClientState = EClientState::Error;
+			return false;
+		}
+
+		const TSharedPtr<FJsonObject>* JsonObjectPtr;
+		if (!Value->TryGetObject(JsonObjectPtr))
+		{
+			HandleCleanup();
+			checkf(false, TEXT("Failed to find work object"));
+			ClientState = EClientState::Error;
+			return false;
+		}
+
+		const FJsonObject& JsonObject = *(JsonObjectPtr->Get());
+
+		bool bShutdown = false;
+		if(JsonObject.TryGetBoolField(TEXT("Shutdown"), bShutdown))
+		{
+			ClientState = EClientState::ShuttingDown;
+			break;
+		}
+
+		FString FunctionAsString;
+		if(!JsonObject.TryGetStringField(TEXT("Function"), FunctionAsString))
+		{
+			HandleCleanup();
+			checkf(false, TEXT("Failed to find function in work packet"));
+			ClientState = EClientState::Error;
+			return false;
+		}
+		const TArray<TSharedPtr<FJsonValue>>* ArgumentsField;
+		if(!JsonObject.TryGetArrayField(TEXT("Arguments"), ArgumentsField))
+		{
+			HandleCleanup();
+			checkf(false, TEXT("Failed to find arguments in work packet"));
+			ClientState = EClientState::Error;
+			return false;
+		}
+
+		TSoftObjectPtr<UFunction> FunctionSoftObjectPtr{FSoftObjectPath{FunctionAsString}};
+		const UFunction* ResolvedFunction = FunctionSoftObjectPtr.LoadSynchronous();
+		if(!ResolvedFunction)
+		{
+			HandleCleanup();
+			checkf(false, TEXT("Failed to resolve function: %s"), *FunctionAsString);
+			ClientState = EClientState::Error;
+			return false;
+		}
+
+		// The coordinator tells us the global starting index for this batch so we can
+		// report the correct job indices back regardless of how many batches have run.
+		double StartIndex = 0.0;
+		JsonObject.TryGetNumberField(TEXT("StartIndex"), StartIndex);
+
+		ClientState = EClientState::Working;
+		UObject* Obj = ResolvedFunction->GetClass()->GetDefaultObject();
+
+		for(int32 Idx = 0; Idx < ArgumentsField->Num(); ++Idx)
+		{
+			const TSharedPtr<FJsonValue>& Argument = (*ArgumentsField)[Idx];
+			FStructOnScope Arg(ResolvedFunction);
+
+			FText FailReason;
+			const bool bResult =
+				FJsonObjectConverter::JsonObjectToUStruct(Argument->AsObject().ToSharedRef(),
+					ResolvedFunction, Arg.GetStructMemory(),
+					0, 0, false,
+					&FailReason);
+
+			if(!bResult)
+			{
+				HandleCleanup();
+				checkf(false, TEXT("Failed to parse job argument: %s"), *FailReason.ToString());
+				ClientState = EClientState::Error;
+				return false;
+			}
+
+			// Capture all log output produced during this job so the coordinator
+			// can surface it in the structured results report.
+			FJobOutputCapture Capture;
+			GLog->AddOutputDevice(&Capture);
+			Obj->ProcessEvent(const_cast<UFunction*>(ResolvedFunction), Arg.GetStructMemory());
+			// Python's unreal.log() calls may be buffered inside GLog's threaded-log
+			// queue rather than dispatched immediately.  Flush before removing the
+			// capture device so no output is lost.
+			GLog->FlushThreadedLogs();
+			GLog->RemoveOutputDevice(&Capture);
+
+			// Report result back to coordinator, including captured output.
+			TSharedRef<FJsonObject> RootObject = ::MakeShared<FJsonObject>();
+			RootObject->SetNumberField(TEXT("JobIndex"), (double)((int32)StartIndex + Idx));
+			RootObject->SetStringField(TEXT("Output"),   Capture.Buffer);
+			RootObject->SetBoolField(  TEXT("Error"),    Capture.bHasError);
+			if (!UE::Private::BatchProcessSocketHelpers::SendJsonValue(*ClientSocket, RootObject))
+			{
+				HandleCleanup();
+				ClientState = EClientState::Error;
+				return false;
+			}
+
+			// Detect coordinator death between jobs. RecvJsonValue handles the case
+			// where the coordinator dies while we're waiting for the next batch, but
+			// here we cover the gap where the coordinator was killed during ProcessEvent
+			// and the socket send above succeeded against buffered data.
+			if (CoordinatorPid != 0 && !FPlatformProcess::IsApplicationRunning(CoordinatorPid))
+			{
+				UE_LOG(LogBatchProcess, Warning, TEXT("Coordinator process has exited; shutting down"));
+				HandleCleanup();
+				ClientState = EClientState::Error;
+				return false;
+			}
+		}
+
+		// Batch complete.  Go back to the top of the loop and wait for the
+		// coordinator to either send another batch or close the connection.
+		ClientState = EClientState::Idle;
+
+		// Trigger GC when physical memory has grown by GCThresholdBytes since the
+		// last collection.  This adapts to actual workload rather than firing on a
+		// fixed schedule: lightweight jobs may never trigger GC; heavy jobs that
+		// allocate lots of transient assets trigger it as needed.
+		static uint64 MemAtLastGC = 0;
+		constexpr uint64 GCThresholdBytes = 6ull * 1024 * 1024 * 1024; // 6 GB
+		const uint64 CurrentMem = FPlatformMemory::GetStats().UsedPhysical;
+		if (MemAtLastGC == 0) { MemAtLastGC = CurrentMem; }
+		if (CurrentMem > MemAtLastGC + GCThresholdBytes)
+		{
+			CollectGarbage(RF_NoFlags);
+			MemAtLastGC = FPlatformMemory::GetStats().UsedPhysical;
+		}
+	}
+
+	// Inform server we've shut down.
+	TSharedRef<FJsonObject> RootObject = ::MakeShared<FJsonObject>();
+	RootObject->SetBoolField(TEXT("Shutdown"), true);
+	if (!UE::Private::BatchProcessSocketHelpers::SendJsonValue(*ClientSocket, RootObject))
+	{
+		UE_LOG(LogBatchProcess, Warning, TEXT("Failed to send shutdown message to coordinator"));
+	}
+
+	ClientSocket->Close();
+	ISocketSubsystem::Get()->DestroySocket(ClientSocket);
+	ClientSocket = nullptr; // prevent double-destroy in HandleCleanup/destructor
+	return true;
+}
+
+bool FBatchProcessWorker::RunSingleCommand(const FString& JobJson)
+{
+	TSharedPtr<FJsonValue> Value;
+	TSharedRef<TJsonReader<>> Reader = TJsonReaderFactory<>::Create(JobJson);
+	if (!FJsonSerializer::Deserialize(Reader, Value) || !Value.IsValid())
+	{
+		UE_LOG(LogBatchProcess, Error, TEXT("--singlecommand: failed to parse JSON"));
+		return false;
+	}
+
+	const TSharedPtr<FJsonObject>* RootObjPtr;
+	if (!Value->TryGetObject(RootObjPtr))
+	{
+		UE_LOG(LogBatchProcess, Error, TEXT("--singlecommand: JSON root must be an object"));
+		return false;
+	}
+	const FJsonObject& RootObj = **RootObjPtr;
+
+	FString FunctionPath;
+	if (!RootObj.TryGetStringField(TEXT("Function"), FunctionPath))
+	{
+		UE_LOG(LogBatchProcess, Error, TEXT("--singlecommand: missing 'function' field"));
+		return false;
+	}
+
+	const TArray<TSharedPtr<FJsonValue>>* ArgumentsArray;
+	if (!RootObj.TryGetArrayField(TEXT("Arguments"), ArgumentsArray) || ArgumentsArray->IsEmpty())
+	{
+		UE_LOG(LogBatchProcess, Error, TEXT("--singlecommand: missing or empty 'arguments' array"));
+		return false;
+	}
+
+	TSoftObjectPtr<UFunction> FunctionSoftPtr{FSoftObjectPath{FunctionPath}};
+	const UFunction* ResolvedFunction = FunctionSoftPtr.LoadSynchronous();
+	if (!ResolvedFunction)
+	{
+		UE_LOG(LogBatchProcess, Error, TEXT("--singlecommand: failed to resolve function: %s"), *FunctionPath);
+		return false;
+	}
+
+	UObject* CDO = ResolvedFunction->GetClass()->GetDefaultObject();
+
+	// Use a capture device solely to detect errors; output still flows to all
+	// normal GLog devices so the user sees it in the console/log file.
+	FJobOutputCapture Capture;
+	GLog->AddOutputDevice(&Capture);
+	for (const TSharedPtr<FJsonValue>& Argument : *ArgumentsArray)
+	{
+		FStructOnScope ArgScope(ResolvedFunction);
+		FText FailReason;
+		if (!FJsonObjectConverter::JsonObjectToUStruct(
+			Argument->AsObject().ToSharedRef(),
+			ResolvedFunction, ArgScope.GetStructMemory(),
+			0, 0, false, &FailReason))
+		{
+			GLog->FlushThreadedLogs();
+			GLog->RemoveOutputDevice(&Capture);
+			UE_LOG(LogBatchProcess, Error, TEXT("--singlecommand: failed to parse argument: %s"), *FailReason.ToString());
+			return false;
+		}
+		CDO->ProcessEvent(const_cast<UFunction*>(ResolvedFunction), ArgScope.GetStructMemory());
+	}
+	GLog->FlushThreadedLogs();
+	GLog->RemoveOutputDevice(&Capture);
+
+	return !Capture.bHasError;
+}
+
+void FBatchProcessWorker::HandleCleanup()
+{
+	if(ClientSocket)
+	{
+		ISocketSubsystem* SocketSubsystem = ISocketSubsystem::Get();
+		check(SocketSubsystem);
+		SocketSubsystem->DestroySocket(ClientSocket);
+		ClientSocket = nullptr;
+	}
+
+	if(LocalCoordinator.IsValid())
+	{
+		FPlatformProcess::TerminateProc(LocalCoordinator);
+		LocalCoordinator = FProcHandle();
+	}
+}
+
+bool FBatchProcessWorker::Connect(const FString& CoordinatorAddrStr, int32 WorkerId)
+{
+	ISocketSubsystem* SocketSubsystem = ISocketSubsystem::Get();
+	if (!SocketSubsystem)
+	{
+		UE_LOG(LogBatchProcess, Error, TEXT("No socket subsystem available"));
+		return false;
+	}
+
+	TSharedPtr<class FInternetAddr> CoordinatorAddr = SocketSubsystem->GetLocalBindAddr(*GLog);
+	bool bIpValid = false;
+	// CoordinatorAddrStr carries only the IP; SetIp() does not parse a port from the
+	// string, so we set it explicitly using the shared DefaultPortNumber constant.
+	CoordinatorAddr->SetIp(*CoordinatorAddrStr, bIpValid);
+	CoordinatorAddr->SetPort(UE::Private::BatchProcessSocketHelpers::DefaultPortNumber);
+	if (!bIpValid)
+	{
+		UE_LOG(LogBatchProcess, Error, TEXT("Invalid coordinator IP: %s"), *CoordinatorAddrStr);
+		return false;
+	}
+
+	ClientSocket = SocketSubsystem->CreateSocket(NAME_Stream, UE::Private::BatchWorkerName, CoordinatorAddr->GetProtocolType());
+	if (!ClientSocket->Connect(*CoordinatorAddr))
+	{
+		UE_LOG(LogBatchProcess, Error, TEXT("Failed to connect to coordinator at %s"), *CoordinatorAddrStr);
+		SocketSubsystem->DestroySocket(ClientSocket);
+		ClientSocket = nullptr;
+		return false;
+	}
+
+	if (!UE::Private::BatchProcessSocketHelpers::SendUint32(*ClientSocket, WorkerId))
+	{
+		UE_LOG(LogBatchProcess, Error, TEXT("Failed to send handshake ID to coordinator"));
+		SocketSubsystem->DestroySocket(ClientSocket);
+		ClientSocket = nullptr;
+		return false;
+	}
+
+	return true;
+}
diff --git a/Engine/Source/Editor/Kismet/Private/BatchProcess/BatchProcessWorker.h b/Engine/Source/Editor/Kismet/Private/BatchProcess/BatchProcessWorker.h
new file mode 100644
--- /dev/null
+++ b/Engine/Source/Editor/Kismet/Private/BatchProcess/BatchProcessWorker.h
@@ -0,0 +1,49 @@
+// Copyright Epic Games, Inc. All Rights Reserved.
+
+#pragma once
+
+#include "CoreMinimal.h"
+
+#include "HAL/PlatformProcess.h"
+
+class FSocket;
+
+struct FBatchProcessWorker
+{
+	FBatchProcessWorker();
+	~FBatchProcessWorker();
+
+	bool LaunchLocalCoordinator(const FString& Command, FString& OutCoordinatorAddr, const FString& ExtraArgs = FString());
+	bool PerformWork(const FString& CoordinatorAddrStr, int32 WorkerId, uint32 CoordinatorPid = 0);
+	/** Returns the handle for the coordinator process started by LaunchLocalCoordinator, if any. */
+	FProcHandle GetLocalCoordinatorHandle() const { return LocalCoordinator; }
+
+	/**
+	 * Run a single job inline without connecting to a coordinator.
+	 * JobJson must have the form: {"function":"<path>","arguments":[{...}]}
+	 * Output flows to the normal log; returns false if any error was logged.
+	 */
+	static bool RunSingleCommand(const FString& JobJson);
+
+private:
+	void HandleCleanup();
+	bool Connect(const FString& CoordinatorAddrStr, int32 WorkerId);
+
+	FBatchProcessWorker(const FBatchProcessWorker&) = delete;
+	FBatchProcessWorker(FBatchProcessWorker&&) = delete;
+	FBatchProcessWorker& operator=(const FBatchProcessWorker&) = delete;
+	FBatchProcessWorker& operator=(FBatchProcessWorker&&) = delete;
+
+	FSocket*     ClientSocket    = nullptr;
+	FProcHandle  LocalCoordinator; // debug facility: client starts its own coordinator for testing
+
+	enum class EClientState
+	{
+		Starting,
+		Working,
+		Idle,
+		ShuttingDown,
+		Error
+	};
+	EClientState ClientState;
+};
diff --git a/Engine/Source/Editor/Kismet/Private/Commandlets/BatchProcessCommandlet.cpp b/Engine/Source/Editor/Kismet/Private/Commandlets/BatchProcessCommandlet.cpp
new file mode 100644
--- /dev/null
+++ b/Engine/Source/Editor/Kismet/Private/Commandlets/BatchProcessCommandlet.cpp
@@ -0,0 +1,102 @@
+// Copyright Epic Games, Inc. All Rights Reserved.
+
+#include "Commandlets/BatchProcessCommandlet.h"
+
+#include "HAL/FileManager.h"
+#include "Misc/FileHelper.h"
+#include "Misc/Paths.h"
+#include "BatchProcess/BatchProcessLog.h"
+#include "BatchProcess/BatchProcessRunner.h"
+
+static void PrintBatchProcessCommandletUsage()
+{
+	UE_LOG(LogBatchProcess, Display,
+		TEXT("Repro command: UnrealEditor ProjectName -run=BatchProcessCommandlet Jobs.json [-resultsfile=Path/To/results.txt] [-numworkers=N]"));
+}
+
+UBatchProcessCommandlet::UBatchProcessCommandlet(const FObjectInitializer& ObjectInitializer)
+	: Super(ObjectInitializer)
+{
+}
+
+int32 UBatchProcessCommandlet::Main(const FString& Params)
+{
+	TArray<FString> Tokens;
+	TArray<FString> Switches;
+	TMap<FString, FString> ParamsMap;
+	UCommandlet::ParseCommandLine(*Params, Tokens, Switches, ParamsMap);
+
+	if (Tokens.Num() != 1)
+	{
+		PrintBatchProcessCommandletUsage();
+		return -1;
+	}
+
+	FString JobJson;
+	if (!FFileHelper::LoadFileToString(JobJson, *Tokens[0]))
+	{
+		UE_LOG(LogBatchProcess, Error, TEXT("Failed to read job file: %s"), *Tokens[0]);
+		return -1;
+	}
+
+	UE::Private::FBatchRunOptions Options;
+	Options.ClientAddress = ParamsMap.Find(TEXT("launchforclient"));
+	Options.NumWorkers    = ParamsMap.Contains(TEXT("numworkers"))
+		? FMath::Max(1, FCString::Atoi(*ParamsMap[TEXT("numworkers")]))
+		: 0; // 0 = auto
+
+	const UE::Private::FBatchRunResult RunResult = UE::Private::RunBatchFromJson(JobJson, Options);
+
+	if (RunResult.bHasCoordinatorErrors)
+	{
+		UE_LOG(LogBatchProcess, Error, TEXT("%s"), *RunResult.CoordinatorErrorString);
+		return -1;
+	}
+
+	// Write a structured results file — failures first so they're easy to spot.
+	FString ResultsFile;
+	if (const FString* Override = ParamsMap.Find(TEXT("resultsfile")))
+	{
+		ResultsFile = *Override;
+	}
+	else
+	{
+		ResultsFile = FPaths::ProjectSavedDir() / TEXT("MultiprocessResults") / TEXT("results.txt");
+	}
+	IFileManager::Get().MakeDirectory(*FPaths::GetPath(ResultsFile), /*Tree=*/true);
+
+	int32 NumFailed = 0;
+	TArray<FString> FailLines;
+	TArray<FString> PassLines;
+	for (int32 Idx = 0; Idx < RunResult.Jobs.Num(); ++Idx)
+	{
+		const FBatchJobResult& Job = RunResult.Jobs[Idx];
+		const FString Output = !Job.Output.IsEmpty()
+			? Job.Output
+			: FString::Printf(TEXT("(no output) arg: %s"), *Job.ArgumentSummary);
+		FString Line = FString::Printf(TEXT("[%s job %4d]  %s"),
+			Job.bEncounteredError ? TEXT("FAIL") : TEXT("PASS"), Idx, *Output);
+		if (Job.bEncounteredError)
+		{
+			++NumFailed;
+			FailLines.Add(MoveTemp(Line));
+		}
+		else
+		{
+			PassLines.Add(MoveTemp(Line));
+		}
+	}
+
+	TArray<FString> AllLines;
+	AllLines.Append(MoveTemp(FailLines));
+	AllLines.Append(MoveTemp(PassLines));
+	AllLines.Add(FString::Printf(TEXT("\n--- %d total | %d passed | %d failed ---"),
+		RunResult.Jobs.Num(), RunResult.Jobs.Num() - NumFailed, NumFailed));
+	FFileHelper::SaveStringArrayToFile(AllLines, *ResultsFile);
+
+	UE_LOG(LogBatchProcess, Display,
+		TEXT("Results: %d/%d passed  —  report: %s"),
+		RunResult.Jobs.Num() - NumFailed, RunResult.Jobs.Num(), *ResultsFile);
+
+	return NumFailed > 0 ? -1 : 0;
+}
diff --git a/Engine/Source/Editor/Kismet/Private/Commandlets/BatchProcessCommandlet.h b/Engine/Source/Editor/Kismet/Private/Commandlets/BatchProcessCommandlet.h
new file mode 100644
--- /dev/null
+++ b/Engine/Source/Editor/Kismet/Private/Commandlets/BatchProcessCommandlet.h
@@ -0,0 +1,20 @@
+// Copyright Epic Games, Inc. All Rights Reserved.
+
+#pragma once
+
+#include "CoreMinimal.h"
+#include "Commandlets/Commandlet.h"
+
+#include "BatchProcessCommandlet.generated.h"
+
+UCLASS()
+class UBatchProcessCommandlet : public UCommandlet
+{
+	GENERATED_BODY()
+
+	UBatchProcessCommandlet(const FObjectInitializer& ObjectInitializer);
+private:
+	// Begin UCommandlet Interface
+	virtual int32 Main(const FString& Params) override;
+	// End UCommandlet Interface
+};
diff --git a/Engine/Source/Editor/Kismet/Private/Commandlets/BatchWorkerCommandlet.cpp b/Engine/Source/Editor/Kismet/Private/Commandlets/BatchWorkerCommandlet.cpp
new file mode 100644
--- /dev/null
+++ b/Engine/Source/Editor/Kismet/Private/Commandlets/BatchWorkerCommandlet.cpp
@@ -0,0 +1,77 @@
+// Copyright Epic Games, Inc. All Rights Reserved.
+
+#include "Commandlets/BatchWorkerCommandlet.h"
+
+#include "Logging/LogMacros.h"
+#include "BatchProcess/BatchProcessLog.h"
+#include "BatchProcess/BatchProcessWorker.h"
+
+
+namespace UE::Private
+{
+static void PrintBatchWorkerCommandletUsage();
+}
+
+void UE::Private::PrintBatchWorkerCommandletUsage()
+{
+	UE_LOG(LogBatchProcess, Display, TEXT("UnrealEditor ProjectName -run=BatchWorkerCommandlet <CoordinatorIP> <WorkerId>"));
+	UE_LOG(LogBatchProcess, Display, TEXT("  or: -run=BatchWorkerCommandlet -singlecommand='{\"function\":\"...\",\"arguments\":[{...}]}'"));
+}
+
+UBatchWorkerCommandlet::UBatchWorkerCommandlet(const FObjectInitializer& ObjectInitializer)
+	: Super(ObjectInitializer)
+{
+}
+
+int32 UBatchWorkerCommandlet::Main(const FString& Params)
+{
+	TArray<FString> Tokens;
+	TArray<FString> Switches;
+	TMap<FString, FString> ParamsMap;
+	UCommandlet::ParseCommandLine(*Params, Tokens, Switches, ParamsMap);
+
+	// -singlecommand mode: run one job inline without a coordinator, for debugging failures.
+	if (const FString* SingleCommand = ParamsMap.Find(TEXT("singlecommand")))
+	{
+		const bool bOk = FBatchProcessWorker::RunSingleCommand(*SingleCommand);
+		return bOk ? 0 : -1;
+	}
+
+	if(Tokens.Num() > 2 || Tokens.Num() == 0)
+	{
+		UE::Private::PrintBatchWorkerCommandletUsage();
+		return -1;
+	}
+
+	FBatchProcessWorker Worker;
+
+	FString CoordinatorAddrStr;
+	int32 WorkerId = 0;
+	if(Tokens.Num() == 1)
+	{
+		// Single-argument mode: launch a local coordinator for debugging.
+		const bool bLocalCoordinatorLaunched = Worker.LaunchLocalCoordinator(Tokens[0], CoordinatorAddrStr);
+		if(!bLocalCoordinatorLaunched)
+		{
+			return -1;
+		}
+	}
+	else
+	{
+		CoordinatorAddrStr = Tokens[0];
+		WorkerId           = FCString::Atoi(*Tokens[1]);
+	}
+
+	uint32 CoordinatorPid = 0;
+	if (const FString* PidStr = ParamsMap.Find(TEXT("coordinatorpid")))
+	{
+		CoordinatorPid = (uint32)FCString::Atoi64(**PidStr);
+	}
+
+	const bool bSucceeded = Worker.PerformWork(CoordinatorAddrStr, WorkerId, CoordinatorPid);
+	if(!bSucceeded)
+	{
+		UE_LOG(LogBatchProcess, Error, TEXT("Worker failed"));
+	}
+	return bSucceeded ? 0 : -1;
+}
diff --git a/Engine/Source/Editor/Kismet/Private/Commandlets/BatchWorkerCommandlet.h b/Engine/Source/Editor/Kismet/Private/Commandlets/BatchWorkerCommandlet.h
new file mode 100644
--- /dev/null
+++ b/Engine/Source/Editor/Kismet/Private/Commandlets/BatchWorkerCommandlet.h
@@ -0,0 +1,20 @@
+// Copyright Epic Games, Inc. All Rights Reserved.
+
+#pragma once
+
+#include "CoreMinimal.h"
+#include "Commandlets/Commandlet.h"
+
+#include "BatchWorkerCommandlet.generated.h"
+
+UCLASS()
+class UBatchWorkerCommandlet : public UCommandlet
+{
+	GENERATED_BODY()
+
+	UBatchWorkerCommandlet(const FObjectInitializer& ObjectInitializer);
+private:
+	// Begin UCommandlet Interface
+	virtual int32 Main(const FString& Params) override;
+	// End UCommandlet Interface
+};
