Fix example to snapshot the whole package, rather than just the asset object

#rnx
#Synced-CL 51986075

diff --git a/Engine/Source/Editor/Kismet/Private/BatchProcess/BatchProcessLibrary.h b/Engine/Source/Editor/Kismet/Private/BatchProcess/BatchProcessLibrary.h
--- a/Engine/Source/Editor/Kismet/Private/BatchProcess/BatchProcessLibrary.h
+++ b/Engine/Source/Editor/Kismet/Private/BatchProcess/BatchProcessLibrary.h
@@ -40,7 +40,7 @@ class SnapshotAsset(unreal.Object):
 			return
 
 		json_str = unreal.JsonObjectGraphFunctionLibrary.stringify(
-			[asset], unreal.JsonStringifyOptions(unreal.JsonStringifyFlags.FILTER_EDITOR_ONLY_DATA))
+			[asset.get_package()], unreal.JsonStringifyOptions(unreal.JsonStringifyFlags.FILTER_EDITOR_ONLY_DATA))
 
 		# Build a nested output path from the package path, stripping the leading '/'.
 		# e.g. "/Game/Baz/Foo" -> "Game/Baz/Foo.json" and write to that file:
