Fix crash if a crashed worker is detect but it has already close its socket

#rnx
#Synced-CL 51986075

diff --git a/Engine/Source/Editor/Kismet/Private/BatchProcess/BatchProcessCoordinator.cpp b/Engine/Source/Editor/Kismet/Private/BatchProcess/BatchProcessCoordinator.cpp
--- a/Engine/Source/Editor/Kismet/Private/BatchProcess/BatchProcessCoordinator.cpp
+++ b/Engine/Source/Editor/Kismet/Private/BatchProcess/BatchProcessCoordinator.cpp
@@ -523,7 +523,10 @@ void FBatchProcessCoordinator::FindCrashedWorkers()
 		}
 		else
 		{
-			Worker.Socket->Close();
+			if(Worker.Socket)
+			{
+				Worker.Socket->Close();
+			}
 
 			HandleError(FString::Printf(TEXT("Worker crashed: %s"), *Worker.LaunchString));
 
